#ttp 331051 - EDITOR: Can we extend the Set Color on folder functionality in the content browser to Collections as well?
#proj UE4
#branch UE4
#summary Collections now store and display a custom colour based on the local user settings
#extra This works the same was as it does for folders.

diff --git a/Engine/Source/Editor/ContentBrowser/Private/CollectionContextMenu.cpp b/Engine/Source/Editor/ContentBrowser/Private/CollectionContextMenu.cpp
--- a/Engine/Source/Editor/ContentBrowser/Private/CollectionContextMenu.cpp
+++ b/Engine/Source/Editor/ContentBrowser/Private/CollectionContextMenu.cpp
@@ -83,6 +83,30 @@ TSharedPtr<SWidget> FCollectionContextMenu::MakeCollectionListContextMenu(TShare
 					FCanExecuteAction::CreateSP( this, &FCollectionContextMenu::CanExecuteDestroyCollection, bAnyManagedBySCC )
 					)
 				);
+
+			if ( CollectionList.Num() == 1 )
+			{
+				// If any colors have already been set, display color options as a sub menu
+				if ( CollectionViewUtils::HasCustomColors() )
+				{
+					// Set Color (submenu)
+					MenuBuilder.AddSubMenu(
+						LOCTEXT("SetColor", "Set Color"),
+						LOCTEXT("SetCollectionColorTooltip", "Sets the color this collection should appear as."),
+						FNewMenuDelegate::CreateRaw( this, &FCollectionContextMenu::MakeSetColorSubMenu )
+						);
+				}
+				else
+				{
+					// Set Color
+					MenuBuilder.AddMenuEntry(
+						LOCTEXT("SetColor", "Set Color"),
+						LOCTEXT("SetCollectionColorTooltip", "Sets the color this collection should appear as."),
+						FSlateIcon(),
+						FUIAction( FExecuteAction::CreateSP( this, &FCollectionContextMenu::ExecutePickColor ) )
+						);
+				}
+			}
 		}
 	}
 	MenuBuilder.EndSection();
@@ -127,6 +151,60 @@ void FCollectionContextMenu::MakeNewCollectionSubMenu(FMenuBuilder& MenuBuilder)
 	MenuBuilder.EndSection();
 }
 
+void FCollectionContextMenu::MakeSetColorSubMenu(FMenuBuilder& MenuBuilder)
+{
+	// New Color
+	MenuBuilder.AddMenuEntry(
+		LOCTEXT("NewColor", "New Color"),
+		LOCTEXT("NewCollectionColorTooltip", "Changes the color this collection should appear as."),
+		FSlateIcon(),
+		FUIAction( FExecuteAction::CreateSP( this, &FCollectionContextMenu::ExecutePickColor ) )
+		);
+
+	// Clear Color (only required if any of the selection has one)
+	if ( SelectedHasCustomColors() )
+	{
+		MenuBuilder.AddMenuEntry(
+			LOCTEXT("ClearColor", "Clear Color"),
+			LOCTEXT("ClearCollectionColorTooltip", "Resets the color this collection appears as."),
+			FSlateIcon(),
+			FUIAction( FExecuteAction::CreateSP( this, &FCollectionContextMenu::ExecuteResetColor ) )
+			);
+	}
+
+	// Add all the custom colors the user has chosen so far
+	TArray< FLinearColor > CustomColors;
+	if ( CollectionViewUtils::HasCustomColors( &CustomColors ) )
+	{	
+		MenuBuilder.BeginSection("PathContextCustomColors", LOCTEXT("CustomColorsExistingColors", "Existing Colors") );
+		{
+			for ( int32 ColorIndex = 0; ColorIndex < CustomColors.Num(); ColorIndex++ )
+			{
+				const FLinearColor& Color = CustomColors[ ColorIndex ];
+				MenuBuilder.AddWidget(
+						SNew(SHorizontalBox)
+						+SHorizontalBox::Slot()
+						.AutoWidth()
+						.Padding(2, 0, 0, 0)
+						[
+							SNew(SButton)
+							.ButtonStyle( FEditorStyle::Get(), "Menu.Button" )
+							.OnClicked( this, &FCollectionContextMenu::OnColorClicked, Color )
+							[
+								SNew(SColorBlock)
+								.Color( Color )
+								.Size( FVector2D(77,16) )
+							]
+						],
+					LOCTEXT("CustomColor", ""),
+					/*bNoIndent=*/true
+				);
+			}
+		}
+		MenuBuilder.EndSection();
+	}
+}
+
 void FCollectionContextMenu::UpdateProjectSourceControl()
 {
 	// Force update of source control so that we're always showing the valid options
@@ -138,6 +216,18 @@ void FCollectionContextMenu::UpdateProjectSourceControl()
 	}
 }
 
+bool FCollectionContextMenu::CanRenameSelectedCollections() const
+{
+	TArray<TSharedPtr<FCollectionItem>> CollectionList = CollectionView.Pin()->CollectionListPtr->GetSelectedItems();
+	
+	if(CollectionList.Num() == 1)
+	{
+		return !(CollectionList[0]->CollectionType != ECollectionShareType::CST_Local) || (bProjectUnderSourceControl && ISourceControlModule::Get().IsEnabled() && ISourceControlModule::Get().GetProvider().IsAvailable());
+	}
+	
+	return false;
+}
+
 void FCollectionContextMenu::ExecuteNewCollection(ECollectionShareType::Type CollectionType)
 {
 	if ( !ensure(CollectionView.IsValid()) )
@@ -228,6 +318,48 @@ FReply FCollectionContextMenu::ExecuteDestroyCollectionConfirmed(TArray<TSharedP
 	return FReply::Handled();
 }
 
+void FCollectionContextMenu::ExecuteResetColor()
+{
+	ResetColors();
+}
+
+void FCollectionContextMenu::ExecutePickColor()
+{
+	TSharedPtr<SCollectionView> CollectionViewPtr = CollectionView.Pin();
+	TArray<FCollectionNameType> SelectedCollections = CollectionViewPtr->GetSelectedCollections();
+
+	// Spawn a color picker, so the user can select which color they want
+	TArray<FLinearColor*> LinearColorArray;
+	FColorPickerArgs PickerArgs;
+	PickerArgs.bIsModal = false;
+	PickerArgs.ParentWidget = CollectionViewPtr;
+	if ( SelectedCollections.Num() > 0 )
+	{
+		// Make sure an color entry exists for all the collections, otherwise they won't update in realtime with the widget color
+		for (int32 PathIdx = SelectedCollections.Num() - 1; PathIdx >= 0; --PathIdx)
+		{
+			const FCollectionNameType& SelectedCollection = SelectedCollections[PathIdx];
+			TSharedPtr<FLinearColor> Color = CollectionViewUtils::LoadColor( SelectedCollection.Name.ToString(), SelectedCollection.Type );
+			if ( !Color.IsValid() )
+			{
+				Color = MakeShareable( new FLinearColor( CollectionViewUtils::GetDefaultColor() ) );
+				CollectionViewUtils::SaveColor( SelectedCollection.Name.ToString(), SelectedCollection.Type, Color, true );
+			}
+			else
+			{
+				// Default the color to the first valid entry
+				PickerArgs.InitialColorOverride = *Color.Get();
+			}
+			LinearColorArray.Add( Color.Get() );
+		}	
+		PickerArgs.LinearColorArray = &LinearColorArray;
+	}
+		
+	PickerArgs.OnColorPickerWindowClosed = FOnWindowClosed::CreateSP(this, &FCollectionContextMenu::NewColorComplete);
+
+	OpenColorPicker(PickerArgs);
+}
+
 bool FCollectionContextMenu::CanExecuteNewCollection(ECollectionShareType::Type CollectionType) const
 {
 	return (CollectionType == ECollectionShareType::CST_Local) || (bProjectUnderSourceControl && ISourceControlModule::Get().IsEnabled() && ISourceControlModule::Get().GetProvider().IsAvailable());
@@ -235,19 +367,78 @@ bool FCollectionContextMenu::CanExecuteNewCollection(ECollectionShareType::Type
 
 bool FCollectionContextMenu::CanExecuteRenameCollection() const
 {
-	TArray<TSharedPtr<FCollectionItem>> CollectionList = CollectionView.Pin()->CollectionListPtr->GetSelectedItems();
-	
-	if(CollectionList.Num() == 1)
+	return CanRenameSelectedCollections();
+}
+
+bool FCollectionContextMenu::CanExecuteDestroyCollection(bool bAnyManagedBySCC) const
+{
+	return !bAnyManagedBySCC || (bProjectUnderSourceControl && ISourceControlModule::Get().IsEnabled() && ISourceControlModule::Get().GetProvider().IsAvailable());
+}
+
+bool FCollectionContextMenu::SelectedHasCustomColors() const
+{
+	TSharedPtr<SCollectionView> CollectionViewPtr = CollectionView.Pin();
+	TArray<FCollectionNameType> SelectedCollections = CollectionViewPtr->GetSelectedCollections();
+
+	for(const FCollectionNameType& SelectedCollection : SelectedCollections)
 	{
-		return !(CollectionList[0]->CollectionType != ECollectionShareType::CST_Local) || (bProjectUnderSourceControl && ISourceControlModule::Get().IsEnabled() && ISourceControlModule::Get().GetProvider().IsAvailable());
+		// Ignore any that are the default color
+		const TSharedPtr<FLinearColor> Color = CollectionViewUtils::LoadColor( SelectedCollection.Name.ToString(), SelectedCollection.Type );
+		if ( Color.IsValid() && !Color->Equals( CollectionViewUtils::GetDefaultColor() ) )
+		{
+			return true;
+		}
 	}
-	
 	return false;
 }
 
-bool FCollectionContextMenu::CanExecuteDestroyCollection(bool bAnyManagedBySCC) const
+void FCollectionContextMenu::NewColorComplete(const TSharedRef<SWindow>& Window)
 {
-	return !bAnyManagedBySCC || (bProjectUnderSourceControl && ISourceControlModule::Get().IsEnabled() && ISourceControlModule::Get().GetProvider().IsAvailable());
+	TSharedPtr<SCollectionView> CollectionViewPtr = CollectionView.Pin();
+	TArray<FCollectionNameType> SelectedCollections = CollectionViewPtr->GetSelectedCollections();
+
+	// Save the colors back in the config (ptr should have already updated by the widget)
+	for(const FCollectionNameType& SelectedCollection : SelectedCollections)
+	{
+		const TSharedPtr<FLinearColor> Color = CollectionViewUtils::LoadColor( SelectedCollection.Name.ToString(), SelectedCollection.Type );
+		check( Color.IsValid() );
+		CollectionViewUtils::SaveColor( SelectedCollection.Name.ToString(), SelectedCollection.Type, Color );
+	}
+}
+
+FReply FCollectionContextMenu::OnColorClicked( const FLinearColor InColor )
+{
+	TSharedPtr<SCollectionView> CollectionViewPtr = CollectionView.Pin();
+	TArray<FCollectionNameType> SelectedCollections = CollectionViewPtr->GetSelectedCollections();
+
+	// Make sure an color entry exists for all the collections, otherwise it can't save correctly
+	for(const FCollectionNameType& SelectedCollection : SelectedCollections)
+	{
+		TSharedPtr<FLinearColor> Color = CollectionViewUtils::LoadColor( SelectedCollection.Name.ToString(), SelectedCollection.Type );
+		if ( !Color.IsValid() )
+		{
+			Color = MakeShareable( new FLinearColor() );
+		}
+		*Color.Get() = InColor;
+		CollectionViewUtils::SaveColor( SelectedCollection.Name.ToString(), SelectedCollection.Type, Color );
+	}
+
+	// Dismiss the menu here, as we can't make the 'clear' option appear if a folder has just had a color set for the first time
+	FSlateApplication::Get().DismissAllMenus();
+
+	return FReply::Handled();
+}
+
+void FCollectionContextMenu::ResetColors()
+{
+	TSharedPtr<SCollectionView> CollectionViewPtr = CollectionView.Pin();
+	TArray<FCollectionNameType> SelectedCollections = CollectionViewPtr->GetSelectedCollections();
+
+	// Clear the custom colors for all the selected collections
+	for(const FCollectionNameType& SelectedCollection : SelectedCollections)
+	{
+		CollectionViewUtils::SaveColor( SelectedCollection.Name.ToString(), SelectedCollection.Type, nullptr );
+	}
 }
 
 #undef LOCTEXT_NAMESPACE
diff --git a/Engine/Source/Editor/ContentBrowser/Private/CollectionContextMenu.h b/Engine/Source/Editor/ContentBrowser/Private/CollectionContextMenu.h
--- a/Engine/Source/Editor/ContentBrowser/Private/CollectionContextMenu.h
+++ b/Engine/Source/Editor/ContentBrowser/Private/CollectionContextMenu.h
@@ -20,6 +20,14 @@ public:
 	/** Update the source control flag the 'CanExecute' functions rely on */
 	void UpdateProjectSourceControl();
 
+	/** Can the selected collections be renamed? */
+	bool CanRenameSelectedCollections() const;
+
+protected:
+	/** Makes the set color submenu */
+	void MakeSetColorSubMenu(FMenuBuilder& MenuBuilder);
+
+private:
 	/** Handler for when a collection is selected in the "New" menu */
 	void ExecuteNewCollection(ECollectionShareType::Type CollectionType);
 
@@ -32,6 +40,12 @@ public:
 	/** Handler for when "Destroy Collection" is confirmed */
 	FReply ExecuteDestroyCollectionConfirmed(TArray<TSharedPtr<FCollectionItem>> CollectionList);
 
+	/** Handler for when reset color is selected */
+	void ExecuteResetColor();
+
+	/** Handler for when new or set color is selected */
+	void ExecutePickColor();
+
 	/** Handler to check to see if "New Collection" can be executed */
 	bool CanExecuteNewCollection(ECollectionShareType::Type CollectionType) const;
 
@@ -41,7 +55,18 @@ public:
 	/** Handler to check to see if "Destroy Collection" can be executed */
 	bool CanExecuteDestroyCollection(bool bAnyManagedBySCC) const;
 
-private:
+	/** Checks to see if any of the selected collections use custom colors */
+	bool SelectedHasCustomColors() const;
+
+	/** Callback when the color picker dialog has been closed */
+	void NewColorComplete(const TSharedRef<SWindow>& Window);
+
+	/** Callback when the color is picked from the set color submenu */
+	FReply OnColorClicked( const FLinearColor InColor );
+
+	/** Resets the colors of the selected collections */
+	void ResetColors();
+
 	TWeakPtr<SCollectionView> CollectionView;
 
 	/** Flag caching whether the project is under source control */
diff --git a/Engine/Source/Editor/ContentBrowser/Private/CollectionViewUtils.cpp b/Engine/Source/Editor/ContentBrowser/Private/CollectionViewUtils.cpp
new file mode 100644
--- /dev/null
+++ b/Engine/Source/Editor/ContentBrowser/Private/CollectionViewUtils.cpp
@@ -0,0 +1,219 @@
+// Copyright 1998-2014 Epic Games, Inc. All Rights Reserved.
+
+
+#include "ContentBrowserPCH.h"
+
+#define LOCTEXT_NAMESPACE "CollectionView"
+
+namespace CollectionViewUtils
+{
+	/** Create a string of the form "CollectionName:CollectionType */
+	FString ToConfigKey(const FString& InCollectionName, const ECollectionShareType::Type& InCollectionType)
+	{
+		static_assert(ECollectionShareType::CST_All == 4, "Update CollectionViewUtils::ToConfigKey for the updated ECollectionShareType values");
+
+		check(InCollectionType != ECollectionShareType::CST_All);
+
+		FString CollectionTypeStr;
+		switch(InCollectionType)
+		{
+		case ECollectionShareType::CST_System:
+			CollectionTypeStr = "System";
+		case ECollectionShareType::CST_Local:
+			CollectionTypeStr = "Local";
+		case ECollectionShareType::CST_Private:
+			CollectionTypeStr = "Private";
+		case ECollectionShareType::CST_Shared:
+			CollectionTypeStr = "Shared";
+		default:
+			break;
+		}
+
+		return InCollectionName + ":" + CollectionTypeStr;
+	}
+
+	/** Convert a string of the form "CollectionName:CollectionType back into its individual elements */
+	bool FromConfigKey(const FString& InKey, FString& OutCollectionName, ECollectionShareType::Type& OutCollectionType)
+	{
+		static_assert(ECollectionShareType::CST_All == 4, "Update CollectionViewUtils::FromConfigKey for the updated ECollectionShareType values");
+
+		FString CollectionTypeStr;
+		if(InKey.Split(":", &OutCollectionName, &CollectionTypeStr, ESearchCase::IgnoreCase, ESearchDir::FromEnd))
+		{
+			if(CollectionTypeStr == "System")
+			{
+				OutCollectionType = ECollectionShareType::CST_System;
+			}
+			else if(CollectionTypeStr == "Local")
+			{
+				OutCollectionType = ECollectionShareType::CST_Local;
+			}
+			else if(CollectionTypeStr == "Private")
+			{
+				OutCollectionType = ECollectionShareType::CST_Private;
+			}
+			else if(CollectionTypeStr == "Shared")
+			{
+				OutCollectionType = ECollectionShareType::CST_Shared;
+			}
+			else
+			{
+				return false;
+			}
+
+			return !OutCollectionName.IsEmpty();
+		}
+
+		return false;
+	}
+
+	// Keep a map of all the collections that have custom colors, so updating the color in one location updates them all
+	static TMap<FString, TSharedPtr<FLinearColor>> CollectionColors;
+}
+
+const TSharedPtr<FLinearColor> CollectionViewUtils::LoadColor(const FString& InCollectionName, const ECollectionShareType::Type& InCollectionType)
+{
+	check(InCollectionType != ECollectionShareType::CST_All);
+
+	const FString ColorKeyStr = ToConfigKey(InCollectionName, InCollectionType);
+
+	// See if we have a value cached first
+	{
+		TSharedPtr<FLinearColor>* const CachedColor = CollectionColors.Find(ColorKeyStr);
+		if(CachedColor)
+		{
+			return *CachedColor;
+		}
+	}
+		
+	// Loads the color of collection at the given path from the config
+	if(FPaths::FileExists(GEditorUserSettingsIni))
+	{
+		// Create a new entry from the config, skip if it's default
+		FString ColorStr;
+		if(GConfig->GetString(TEXT("CollectionColor"), *ColorKeyStr, ColorStr, GEditorUserSettingsIni))
+		{
+			FLinearColor Color;
+			if(Color.InitFromString(ColorStr) && !Color.Equals(CollectionViewUtils::GetDefaultColor()))
+			{
+				return CollectionColors.Add(ColorKeyStr, MakeShareable(new FLinearColor(Color)));
+			}
+		}
+		else
+		{
+			return CollectionColors.Add(ColorKeyStr, MakeShareable(new FLinearColor(CollectionViewUtils::GetDefaultColor())));
+		}
+	}
+
+	return nullptr;
+}
+
+void CollectionViewUtils::SaveColor(const FString& InCollectionName, const ECollectionShareType::Type& InCollectionType, const TSharedPtr<FLinearColor> CollectionColor, const bool bForceAdd)
+{
+	check(InCollectionType != ECollectionShareType::CST_All);
+
+	const FString ColorKeyStr = ToConfigKey(InCollectionName, InCollectionType);
+
+	// Remove the color if it's invalid or default
+	const bool bRemove = !CollectionColor.IsValid() || (!bForceAdd && CollectionColor->Equals(CollectionViewUtils::GetDefaultColor()));
+
+	// Saves the color of the collection to the config
+	if(FPaths::FileExists(GEditorUserSettingsIni))
+	{
+		// If this is no longer custom, remove it
+		if(bRemove)
+		{
+			GConfig->RemoveKey(TEXT("CollectionColor"), *ColorKeyStr, GEditorUserSettingsIni);
+		}
+		else
+		{
+			GConfig->SetString(TEXT("CollectionColor"), *ColorKeyStr, *CollectionColor->ToString(), GEditorUserSettingsIni);
+		}
+	}
+
+	// Update the map too
+	if(bRemove)
+	{
+		CollectionColors.Remove(ColorKeyStr);
+	}
+	else
+	{
+		CollectionColors.Add(ColorKeyStr, CollectionColor);
+	}
+}
+
+bool CollectionViewUtils::HasCustomColors( TArray< FLinearColor >* OutColors )
+{
+	if(!FPaths::FileExists(GEditorUserSettingsIni))
+	{
+		return false;
+	}
+
+	// Read individual entries from a config file.
+	TArray<FString> Section;
+	GConfig->GetSection(TEXT("CollectionColor"), Section, GEditorUserSettingsIni);
+
+	bool bHasCustom = false;
+	const FCollectionManagerModule& CollectionManagerModule = FModuleManager::Get().LoadModuleChecked<FCollectionManagerModule>("CollectionManager");
+	const ICollectionManager& CollectionManager = CollectionManagerModule.Get();
+
+	for(FString& EntryStr : Section)
+	{
+		EntryStr.Trim();
+
+		FString ColorKeyStr;
+		FString ColorStr;
+		if(!EntryStr.Split("=", &ColorKeyStr, &ColorStr))
+		{
+			continue;
+		}
+
+		// Ignore any that have invalid or default colors
+		FLinearColor CurrentColor;
+		if(!CurrentColor.InitFromString(ColorStr) || CurrentColor.Equals(CollectionViewUtils::GetDefaultColor()))
+		{
+			continue;
+		}
+
+		// Ignore any that reference old collections
+		FString CollectionName;
+		ECollectionShareType::Type CollectionType;
+		if(!FromConfigKey(ColorKeyStr, CollectionName, CollectionType) || !CollectionManager.CollectionExists(*CollectionName, CollectionType))
+		{
+			continue;
+		}
+
+		bHasCustom = true;
+		if(OutColors)
+		{
+			// Only add if not already present (ignores near matches too)
+			bool bAdded = false;
+			for(const FLinearColor& Color : *OutColors)
+			{
+				if(CurrentColor.Equals(Color))
+				{
+					bAdded = true;
+					break;
+				}
+			}
+			if(!bAdded)
+			{
+				OutColors->Add(CurrentColor);
+			}
+		}
+		else
+		{
+			break;
+		}
+	}
+
+	return bHasCustom;
+}
+
+FLinearColor CollectionViewUtils::GetDefaultColor()
+{
+	// The default tint the folder should appear as
+	return FLinearColor::Gray;
+}
+
+#undef LOCTEXT_NAMESPACE
diff --git a/Engine/Source/Editor/ContentBrowser/Private/CollectionViewUtils.h b/Engine/Source/Editor/ContentBrowser/Private/CollectionViewUtils.h
new file mode 100644
--- /dev/null
+++ b/Engine/Source/Editor/ContentBrowser/Private/CollectionViewUtils.h
@@ -0,0 +1,36 @@
+// Copyright 1998-2014 Epic Games, Inc. All Rights Reserved.
+
+#pragma once
+
+namespace CollectionViewUtils
+{
+	/**
+	 * Loads the color of this collection from the config
+	 *
+	 * @param InCollectionName - The name of the collection to get the color for
+	 * @param InCollectionType - The type of the collection to get the color for
+	 * @return The color the collection should appear as, will be null if not customized
+	 */
+	const TSharedPtr<FLinearColor> LoadColor(const FString& InCollectionName, const ECollectionShareType::Type& InCollectionType);
+
+	/**
+	 * Saves the color of the collection to the config
+	 *
+	 * @param InCollectionName - The name of the collection to set the color for
+	 * @param InCollectionType - The type of the collection to set the color for
+	 * @param CollectionColor - The color the collection should appear as
+	 * @param bForceAdd - If true, force the color to be added even if it's the default color
+	 */
+	void SaveColor(const FString& InCollectionName, const ECollectionShareType::Type& InCollectionType, const TSharedPtr<FLinearColor> CollectionColor, const bool bForceAdd = false);
+
+	/**
+	 * Checks to see if any collection has a custom color, optionally outputs them to a list
+	 *
+	 * @param OutColors - If specified, returns all the custom colors being used
+	 * @return true, if custom colors are present
+	 */
+	bool HasCustomColors( TArray< FLinearColor >* OutColors = nullptr );
+
+	/** Gets the default color the collection should appear as */
+	FLinearColor GetDefaultColor();
+}
diff --git a/Engine/Source/Editor/ContentBrowser/Private/ContentBrowserPCH.h b/Engine/Source/Editor/ContentBrowser/Private/ContentBrowserPCH.h
--- a/Engine/Source/Editor/ContentBrowser/Private/ContentBrowserPCH.h
+++ b/Engine/Source/Editor/ContentBrowser/Private/ContentBrowserPCH.h
@@ -30,6 +30,7 @@ DECLARE_DELEGATE_OneParam(FOnAssetRenameCommitted, const TArray<FAssetData>& /*A
 #include "SAssetSearchBox.h"
 #include "SFilterList.h"
 #include "SPathView.h"
+#include "CollectionViewUtils.h"
 #include "SCollectionView.h"
 #include "AssetViewSortManager.h"
 #include "OtherDevelopersAssetFilter.h"
diff --git a/Engine/Source/Editor/ContentBrowser/Private/ContentBrowserUtils.cpp b/Engine/Source/Editor/ContentBrowser/Private/ContentBrowserUtils.cpp
--- a/Engine/Source/Editor/ContentBrowser/Private/ContentBrowserUtils.cpp
+++ b/Engine/Source/Editor/ContentBrowser/Private/ContentBrowserUtils.cpp
@@ -12,8 +12,11 @@
 
 #define LOCTEXT_NAMESPACE "ContentBrowser"
 
-// Keep a map of all the paths that have custom colors, so updating the color in one location updates them all
-static TMap< FString, TSharedPtr< FLinearColor > > PathColors;
+namespace ContentBrowserUtils
+{
+	// Keep a map of all the paths that have custom colors, so updating the color in one location updates them all
+	static TMap< FString, TSharedPtr< FLinearColor > > PathColors;
+}
 
 class SContentBrowserPopup : public SCompoundWidget
 {
diff --git a/Engine/Source/Editor/ContentBrowser/Private/SCollectionView.cpp b/Engine/Source/Editor/ContentBrowser/Private/SCollectionView.cpp
--- a/Engine/Source/Editor/ContentBrowser/Private/SCollectionView.cpp
+++ b/Engine/Source/Editor/ContentBrowser/Private/SCollectionView.cpp
@@ -511,7 +511,7 @@ void SCollectionView::CollectionItemScrolledIntoView( TSharedPtr<FCollectionItem
 bool SCollectionView::IsCollectionNotRenamable() const
 {
 	CollectionContextMenu->UpdateProjectSourceControl();
-	return !CollectionContextMenu->CanExecuteRenameCollection();
+	return !CollectionContextMenu->CanRenameSelectedCollections();
 }
 
 bool SCollectionView::CollectionNameChangeCommit( const TSharedPtr< FCollectionItem >& CollectionItem, const FString& NewName, bool bChangeConfirmed, FText& OutWarningMessage )
diff --git a/Engine/Source/Editor/ContentBrowser/Private/SourcesViewWidgets.cpp b/Engine/Source/Editor/ContentBrowser/Private/SourcesViewWidgets.cpp
--- a/Engine/Source/Editor/ContentBrowser/Private/SourcesViewWidgets.cpp
+++ b/Engine/Source/Editor/ContentBrowser/Private/SourcesViewWidgets.cpp
@@ -8,6 +8,7 @@
 #include "DragAndDrop/AssetDragDropOp.h"
 #include "DragAndDrop/AssetPathDragDropOp.h"
 #include "ContentBrowserUtils.h"
+#include "CollectionViewUtils.h"
 
 #define LOCTEXT_NAMESPACE "ContentBrowser"
 
@@ -425,6 +426,7 @@ void SCollectionListItem::Construct( const FArguments& InArgs )
 				// Type Icon
 				SNew(SImage)
 				.Image( FEditorStyle::GetBrush(*CollectionTypeImage) )
+				.ColorAndOpacity( this, &SCollectionListItem::GetCollectionColor )
 			]
 
 			+SHorizontalBox::Slot()
@@ -584,6 +586,22 @@ FText SCollectionListItem::GetNameText() const
 	}
 }
 
+FSlateColor SCollectionListItem::GetCollectionColor() const
+{
+	TSharedPtr<FCollectionItem> CollectionItemPtr = CollectionItem.Pin();
+
+	if ( CollectionItemPtr.IsValid() )
+	{
+		const TSharedPtr<FLinearColor> Color = CollectionViewUtils::LoadColor(CollectionItemPtr->CollectionName, CollectionItemPtr->CollectionType);
+		if( Color.IsValid() )
+		{
+			return *Color.Get();
+		}
+	}
+	
+	return CollectionViewUtils::GetDefaultColor();
+}
+
 const FSlateBrush* SCollectionListItem::GetBorderImage() const
 {
 	return bDraggedOver ? FEditorStyle::GetBrush("Menu.Background") : FEditorStyle::GetBrush("NoBorder");
diff --git a/Engine/Source/Editor/ContentBrowser/Private/SourcesViewWidgets.h b/Engine/Source/Editor/ContentBrowser/Private/SourcesViewWidgets.h
--- a/Engine/Source/Editor/ContentBrowser/Private/SourcesViewWidgets.h
+++ b/Engine/Source/Editor/ContentBrowser/Private/SourcesViewWidgets.h
@@ -195,6 +195,9 @@ private:
 	/** Returns the text of the collection name */
 	FText GetNameText() const;
 
+	/** Returns the color of the collection name */
+	FSlateColor GetCollectionColor() const;
+
 	/** Returns the image for the border around this item. Used for drag/drop operations */
 	const FSlateBrush* GetBorderImage() const;
 
