#TTP 332018 - Regression: EDITOR: LANDSCAPE: CRASH: Choosing a landscape material before creation and then closing the editor will result in a crash

#proj UE4

#summary Fixed crash when closing new landscape dialog with a material selected

diff --git a/Engine/Source/Editor/ContentBrowser/Private/SAssetView.cpp b/Engine/Source/Editor/ContentBrowser/Private/SAssetView.cpp
--- a/Engine/Source/Editor/ContentBrowser/Private/SAssetView.cpp
+++ b/Engine/Source/Editor/ContentBrowser/Private/SAssetView.cpp
@@ -37,7 +37,7 @@ SAssetView::~SAssetView()
 	FCoreDelegates::OnObjectPropertyChanged.RemoveAll( this );
 
 	// Remove the listener for when view settings are changed
-	GetMutableDefault<UContentBrowserSettings>()->OnSettingChanged().RemoveAll( this );
+	UContentBrowserSettings::OnSettingChanged().RemoveAll(this);
 
 	if ( FrontendFilters.IsValid() )
 	{
@@ -81,7 +81,7 @@ void SAssetView::Construct( const FArguments& InArgs )
 	FCoreDelegates::OnObjectPropertyChanged.Add(FCoreDelegates::FOnObjectPropertyChanged::FDelegate::CreateSP(this, &SAssetView::OnObjectPropertyChanged));
 
 	// Listen for when view settings are changed
-	GetMutableDefault<UContentBrowserSettings>()->OnSettingChanged().Add(UContentBrowserSettings::FSettingChangedEvent::FDelegate::CreateSP(this, &SAssetView::HandleSettingChanged));
+	UContentBrowserSettings::OnSettingChanged().Add(UContentBrowserSettings::FSettingChangedEvent::FDelegate::CreateSP(this, &SAssetView::HandleSettingChanged));
 
 	// Get desktop metrics
 	FDisplayMetrics DisplayMetrics;
diff --git a/Engine/Source/Editor/ContentBrowser/Private/SContentBrowser.cpp b/Engine/Source/Editor/ContentBrowser/Private/SContentBrowser.cpp
--- a/Engine/Source/Editor/ContentBrowser/Private/SContentBrowser.cpp
+++ b/Engine/Source/Editor/ContentBrowser/Private/SContentBrowser.cpp
@@ -21,7 +21,7 @@ const FString SContentBrowser::SettingsIniSection = TEXT("ContentBrowser");
 SContentBrowser::~SContentBrowser()
 {
 	// Remove the listener for when view settings are changed
-	GetMutableDefault<UContentBrowserSettings>()->OnSettingChanged().RemoveAll( this );
+	UContentBrowserSettings::OnSettingChanged().RemoveAll( this );
 }
 
 BEGIN_SLATE_FUNCTION_BUILD_OPTIMIZATION
@@ -41,7 +41,7 @@ void SContentBrowser::Construct( const FArguments& InArgs, const FName& InInstan
 	bIsLocked = InArgs._InitiallyLocked;
 
 	// Listen for when view settings are changed
-	GetMutableDefault<UContentBrowserSettings>()->OnSettingChanged().Add(UContentBrowserSettings::FSettingChangedEvent::FDelegate::CreateSP(this, &SContentBrowser::HandleSettingChanged));
+	UContentBrowserSettings::OnSettingChanged().Add(UContentBrowserSettings::FSettingChangedEvent::FDelegate::CreateSP(this, &SContentBrowser::HandleSettingChanged));
 
 	HistoryManager.SetOnApplyHistoryData(FOnApplyHistoryData::CreateSP(this, &SContentBrowser::OnApplyHistoryData));
 	HistoryManager.SetOnUpdateHistoryData(FOnUpdateHistoryData::CreateSP(this, &SContentBrowser::OnUpdateHistoryData));
diff --git a/Engine/Source/Editor/UnrealEd/Classes/Settings/ContentBrowserSettings.h b/Engine/Source/Editor/UnrealEd/Classes/Settings/ContentBrowserSettings.h
--- a/Engine/Source/Editor/UnrealEd/Classes/Settings/ContentBrowserSettings.h
+++ b/Engine/Source/Editor/UnrealEd/Classes/Settings/ContentBrowserSettings.h
@@ -68,7 +68,7 @@ public:
 	 * @return The delegate.
 	 */
 	DECLARE_EVENT_OneParam(UContentBrowserSettings, FSettingChangedEvent, FName /*PropertyName*/);
-	FSettingChangedEvent& OnSettingChanged( ) { return SettingChangedEvent; }
+	static FSettingChangedEvent& OnSettingChanged( ) { return SettingChangedEvent; }
 
 protected:
 
@@ -95,5 +95,5 @@ private:
 	bool OverrideDisplayDevelopersFolder;
 
 	// Holds an event delegate that is executed when a setting has changed.
-	FSettingChangedEvent SettingChangedEvent;
+	static FSettingChangedEvent SettingChangedEvent;
 };
