* The thumbnail range slider is now based on the vertical resolution instead of the horizontal.  That way users using the unified desktop driver stuff wont have a 8K monitor causing their thumbnails to be gigantic.

diff --git a/Engine/Source/Editor/ContentBrowser/Private/SAssetView.cpp b/Engine/Source/Editor/ContentBrowser/Private/SAssetView.cpp
--- a/Engine/Source/Editor/ContentBrowser/Private/SAssetView.cpp
+++ b/Engine/Source/Editor/ContentBrowser/Private/SAssetView.cpp
@@ -91,7 +91,7 @@ void SAssetView::Construct( const FArguments& InArgs )
 		DisplayMetrics.PrimaryDisplayWorkAreaRect.Right - DisplayMetrics.PrimaryDisplayWorkAreaRect.Left,
 		DisplayMetrics.PrimaryDisplayWorkAreaRect.Bottom - DisplayMetrics.PrimaryDisplayWorkAreaRect.Top );
 
-	const float ThumbnailScaleRangeScalar = ( DisplaySize.X / 1920 );
+	const float ThumbnailScaleRangeScalar = ( DisplaySize.Y / 1080 );
 
 	// Create a thumbnail pool for rendering thumbnails	
 	AssetThumbnailPool = MakeShareable( new FAssetThumbnailPool(1024, InArgs._AreRealTimeThumbnailsAllowed) );
