UE4: Finding world assets in explorer now properly uses the .umap extension

diff --git a/Engine/Source/Editor/ContentBrowser/Private/AssetContextMenu.cpp b/Engine/Source/Editor/ContentBrowser/Private/AssetContextMenu.cpp
--- a/Engine/Source/Editor/ContentBrowser/Private/AssetContextMenu.cpp
+++ b/Engine/Source/Editor/ContentBrowser/Private/AssetContextMenu.cpp
@@ -509,10 +509,11 @@ void FAssetContextMenu::ExecuteFindInExplorer()
 		{
 			FAssetData AssetData(Asset);
 
-			const FString FilePath = FPackageName::LongPackageNameToFilename(AssetData.PackageName.ToString());
+			const bool bIsWorldAsset = (AssetData.AssetClass == UWorld::StaticClass()->GetFName());
+			const FString Extension = bIsWorldAsset ? FPackageName::GetMapPackageExtension() : FPackageName::GetAssetPackageExtension();
+			const FString FilePath = FPackageName::LongPackageNameToFilename(AssetData.PackageName.ToString(), Extension);
 			const FString FullFilePath = FPaths::ConvertRelativePathToFull(FilePath);
-			const FString Filename = FullFilePath + FPackageName::GetAssetPackageExtension();
-			FPlatformProcess::ExploreFolder(*Filename);
+			FPlatformProcess::ExploreFolder(*FullFilePath);
 		}
 	}
 }
