Fixing the text in the asset menu, was stomping other key'ed text.

diff --git a/Engine/Source/Editor/ContentBrowser/Private/NewAssetContextMenu.cpp b/Engine/Source/Editor/ContentBrowser/Private/NewAssetContextMenu.cpp
--- a/Engine/Source/Editor/ContentBrowser/Private/NewAssetContextMenu.cpp
+++ b/Engine/Source/Editor/ContentBrowser/Private/NewAssetContextMenu.cpp
@@ -156,7 +156,7 @@ void FNewAssetContextMenu::MakeContextMenu(FMenuBuilder& MenuBuilder, const FStr
 			);
 
 		MenuBuilder.AddSubMenu(
-			LOCTEXT("MaterialAssetCategory", "User Interface"),
+			LOCTEXT("UserInterfaceAssetCategory", "User Interface"),
 			FText::GetEmpty(),
 			FNewMenuDelegate::CreateStatic(&FNewAssetContextMenu::CreateNewAssetMenuCategory, EAssetTypeCategories::UI, InPath, InOnNewAssetRequested)
 			);
