Further limit the game name length (20 characters), predict the cooking path length and limit the naming length accordingly.
#ttp 332328 Content Browser: Need new warnings about long relative path lengths for assets
#proj Editor
#branch UE4

diff --git a/Engine/Source/Editor/ContentBrowser/Private/SAssetView.cpp b/Engine/Source/Editor/ContentBrowser/Private/SAssetView.cpp
--- a/Engine/Source/Editor/ContentBrowser/Private/SAssetView.cpp
+++ b/Engine/Source/Editor/ContentBrowser/Private/SAssetView.cpp
@@ -17,6 +17,8 @@
 #define MAX_THUMBNAIL_SIZE 4096
 #define MAX_CLASS_NAME_LENGTH 32 // Enforce a reasonable class name length so the path is not too long for PLATFORM_MAX_FILEPATH_LENGTH
 
+#define MAX_PROJECTED_COOKING_PATH 165
+
 const double SAssetView::FQuickJumpData::JumpDelaySeconds = 0.6;
 
 SAssetView::~SAssetView()
@@ -3133,12 +3135,38 @@ bool SAssetView::AssetVerifyRenameCommit(const TSharedPtr<FAssetViewItem>& Item,
 			return false;
 		}
 
+		// The following checks are done mostly to prevent / alleviate the problems that "long" paths are causing with the BuildFarm and cooked builds.
+		// The BuildFarm buildmachines use a verbose path to encode extra information to provide more information when things fail, however
+		// this makes the path limitation (260 chars on Windows) a problem. It doubles up the GGameName and does the cooking in another
+		// sub-folder, one of the "saved/sandboxes", with folder duplication.
+
+		// Get the SubPath containing folders without the "game name" folder itself
+		const FString GameNameStr(GGameName);
+		FString SubPath = FPaths::GameDir();
+		FPaths::NormalizeDirectoryName(SubPath);
+		SubPath = SubPath.Replace(*(FString(TEXT("../../../")) + GameNameStr), TEXT(""));
+		FPaths::RemoveDuplicateSlashes(SubPath);
+
+		// Calculate the maximum path length this will generate when doing a cooked build.
+		const int32 PathCalcLen = SubPath.Len() + (2 * GameNameStr.Len()) + (NewPackageName + FPackageName::GetAssetPackageExtension()).Len();
+		if ( PathCalcLen >= MAX_PROJECTED_COOKING_PATH )
+		{
+			// The projected length of the path for cooking is too long
+			OutErrorMessage = FText::Format( LOCTEXT("AssetCookingPathTooLong", 
+				"The path to the asset is too long for cooking, the maximum is '{0}' characters.\nPlease choose a shorter name for the asset or create it in a shallower folder structure with shorter folder names."),
+				FText::FromString(FString::Printf(TEXT("%d"), MAX_PROJECTED_COOKING_PATH)) );
+			// Return false to indicate that the user should enter a new name
+			return false;
+		}
+
 		// Make sure we are not creating an path that is too long for the OS
-		if ( ObjectPathStr.Len() > PLATFORM_MAX_FILEPATH_LENGTH - MAX_CLASS_NAME_LENGTH )
+		const FString RelativePathFilename = FPackageName::LongPackageNameToFilename(NewPackageName, FPackageName::GetAssetPackageExtension());	// full relative path with name + extension
+		const FString FullPath = FPaths::ConvertRelativePathToFull(RelativePathFilename);	// path to file on disk
+		if ( ObjectPathStr.Len() > (PLATFORM_MAX_FILEPATH_LENGTH - MAX_CLASS_NAME_LENGTH) || FullPath.Len() > PLATFORM_MAX_FILEPATH_LENGTH )
 		{
 			// The full path for the asset is too long
 			OutErrorMessage = FText::Format( LOCTEXT("AssetPathTooLong", 
-				"The full path for the asset is too deep, the maximum is '{0}'. Please choose a shorter name for the asset or create it in a shallower folder structure."), 
+				"The full path for the asset is too deep, the maximum is '{0}'. \nPlease choose a shorter name for the asset or create it in a shallower folder structure."), 
 				FText::FromString(FString::Printf(TEXT("%d"), PLATFORM_MAX_FILEPATH_LENGTH)) );
 			// Return false to indicate that the user should enter a new name
 			return false;
diff --git a/Engine/Source/Editor/GameProjectGeneration/Private/GameProjectUtils.cpp b/Engine/Source/Editor/GameProjectGeneration/Private/GameProjectUtils.cpp
--- a/Engine/Source/Editor/GameProjectGeneration/Private/GameProjectUtils.cpp
+++ b/Engine/Source/Editor/GameProjectGeneration/Private/GameProjectUtils.cpp
@@ -15,7 +15,7 @@
 #define LOCTEXT_NAMESPACE "GameProjectUtils"
 
 #define MAX_PROJECT_PATH_BUFFER_SPACE 130 // Leave a reasonable buffer of additional characters to account for files created in the content directory during or after project generation
-#define MAX_PROJECT_NAME_LENGTH 22 // Enforce a reasonable project name length so the path is not too long for PLATFORM_MAX_FILEPATH_LENGTH
+#define MAX_PROJECT_NAME_LENGTH 20 // Enforce a reasonable project name length so the path is not too long for PLATFORM_MAX_FILEPATH_LENGTH
 checkAtCompileTime(PLATFORM_MAX_FILEPATH_LENGTH - MAX_PROJECT_PATH_BUFFER_SPACE > 0, filesystem_path_shorter_than_project_creation_buffer_space);
 
 #define MAX_CLASS_NAME_LENGTH 32 // Enforce a reasonable class name length so the path is not too long for PLATFORM_MAX_FILEPATH_LENGTH
