Content Browser: Enter and SpaceBar keys can now be remapped.

#ttp 333671 - EDITOR: Content Browser: Customizable shortcut for opening assets and folders

#branch UE4

#add Added two new commands so that the Enter and SpaceBar functionality can be remapped.

#change Renamed EAssetTypeActivationMethod types so they don't mention which key was used, but the action that key performed.

diff --git a/Engine/Source/Developer/AssetTools/Private/AssetTypeActions/AssetTypeActions_Redirector.cpp b/Engine/Source/Developer/AssetTools/Private/AssetTypeActions/AssetTypeActions_Redirector.cpp
--- a/Engine/Source/Developer/AssetTools/Private/AssetTypeActions/AssetTypeActions_Redirector.cpp
+++ b/Engine/Source/Developer/AssetTools/Private/AssetTypeActions/AssetTypeActions_Redirector.cpp
@@ -34,7 +34,7 @@ void FAssetTypeActions_Redirector::GetActions( const TArray<UObject*>& InObjects
 
 void FAssetTypeActions_Redirector::AssetsActivated( const TArray<UObject*>& InObjects, EAssetTypeActivationMethod::Type ActivationType )
 {
-	if ( ActivationType == EAssetTypeActivationMethod::DoubleClicked || ActivationType == EAssetTypeActivationMethod::EnterPressed )
+	if ( ActivationType == EAssetTypeActivationMethod::DoubleClicked || ActivationType == EAssetTypeActivationMethod::Opened )
 	{
 		// Sync to the target instead of opening an editor when double clicked
 		TArray<UObjectRedirector*> Redirectors;
diff --git a/Engine/Source/Developer/AssetTools/Private/AssetTypeActions/AssetTypeActions_SoundBase.cpp b/Engine/Source/Developer/AssetTools/Private/AssetTypeActions/AssetTypeActions_SoundBase.cpp
--- a/Engine/Source/Developer/AssetTools/Private/AssetTypeActions/AssetTypeActions_SoundBase.cpp
+++ b/Engine/Source/Developer/AssetTools/Private/AssetTypeActions/AssetTypeActions_SoundBase.cpp
@@ -44,7 +44,7 @@ bool FAssetTypeActions_SoundBase::CanExecutePlayCommand(TArray<TWeakObjectPtr<US
 
 void FAssetTypeActions_SoundBase::AssetsActivated( const TArray<UObject*>& InObjects, EAssetTypeActivationMethod::Type ActivationType )
 {
-	if (ActivationType == EAssetTypeActivationMethod::SpacePressed)
+	if (ActivationType == EAssetTypeActivationMethod::Previewed)
 	{
 		USoundBase* TargetSound = NULL;
 
diff --git a/Engine/Source/Developer/AssetTools/Public/AssetTypeActions_Base.h b/Engine/Source/Developer/AssetTools/Public/AssetTypeActions_Base.h
--- a/Engine/Source/Developer/AssetTools/Public/AssetTypeActions_Base.h
+++ b/Engine/Source/Developer/AssetTools/Public/AssetTypeActions_Base.h
@@ -27,7 +27,7 @@ public:
 	
 	virtual void AssetsActivated( const TArray<UObject*>& InObjects, EAssetTypeActivationMethod::Type ActivationType ) OVERRIDE
 	{
-		if ( ActivationType == EAssetTypeActivationMethod::DoubleClicked || ActivationType == EAssetTypeActivationMethod::EnterPressed )
+		if ( ActivationType == EAssetTypeActivationMethod::DoubleClicked || ActivationType == EAssetTypeActivationMethod::Opened )
 		{
 			if ( InObjects.Num() == 1 )
 			{
diff --git a/Engine/Source/Developer/AssetTools/Public/IAssetTypeActions.h b/Engine/Source/Developer/AssetTools/Public/IAssetTypeActions.h
--- a/Engine/Source/Developer/AssetTools/Public/IAssetTypeActions.h
+++ b/Engine/Source/Developer/AssetTools/Public/IAssetTypeActions.h
@@ -22,8 +22,8 @@ namespace EAssetTypeActivationMethod
 	enum Type
 	{
 		DoubleClicked,
-		EnterPressed,
-		SpacePressed
+		Opened,
+		Previewed
 	};
 }
 
diff --git a/Engine/Source/Editor/ContentBrowser/Private/ContentBrowserCommands.cpp b/Engine/Source/Editor/ContentBrowser/Private/ContentBrowserCommands.cpp
--- a/Engine/Source/Editor/ContentBrowser/Private/ContentBrowserCommands.cpp
+++ b/Engine/Source/Editor/ContentBrowser/Private/ContentBrowserCommands.cpp
@@ -7,6 +7,8 @@
 
 void FContentBrowserCommands::RegisterCommands()
 {
+	UI_COMMAND( OpenAssetsOrFolders, "Open Assets or Folders", "Opens the selected assets or folders, depending on the selection", EUserInterfaceActionType::Button, FInputGesture(EKeys::Enter) );
+	UI_COMMAND( PreviewAssets, "Preview Assets", "Loads the selected assets and previews them if possible", EUserInterfaceActionType::Button, FInputGesture(EKeys::SpaceBar) );
 	UI_COMMAND( DirectoryUp, "Up", "Up to parent directory", EUserInterfaceActionType::Button, FInputGesture(EModifierKey::Alt, EKeys::Up) );
 }
 
diff --git a/Engine/Source/Editor/ContentBrowser/Private/ContentBrowserCommands.h b/Engine/Source/Editor/ContentBrowser/Private/ContentBrowserCommands.h
--- a/Engine/Source/Editor/ContentBrowser/Private/ContentBrowserCommands.h
+++ b/Engine/Source/Editor/ContentBrowser/Private/ContentBrowserCommands.h
@@ -14,5 +14,7 @@ public:
 
 public:
 
+	TSharedPtr< FUICommandInfo > OpenAssetsOrFolders;
+	TSharedPtr< FUICommandInfo > PreviewAssets;
 	TSharedPtr< FUICommandInfo > DirectoryUp;
 };
\ No newline at end of file
diff --git a/Engine/Source/Editor/ContentBrowser/Private/SAssetPicker.cpp b/Engine/Source/Editor/ContentBrowser/Private/SAssetPicker.cpp
--- a/Engine/Source/Editor/ContentBrowser/Private/SAssetPicker.cpp
+++ b/Engine/Source/Editor/ContentBrowser/Private/SAssetPicker.cpp
@@ -260,7 +260,7 @@ FReply SAssetPicker::OnKeyDown(const FGeometry& MyGeometry, const FKeyboardEvent
 	else if (InKeyboardEvent.GetKey() == EKeys::Enter)
 	{
 		TArray<FAssetData> SelectionSet = AssetViewPtr->GetSelectedAssets();
-		HandleAssetsActivated(SelectionSet, EAssetTypeActivationMethod::EnterPressed);
+		HandleAssetsActivated(SelectionSet, EAssetTypeActivationMethod::Opened);
 
 		return FReply::Handled();
 	}
@@ -297,7 +297,7 @@ void SAssetPicker::OnSearchBoxCommitted(const FText& InSearchText, ETextCommit::
 			AssetViewPtr->AdjustActiveSelection(1);
 			SelectionSet = AssetViewPtr->GetSelectedAssets();
 		}
-		HandleAssetsActivated(SelectionSet, EAssetTypeActivationMethod::EnterPressed);
+		HandleAssetsActivated(SelectionSet, EAssetTypeActivationMethod::Opened);
 	}
 }
 
@@ -328,7 +328,7 @@ void SAssetPicker::HandleAssetsActivated(const TArray<FAssetData>& ActivatedAsse
 			OnAssetDoubleClicked.ExecuteIfBound(ActivatedAssets[0]);
 		}
 	}
-	else if (ActivationMethod == EAssetTypeActivationMethod::EnterPressed)
+	else if (ActivationMethod == EAssetTypeActivationMethod::Opened)
 	{
 		OnAssetEnterPressed.ExecuteIfBound(ActivatedAssets);
 	}
diff --git a/Engine/Source/Editor/ContentBrowser/Private/SAssetView.cpp b/Engine/Source/Editor/ContentBrowser/Private/SAssetView.cpp
--- a/Engine/Source/Editor/ContentBrowser/Private/SAssetView.cpp
+++ b/Engine/Source/Editor/ContentBrowser/Private/SAssetView.cpp
@@ -1150,28 +1150,6 @@ FReply SAssetView::OnKeyDown( const FGeometry& MyGeometry, const FKeyboardEvent&
 		}
 	}
 
-	if ( InKeyboardEvent.GetKey() == EKeys::Enter )
-	{
-		TArray<FAssetData> SelectedAssets = GetSelectedAssets();
-		TArray<FString> SelectedFolders = GetSelectedFolders();
-		if(SelectedAssets.Num() > 0 && SelectedFolders.Num() == 0)
-		{
-			OnAssetsActivated.ExecuteIfBound( SelectedAssets, EAssetTypeActivationMethod::EnterPressed );
-		}
-		else if(SelectedAssets.Num() == 0 && SelectedFolders.Num() > 0)
-		{
-			OnPathSelected.ExecuteIfBound(SelectedFolders[0]);
-		}
-
-		return FReply::Handled();
-	}
-	else if ( InKeyboardEvent.GetKey() == EKeys::SpaceBar )
-	{
-		OnAssetsActivated.ExecuteIfBound( GetSelectedAssets(), EAssetTypeActivationMethod::SpacePressed );
-
-		return FReply::Handled();
-	}
-
 	return FReply::Unhandled();
 }
 
@@ -2454,6 +2432,25 @@ void SAssetView::RequestScrollIntoView(const TSharedPtr<FAssetViewItem>& Item)
 	}
 }
 
+void SAssetView::OnOpenAssetsOrFolders()
+{
+	TArray<FAssetData> SelectedAssets = GetSelectedAssets();
+	TArray<FString> SelectedFolders = GetSelectedFolders();
+	if (SelectedAssets.Num() > 0 && SelectedFolders.Num() == 0)
+	{
+		OnAssetsActivated.ExecuteIfBound(SelectedAssets, EAssetTypeActivationMethod::Opened);
+	}
+	else if (SelectedAssets.Num() == 0 && SelectedFolders.Num() > 0)
+	{
+		OnPathSelected.ExecuteIfBound(SelectedFolders[0]);
+	}
+}
+
+void SAssetView::OnPreviewAssets()
+{
+	OnAssetsActivated.ExecuteIfBound(GetSelectedAssets(), EAssetTypeActivationMethod::Previewed);
+}
+
 void SAssetView::ClearSelection()
 {
 	switch ( GetCurrentViewType() )
diff --git a/Engine/Source/Editor/ContentBrowser/Private/SAssetView.h b/Engine/Source/Editor/ContentBrowser/Private/SAssetView.h
--- a/Engine/Source/Editor/ContentBrowser/Private/SAssetView.h
+++ b/Engine/Source/Editor/ContentBrowser/Private/SAssetView.h
@@ -225,6 +225,12 @@ public:
 	virtual FReply OnMouseWheel( const FGeometry& MyGeometry, const FPointerEvent& MouseEvent ) OVERRIDE;
 	virtual void OnKeyboardFocusChanging( const FWeakWidgetPath& PreviousFocusPath, const FWidgetPath& NewWidgetPath ) OVERRIDE;
 
+	/** Opens the selected assets or folders, depending on the selection */
+	void OnOpenAssetsOrFolders();
+
+	/** Loads the selected assets and previews them if possible */
+	void OnPreviewAssets();
+
 	/** Clears the selection of all the lists in the view */
 	void ClearSelection();
 
diff --git a/Engine/Source/Editor/ContentBrowser/Private/SContentBrowser.cpp b/Engine/Source/Editor/ContentBrowser/Private/SContentBrowser.cpp
--- a/Engine/Source/Editor/ContentBrowser/Private/SContentBrowser.cpp
+++ b/Engine/Source/Editor/ContentBrowser/Private/SContentBrowser.cpp
@@ -591,6 +591,14 @@ void SContentBrowser::BindCommands()
 {
 	Commands = TSharedPtr< FUICommandList >(new FUICommandList);
 
+	Commands->MapAction(FContentBrowserCommands::Get().OpenAssetsOrFolders, FUIAction(
+		FExecuteAction::CreateSP(this, &SContentBrowser::OnOpenAssetsOrFolders)
+		));
+
+	Commands->MapAction(FContentBrowserCommands::Get().PreviewAssets, FUIAction(
+		FExecuteAction::CreateSP(this, &SContentBrowser::OnPreviewAssets)
+		));
+
 	Commands->MapAction( FContentBrowserCommands::Get().DirectoryUp, FUIAction(
 		FExecuteAction::CreateSP( this, &SContentBrowser::OnDirectoryUp )
 		));
@@ -1343,7 +1351,7 @@ void SContentBrowser::OnAssetsActivated(const TArray<FAssetData>& ActivatedAsset
 	}
 
 	// Finally, open a simple asset editor for all assets which do not have asset type actions if activating with enter or double click
-	if ( ActivationMethod == EAssetTypeActivationMethod::DoubleClicked || ActivationMethod == EAssetTypeActivationMethod::EnterPressed )
+	if ( ActivationMethod == EAssetTypeActivationMethod::DoubleClicked || ActivationMethod == EAssetTypeActivationMethod::Opened )
 	{
 		ContentBrowserUtils::OpenEditorForAsset(ObjectsWithoutTypeActions);
 	}
@@ -1447,6 +1455,16 @@ FReply SContentBrowser::ForwardClicked()
 	return FReply::Handled();
 }
 
+void SContentBrowser::OnOpenAssetsOrFolders()
+{
+	AssetViewPtr->OnOpenAssetsOrFolders();
+}
+
+void SContentBrowser::OnPreviewAssets()
+{
+	AssetViewPtr->OnPreviewAssets();
+}
+
 FReply SContentBrowser::OnDirectoryUpClicked()
 {
 	OnDirectoryUp();
diff --git a/Engine/Source/Editor/ContentBrowser/Private/SContentBrowser.h b/Engine/Source/Editor/ContentBrowser/Private/SContentBrowser.h
--- a/Engine/Source/Editor/ContentBrowser/Private/SContentBrowser.h
+++ b/Engine/Source/Editor/ContentBrowser/Private/SContentBrowser.h
@@ -195,6 +195,12 @@ private:
 	/** Handler for clicking the history forward button */
 	FReply ForwardClicked();
 
+	/** Handler for opening assets or folders */
+	void OnOpenAssetsOrFolders();
+
+	/** Handler for previewing assets */
+	void OnPreviewAssets();
+
 	/** Handler for clicking the directory up button */
 	void OnDirectoryUp();
 
diff --git a/Engine/Source/Editor/UnrealEd/Private/Dialogs/SDeleteAssetsDialog.cpp b/Engine/Source/Editor/UnrealEd/Private/Dialogs/SDeleteAssetsDialog.cpp
--- a/Engine/Source/Editor/UnrealEd/Private/Dialogs/SDeleteAssetsDialog.cpp
+++ b/Engine/Source/Editor/UnrealEd/Private/Dialogs/SDeleteAssetsDialog.cpp
@@ -690,7 +690,7 @@ FText SDeleteAssetsDialog::GetReferencingAssetsEmptyText() const
 void SDeleteAssetsDialog::OnAssetsActivated(const TArray<FAssetData>& ActivatedAssets, EAssetTypeActivationMethod::Type ActivationMethod)
 {
 	// Open a simple asset editor for all assets which do not have asset type actions if activating with enter or double click
-	if ( ActivationMethod == EAssetTypeActivationMethod::DoubleClicked || ActivationMethod == EAssetTypeActivationMethod::EnterPressed )
+	if ( ActivationMethod == EAssetTypeActivationMethod::DoubleClicked || ActivationMethod == EAssetTypeActivationMethod::Opened )
 	{
 		ParentWindow.Get()->RequestDestroyWindow();
 
