Removed "New Folder" from the New Asset context menu if a collection is currently selected.
#ttp 334526 - EDITOR: Content Browser: CRASH: Creating a new folder in a collection results in an ensure
#branch UE4
#proj Editor.ContentBrowser
#change In SContentBrowser::MakeCreateAssetContextMenu, if a collection is selected, pass an unbound FOnNewFolderRequested delegate to FNewAssetContextMenu::MakeContextMenu, so that the "New Folder" item is not added.

diff --git a/Engine/Source/Editor/ContentBrowser/Private/SContentBrowser.cpp b/Engine/Source/Editor/ContentBrowser/Private/SContentBrowser.cpp
--- a/Engine/Source/Editor/ContentBrowser/Private/SContentBrowser.cpp
+++ b/Engine/Source/Editor/ContentBrowser/Private/SContentBrowser.cpp
@@ -1235,12 +1235,19 @@ TSharedRef<SWidget> SContentBrowser::MakeCreateAssetContextMenu()
 	TSharedPtr<FExtender> MenuExtender = FExtender::Combine(Extenders);
 
 	FMenuBuilder MenuBuilder(/*bInShouldCloseWindowAfterMenuSelection=*/true, NULL, MenuExtender);
-		
+
+	// Only add "New Folder" item if we do not have a collection selected
+	FNewAssetContextMenu::FOnNewFolderRequested OnNewFolderRequested;
+	if (CollectionViewPtr->GetSelectedCollections().Num() == 0)
+	{
+		OnNewFolderRequested = FNewAssetContextMenu::FOnNewFolderRequested::CreateSP(this, &SContentBrowser::NewFolderRequested);
+	}
+
 	FNewAssetContextMenu::MakeContextMenu(
 		MenuBuilder, 
 		CurrentPath, 
 		FNewAssetContextMenu::FOnNewAssetRequested::CreateSP(this, &SContentBrowser::NewAssetRequested),
-		FNewAssetContextMenu::FOnNewFolderRequested::CreateSP(this, &SContentBrowser::NewFolderRequested));
+		OnNewFolderRequested);
 
 	FDisplayMetrics DisplayMetrics;
 	FSlateApplication::Get().GetDisplayMetrics( DisplayMetrics );
