UE4: Removed many unnecessary checks in the editor to special-case assets found in umap files.

diff --git a/Engine/Source/Editor/ContentBrowser/Private/AssetContextMenu.cpp b/Engine/Source/Editor/ContentBrowser/Private/AssetContextMenu.cpp
--- a/Engine/Source/Editor/ContentBrowser/Private/AssetContextMenu.cpp
+++ b/Engine/Source/Editor/ContentBrowser/Private/AssetContextMenu.cpp
@@ -531,12 +531,6 @@ void FAssetContextMenu::GetSelectedAssets(TArray<UObject*>& Assets, bool SkipRed
 {
 	for (int32 AssetIdx = 0; AssetIdx < SelectedAssets.Num(); ++AssetIdx)
 	{
-		if (!SelectedAssets[AssetIdx].IsAssetLoaded() && FEditorFileUtils::IsMapPackageAsset(SelectedAssets[AssetIdx].ObjectPath.ToString()))
-		{
-			// Don't load assets in map packages
-			continue;
-		}
-
 		if (SkipRedirectors && (SelectedAssets[AssetIdx].AssetClass == UObjectRedirector::StaticClass()->GetFName()))
 		{
 			// Don't operate on Redirectors
@@ -810,12 +804,6 @@ void FAssetContextMenu::ExecuteDelete()
 		{
 			const FAssetData& AssetData = *AssetIt;
 
-			if( !AssetData.IsAssetLoaded() && FEditorFileUtils::IsMapPackageAsset( AssetData.ObjectPath.ToString() ) )
-			{
-				// Don't load assets in map packages
-				continue;
-			}
-
 			if( AssetData.AssetClass == UObjectRedirector::StaticClass()->GetFName() )
 			{
 				// Don't operate on Redirectors
diff --git a/Engine/Source/Editor/ContentBrowser/Private/ContentBrowserUtils.cpp b/Engine/Source/Editor/ContentBrowser/Private/ContentBrowserUtils.cpp
--- a/Engine/Source/Editor/ContentBrowser/Private/ContentBrowserUtils.cpp
+++ b/Engine/Source/Editor/ContentBrowser/Private/ContentBrowserUtils.cpp
@@ -278,22 +278,23 @@ bool ContentBrowserUtils::LoadAssetsIfNeeded(const TArray<FString>& ObjectPaths,
 
 	// Build a list of unloaded assets
 	TArray<FString> UnloadedObjectPaths;
+	bool bAtLeastOneUnloadedMap = false;
 	for (int32 PathIdx = 0; PathIdx < ObjectPaths.Num(); ++PathIdx)
 	{
 		const FString& ObjectPath = ObjectPaths[PathIdx];
 
-		if ( !FEditorFileUtils::IsMapPackageAsset(ObjectPath) )
+		UObject* FoundObject = FindObject<UObject>(NULL, *ObjectPath);
+		if ( FoundObject )
 		{
-			UObject* FoundObject = FindObject<UObject>(NULL, *ObjectPath);
-
-			if ( FoundObject )
-			{
-				LoadedObjects.Add(FoundObject);
-			}
-			else
+			LoadedObjects.Add(FoundObject);
+		}
+		else
+		{
+			// Unloaded asset, we will load it later
+			UnloadedObjectPaths.Add(ObjectPath);
+			if ( FEditorFileUtils::IsMapPackageAsset(ObjectPath) )
 			{
-				// Unloaded asset, we will load it later
-				UnloadedObjectPaths.Add(ObjectPath);
+				bAtLeastOneUnloadedMap = true;
 			}
 		}
 	}
@@ -317,7 +318,7 @@ bool ContentBrowserUtils::LoadAssetsIfNeeded(const TArray<FString>& ObjectPaths,
 	if ( UnloadedObjectPaths.Num() > 0 )
 	{
 		// Get the maximum objects to load before displaying the slow task
-		const bool bShowProgressDialog = UnloadedObjectPaths.Num() > GetDefault<UContentBrowserSettings>()->NumObjectsToLoadBeforeWarning;
+		const bool bShowProgressDialog = (UnloadedObjectPaths.Num() > GetDefault<UContentBrowserSettings>()->NumObjectsToLoadBeforeWarning) || bAtLeastOneUnloadedMap;
 		GWarn->BeginSlowTask(LOCTEXT("LoadingObjects", "Loading Objects..."), bShowProgressDialog);
 
 		GIsEditorLoadingPackage = true;
@@ -383,14 +384,11 @@ bool ContentBrowserUtils::ShouldPromptToLoadAssets(const TArray<FString>& Object
 	{
 		const FString& ObjectPath = ObjectPaths[PathIdx];
 
-		if ( !FEditorFileUtils::IsMapPackageAsset(ObjectPath) )
+		UObject* FoundObject = FindObject<UObject>(NULL, *ObjectPath);
+		if ( !FoundObject )
 		{
-			UObject* FoundObject = FindObject<UObject>(NULL, *ObjectPath);
-			if ( !FoundObject )
-			{
-				// Unloaded asset, we will load it later
-				OutUnloadedObjects.Add(ObjectPath);
-			}
+			// Unloaded asset, we will load it later
+			OutUnloadedObjects.Add(ObjectPath);
 		}
 	}
 
@@ -623,12 +621,6 @@ TArray<UPackage*> ContentBrowserUtils::LoadPackages(const TArray<FString>& Packa
 			continue;
 		}
 
-		if ( FEditorFileUtils::IsMapPackageAsset(PackageName) )
-		{
-			// Map package. Skip it.
-			continue;
-		}
-
 		UPackage* Package = FindPackage(NULL, *PackageName);
 
 		if ( Package != NULL )
@@ -1015,12 +1007,6 @@ void ContentBrowserUtils::GetObjectsInAssetData(const TArray<FAssetData>& AssetL
 	{
 		const FAssetData& AssetData = AssetList[AssetIdx];
 
-		if ( !AssetData.IsAssetLoaded() && FEditorFileUtils::IsMapPackageAsset(AssetData.ObjectPath.ToString()) )
-		{
-			// Don't load assets in map packages
-			continue;
-		}
-
 		UObject* Obj = AssetData.GetAsset();
 		if (Obj)
 		{
diff --git a/Engine/Source/Editor/ContentBrowser/Private/SAssetView.cpp b/Engine/Source/Editor/ContentBrowser/Private/SAssetView.cpp
--- a/Engine/Source/Editor/ContentBrowser/Private/SAssetView.cpp
+++ b/Engine/Source/Editor/ContentBrowser/Private/SAssetView.cpp
@@ -528,22 +528,19 @@ void SAssetView::DuplicateAsset(const FString& PackagePath, const TWeakObjectPtr
 
 void SAssetView::RenameAsset(const FAssetData& ItemToRename)
 {
-	if ( !FEditorFileUtils::IsMapPackageAsset(ItemToRename.ObjectPath.ToString()) )
+	for ( auto ItemIt = FilteredAssetItems.CreateConstIterator(); ItemIt; ++ItemIt )
 	{
-		for ( auto ItemIt = FilteredAssetItems.CreateConstIterator(); ItemIt; ++ItemIt )
+		const TSharedPtr<FAssetViewItem>& Item = *ItemIt;
+		if ( Item.IsValid() && Item->GetType() != EAssetItemType::Folder )	
 		{
-			const TSharedPtr<FAssetViewItem>& Item = *ItemIt;
-			if ( Item.IsValid() && Item->GetType() != EAssetItemType::Folder )	
+			const TSharedPtr<FAssetViewAsset>& ItemAsAsset = StaticCastSharedPtr<FAssetViewAsset>(Item);
+			if ( ItemAsAsset->Data.ObjectPath == ItemToRename.ObjectPath )
 			{
-				const TSharedPtr<FAssetViewAsset>& ItemAsAsset = StaticCastSharedPtr<FAssetViewAsset>(Item);
-				if ( ItemAsAsset->Data.ObjectPath == ItemToRename.ObjectPath )
-				{
-					ItemAsAsset->bRenameWhenScrolledIntoview = true;
+				ItemAsAsset->bRenameWhenScrolledIntoview = true;
 
-					SetSelection(Item);
-					RequestScrollIntoView(Item);
-					break;
-				}
+				SetSelection(Item);
+				RequestScrollIntoView(Item);
+				break;
 			}
 		}
 	}
@@ -3045,7 +3042,7 @@ FReply SAssetView::OnDraggingAssetItem( const FGeometry& MyGeometry, const FPoin
 					// If dragging a class, send though an FAssetData whose name is null and class is this class' name
 					InAssetData.Add(AssetData);
 				}
-				else if ( AssetData.IsAssetLoaded() || !FEditorFileUtils::IsMapPackageAsset(AssetData.ObjectPath.ToString()) )
+				else if ( AssetData.IsAssetLoaded() )
 				{
 					InAssetData.Add(AssetData);
 				}
diff --git a/Engine/Source/Editor/ContentBrowser/Private/SContentBrowser.cpp b/Engine/Source/Editor/ContentBrowser/Private/SContentBrowser.cpp
--- a/Engine/Source/Editor/ContentBrowser/Private/SContentBrowser.cpp
+++ b/Engine/Source/Editor/ContentBrowser/Private/SContentBrowser.cpp
@@ -718,13 +718,13 @@ void SContentBrowser::LoadSelectedObjectsIfNeeded()
 	// Load every asset that isn't already in memory
 	for ( auto AssetIt = SelectedAssets.CreateConstIterator(); AssetIt; ++AssetIt )
 	{
-		if ( !(*AssetIt).IsAssetLoaded() && FEditorFileUtils::IsMapPackageAsset((*AssetIt).ObjectPath.ToString()) )
-		{
-			// Don't load assets in map packages
-			continue;
-		}
+		const FAssetData& AssetData = *AssetIt;
+		const bool bShowProgressDialog = (!AssetData.IsAssetLoaded() && FEditorFileUtils::IsMapPackageAsset(AssetData.ObjectPath.ToString()));
+		GWarn->BeginSlowTask(LOCTEXT("LoadingObjects", "Loading Objects..."), bShowProgressDialog);
 
 		(*AssetIt).GetAsset();
+
+		GWarn->EndSlowTask();
 	}
 
 	// Sync the global selection set if we are the primary browser
@@ -1323,14 +1323,14 @@ void SContentBrowser::OnAssetsActivated(const TArray<FAssetData>& ActivatedAsset
 	// This way individual asset type actions will get a batched list of assets to operate on
 	for ( auto AssetIt = ActivatedAssets.CreateConstIterator(); AssetIt; ++AssetIt )
 	{
-		if ( !(*AssetIt).IsAssetLoaded() && FEditorFileUtils::IsMapPackageAsset((*AssetIt).ObjectPath.ToString()) )
-		{
-			// Skip unloaded assets in map packages, it is illegal to load them now
-			continue;
-		}
+		const FAssetData& AssetData = *AssetIt;
+		const bool bShowProgressDialog = (!AssetData.IsAssetLoaded() && FEditorFileUtils::IsMapPackageAsset(AssetData.ObjectPath.ToString()));
+		GWarn->BeginSlowTask(LOCTEXT("LoadingObjects", "Loading Objects..."), bShowProgressDialog);
 
 		UObject* Asset = (*AssetIt).GetAsset();
 
+		GWarn->EndSlowTask();
+
 		if ( Asset != NULL )
 		{
 			FAssetToolsModule& AssetToolsModule = FModuleManager::Get().LoadModuleChecked<FAssetToolsModule>("AssetTools");
diff --git a/Engine/Source/Editor/ContentBrowser/Private/SThumbnailEditModeTools.cpp b/Engine/Source/Editor/ContentBrowser/Private/SThumbnailEditModeTools.cpp
--- a/Engine/Source/Editor/ContentBrowser/Private/SThumbnailEditModeTools.cpp
+++ b/Engine/Source/Editor/ContentBrowser/Private/SThumbnailEditModeTools.cpp
@@ -135,7 +135,7 @@ FReply SThumbnailEditModeTools::OnMouseButtonDown( const FGeometry& MyGeometry,
 	{
 		// Load the asset, unless it is in an unloaded map package or already loaded
 		const FAssetData& AssetData = AssetThumbnail.Pin()->GetAssetData();
-		if ( !AssetData.IsAssetLoaded() && !FEditorFileUtils::IsMapPackageAsset(AssetData.ObjectPath.ToString()) )
+		if ( !AssetData.IsAssetLoaded() )
 		{
 			// Getting the asset loads it.
 			AssetData.GetAsset();
diff --git a/Engine/Source/Editor/UnrealEd/Private/ObjectTools.cpp b/Engine/Source/Editor/UnrealEd/Private/ObjectTools.cpp
--- a/Engine/Source/Editor/UnrealEd/Private/ObjectTools.cpp
+++ b/Engine/Source/Editor/UnrealEd/Private/ObjectTools.cpp
@@ -1243,11 +1243,7 @@ namespace ObjectTools
 			FString PackageFilename;
 			if( FPackageName::DoesPackageExist( Package->GetName(), NULL, &PackageFilename ) )
 			{
-				// Cull out non-UAssets
-				if ( FPaths::GetExtension(PackageFilename, /*bIncludeDot=*/true).ToLower() != FPackageName::GetAssetPackageExtension() )
-				{
-					PackagesToDelete.RemoveAt(PackageIdx);
-				}
+				PackagesToDelete.RemoveAt(PackageIdx);
 			}
 			else
 			{
@@ -1361,13 +1357,6 @@ namespace ObjectTools
 					continue;
 				}
 
-				if ( FPaths::GetExtension(PackageFilename, /*bIncludeDot=*/true).ToLower() != FPackageName::GetAssetPackageExtension() )
-				{
-					// Only delete UAsset packages because that is what we checked for in ShowDeleteConfirmationDialog()
-					PackagesToDelete.RemoveAt(PackageIdx);
-					continue;
-				}
-
 				PackageFilesToDelete.Add(PackageFilename);
 				Cast<UPackage>(Package)->SetDirtyFlag(false);
 				if ( ISourceControlModule::Get().IsEnabled() )
