Asset picker now supports rename functionality.
#TTP 326521 - EDITOR: Feature Request: PERSONA: Have ability to rename assets in Asset Browser with F2 and right click

#branch UE4

#proj Editor.ContentBrowser

#add Added support for the rename UICommand to the Asset picker using a UICommandList and processing it in its onkeydown method.

diff --git a/Engine/Source/Editor/ContentBrowser/Private/SAssetPicker.cpp b/Engine/Source/Editor/ContentBrowser/Private/SAssetPicker.cpp
--- a/Engine/Source/Editor/ContentBrowser/Private/SAssetPicker.cpp
+++ b/Engine/Source/Editor/ContentBrowser/Private/SAssetPicker.cpp
@@ -18,6 +18,8 @@ void SAssetPicker::Construct( const FArguments& InArgs )
 		FilterCollection->Add( InArgs._AssetPickerConfig.ExtraFilters->GetFilterAtIndex( Index ) );
 	}
 
+	BindCommands();
+
 	OnAssetsActivated = InArgs._AssetPickerConfig.OnAssetsActivated;
 	OnAssetSelected = InArgs._AssetPickerConfig.OnAssetSelected;
 	OnAssetClicked = InArgs._AssetPickerConfig.OnAssetClicked;
@@ -272,6 +274,11 @@ FReply SAssetPicker::OnKeyDown(const FGeometry& MyGeometry, const FKeyboardEvent
 		return FReply::Handled();
 	}
 
+	if (Commands->ProcessCommandBindings(InKeyboardEvent))
+	{
+		return FReply::Handled();
+	}
+
 	return FReply::Unhandled();
 }
 
@@ -373,4 +380,44 @@ ESlateCheckBoxState::Type SAssetPicker::GetShowOtherDevelopersCheckState() const
 	return OtherDevelopersFilter->GetShowOtherDeveloperAssets() ? ESlateCheckBoxState::Checked : ESlateCheckBoxState::Unchecked;
 }
 
+void SAssetPicker::OnRenameRequested() const
+{
+	TArray< FAssetData > AssetViewSelectedAssets = AssetViewPtr->GetSelectedAssets();
+	TArray< FString > SelectedFolders = AssetViewPtr->GetSelectedFolders();
+
+	if ( AssetViewSelectedAssets.Num() == 1 && SelectedFolders.Num() == 0 )
+	{
+		// Don't operate on Redirectors
+		if ( AssetViewSelectedAssets[0].AssetClass != UObjectRedirector::StaticClass()->GetFName() )
+		{
+			AssetViewPtr->RenameAsset(AssetViewSelectedAssets[0]);
+		}
+	}
+	else if ( AssetViewSelectedAssets.Num() == 0 && SelectedFolders.Num() == 1 )
+	{
+		AssetViewPtr->RenameFolder(SelectedFolders[0]);
+	}
+}
+
+bool SAssetPicker::CanExecuteRenameRequested()
+{
+	TArray< FAssetData > AssetViewSelectedAssets = AssetViewPtr->GetSelectedAssets();
+	TArray< FString > SelectedFolders = AssetViewPtr->GetSelectedFolders();
+
+	const bool bCanRenameFolder = ( AssetViewSelectedAssets.Num() == 0 && SelectedFolders.Num() == 1 );
+	const bool bCanRenameAsset = ( AssetViewSelectedAssets.Num() == 1 && SelectedFolders.Num() == 0 ) && ( AssetViewSelectedAssets[0].AssetClass != UObjectRedirector::StaticClass()->GetFName() ); 
+
+	return	bCanRenameFolder || bCanRenameAsset;
+}
+
+void SAssetPicker::BindCommands()
+{
+	Commands = MakeShareable(new FUICommandList);
+	// bind commands
+	Commands->MapAction( FGenericCommands::Get().Rename, FUIAction(
+		FExecuteAction::CreateSP( this, &SAssetPicker::OnRenameRequested ),
+		FCanExecuteAction::CreateSP( this, &SAssetPicker::CanExecuteRenameRequested )
+		));
+}
+
 #undef LOCTEXT_NAMESPACE
\ No newline at end of file
diff --git a/Engine/Source/Editor/ContentBrowser/Private/SAssetPicker.h b/Engine/Source/Editor/ContentBrowser/Private/SAssetPicker.h
--- a/Engine/Source/Editor/ContentBrowser/Private/SAssetPicker.h
+++ b/Engine/Source/Editor/ContentBrowser/Private/SAssetPicker.h
@@ -63,6 +63,15 @@ private:
 	/** Gets if showing other developers assets */
 	ESlateCheckBoxState::Type GetShowOtherDevelopersCheckState() const;
 
+	/** Called upon the Rename UICommand being executed, sends rename request to the asset view */
+	void OnRenameRequested() const;
+
+	/** Returns true if the user is able to execute a rename request */
+	bool CanExecuteRenameRequested();
+
+	/** Bind our UI commands */
+	void BindCommands();
+
 private:
 	/** The asset view widget */
 	TSharedPtr<SAssetView> AssetViewPtr;
@@ -100,4 +109,7 @@ private:
 
 	/** Initial filter we were created with, needed reset asset view after we have custom filtered */
 	FARFilter InitialBackendFilter;
+
+	/** UICommand list, holds list of actions for processing */
+	TSharedPtr< FUICommandList > Commands;
 };
\ No newline at end of file
