UE4: Fixed a bug where unloaded assets were not dragable in the content browser.

diff --git a/Engine/Source/Editor/ContentBrowser/Private/SAssetView.cpp b/Engine/Source/Editor/ContentBrowser/Private/SAssetView.cpp
--- a/Engine/Source/Editor/ContentBrowser/Private/SAssetView.cpp
+++ b/Engine/Source/Editor/ContentBrowser/Private/SAssetView.cpp
@@ -3042,7 +3042,7 @@ FReply SAssetView::OnDraggingAssetItem( const FGeometry& MyGeometry, const FPoin
 					// If dragging a class, send though an FAssetData whose name is null and class is this class' name
 					InAssetData.Add(AssetData);
 				}
-				else if ( AssetData.IsAssetLoaded() )
+				else
 				{
 					InAssetData.Add(AssetData);
 				}
