Fixed up source control state delegate to use shared pointer semantics

TTP# 335201  - ContentBrowser: Improve idle performance

diff --git a/Engine/Source/Editor/ContentBrowser/Private/AssetViewWidgets.cpp b/Engine/Source/Editor/ContentBrowser/Private/AssetViewWidgets.cpp
--- a/Engine/Source/Editor/ContentBrowser/Private/AssetViewWidgets.cpp
+++ b/Engine/Source/Editor/ContentBrowser/Private/AssetViewWidgets.cpp
@@ -116,8 +116,6 @@ SAssetViewItem::~SAssetViewItem()
 	OnItemDestroyed.ExecuteIfBound( AssetItem );
 
 	SetForceMipLevelsToBeResident(false);
-
-	ISourceControlModule::Get().GetProvider().UnregisterSourceControlStateChanged(FSourceControlStateChanged::FDelegate::CreateRaw(this, &SAssetViewItem::HandleSourceControlStateChanged));
 }
 
 void SAssetViewItem::Construct( const FArguments& InArgs )
@@ -154,7 +152,7 @@ void SAssetViewItem::Construct( const FArguments& InArgs )
 	SourceControlStateDelay = 0.0f;
 	bSourceControlStateRequested = false;
 
-	ISourceControlModule::Get().GetProvider().RegisterSourceControlStateChanged(FSourceControlStateChanged::FDelegate::CreateRaw(this, &SAssetViewItem::HandleSourceControlStateChanged));
+	ISourceControlModule::Get().GetProvider().RegisterSourceControlStateChanged(FSourceControlStateChanged::FDelegate::CreateSP(this, &SAssetViewItem::HandleSourceControlStateChanged));
 }
 
 void SAssetViewItem::Tick( const FGeometry& AllottedGeometry, const double InCurrentTime, const float InDeltaTime )
