Prevented content browser search filter from being cleared when renaming an asset.
#ttp 335652 - EDITOR: Content browser text search filter should not be wiped when renaming an asset.
#branch UE4
#proj Editor.ContentBrowser
#remove Removed the call to clear the filter textbox in SContentBrowser::SyncToAssets.

diff --git a/Engine/Source/Editor/ContentBrowser/Private/SContentBrowser.cpp b/Engine/Source/Editor/ContentBrowser/Private/SContentBrowser.cpp
--- a/Engine/Source/Editor/ContentBrowser/Private/SContentBrowser.cpp
+++ b/Engine/Source/Editor/ContentBrowser/Private/SContentBrowser.cpp
@@ -678,7 +678,6 @@ void SContentBrowser::SyncToAssets( const TArray<FAssetData>& AssetDataList, con
 
 	// Tell the sources view first so the asset view will be up to date by the time we request the sync
 	PathViewPtr->SyncToAssets(AssetDataList, bAllowImplicitSync);
-	SearchBoxPtr->SetText(FText::GetEmpty());
 	AssetViewPtr->SyncToAssets(AssetDataList);
 }
 
