UE4: Fixed a bug I found while running static analysis tools. The correct collection type string will now be detected when determining the config key.

diff --git a/Engine/Source/Editor/ContentBrowser/Private/CollectionViewUtils.cpp b/Engine/Source/Editor/ContentBrowser/Private/CollectionViewUtils.cpp
--- a/Engine/Source/Editor/ContentBrowser/Private/CollectionViewUtils.cpp
+++ b/Engine/Source/Editor/ContentBrowser/Private/CollectionViewUtils.cpp
@@ -19,12 +19,16 @@ namespace CollectionViewUtils
 		{
 		case ECollectionShareType::CST_System:
 			CollectionTypeStr = "System";
+			break;
 		case ECollectionShareType::CST_Local:
 			CollectionTypeStr = "Local";
+			break;
 		case ECollectionShareType::CST_Private:
 			CollectionTypeStr = "Private";
+			break;
 		case ECollectionShareType::CST_Shared:
 			CollectionTypeStr = "Shared";
+			break;
 		default:
 			break;
 		}
