UE4: Another small fix caught by static code analysis

diff --git a/Engine/Source/Editor/ContentBrowser/Private/SAssetView.cpp b/Engine/Source/Editor/ContentBrowser/Private/SAssetView.cpp
--- a/Engine/Source/Editor/ContentBrowser/Private/SAssetView.cpp
+++ b/Engine/Source/Editor/ContentBrowser/Private/SAssetView.cpp
@@ -3125,7 +3125,7 @@ bool SAssetView::AssetVerifyRenameCommit(const TSharedPtr<FAssetViewItem>& Item,
 		if ( ObjectPathStr.Len() > NAME_SIZE )
 		{
 			// This asset already exists at this location, inform the user and continue
-			OutErrorMessage = OutErrorMessage = LOCTEXT("AssetNameTooLong", "This asset name is too long. Please choose a shorter name.");
+			OutErrorMessage = LOCTEXT("AssetNameTooLong", "This asset name is too long. Please choose a shorter name.");
 			// Return false to indicate that the user should enter a new name
 			return false;
 		}
