Content browser refreshes when package dirty state is updated.
#ttp 334481 - Regression: EDITOR: Content Browser: If you have the modified filter enabled and then save a modified asset, the asset is still visible in the browser
#branch UE4
#proj Editor.ContentBrowser
#add SAssetView registers a delegate with UPackage::PackageDirtyStateUpdatedEvent so that it can be refreshed.

diff --git a/Engine/Source/Editor/ContentBrowser/Private/SAssetView.cpp b/Engine/Source/Editor/ContentBrowser/Private/SAssetView.cpp
--- a/Engine/Source/Editor/ContentBrowser/Private/SAssetView.cpp
+++ b/Engine/Source/Editor/ContentBrowser/Private/SAssetView.cpp
@@ -83,6 +83,9 @@ void SAssetView::Construct( const FArguments& InArgs )
 	FCoreDelegates::OnAssetLoaded.AddSP(this, &SAssetView::OnAssetLoaded);
 	FCoreDelegates::OnObjectPropertyChanged.AddSP(this, &SAssetView::OnObjectPropertyChanged);
 
+	// Listen for when packages' dirty states are updated
+	UPackage::PackageDirtyStateUpdatedEvent.AddSP(this, &SAssetView::OnPackageDirtyStateUpdated);
+
 	// Listen for when view settings are changed
 	UContentBrowserSettings::OnSettingChanged().AddSP(this, &SAssetView::HandleSettingChanged);
 
@@ -1924,6 +1927,11 @@ void SAssetView::OnObjectPropertyChanged(UObject* Asset, FPropertyChangedEvent&
 	}
 }
 
+void SAssetView::OnPackageDirtyStateUpdated(UPackage* Package)
+{
+	RequestListRefresh();
+}
+
 void SAssetView::OnDynamicFiltersChanged()
 {
 	ResetQuickJump();
diff --git a/Engine/Source/Editor/ContentBrowser/Private/SAssetView.h b/Engine/Source/Editor/ContentBrowser/Private/SAssetView.h
--- a/Engine/Source/Editor/ContentBrowser/Private/SAssetView.h
+++ b/Engine/Source/Editor/ContentBrowser/Private/SAssetView.h
@@ -317,6 +317,9 @@ private:
 	/** Handler for when an asset's property has changed */
 	void OnObjectPropertyChanged(UObject* Asset, FPropertyChangedEvent& PropertyChangedEvent);
 
+	/** Handler for when a package's dirty state has changed */
+	void OnPackageDirtyStateUpdated(UPackage* Package);
+
 	/** Handler for when any dynamic filters have been changed */
 	void OnDynamicFiltersChanged();
 
