#UE4 Fixed source control status disapearing when an asset item is filtered or occluded and then becomes visible again..

diff --git a/Engine/Source/Editor/ContentBrowser/Private/AssetViewWidgets.cpp b/Engine/Source/Editor/ContentBrowser/Private/AssetViewWidgets.cpp
--- a/Engine/Source/Editor/ContentBrowser/Private/AssetViewWidgets.cpp
+++ b/Engine/Source/Editor/ContentBrowser/Private/AssetViewWidgets.cpp
@@ -153,6 +153,13 @@ void SAssetViewItem::Construct( const FArguments& InArgs )
 	bSourceControlStateRequested = false;
 
 	ISourceControlModule::Get().GetProvider().RegisterSourceControlStateChanged(FSourceControlStateChanged::FDelegate::CreateSP(this, &SAssetViewItem::HandleSourceControlStateChanged));
+
+	// Source control state may have already been cached, make sure the control is in sync with 
+	// cached state as the delegate is not going to be invoked again until source control state 
+	// changes. This will be necessary any time the widget is destroyed and recreated after source 
+	// control state has been cached; for instance when the widget is killed via FWidgetGenerator::OnEndGenerationPass 
+	// or a view is refreshed due to user filtering/navigating):
+	HandleSourceControlStateChanged();
 }
 
 void SAssetViewItem::Tick( const FGeometry& AllottedGeometry, const double InCurrentTime, const float InDeltaTime )
