AssetRegistrySearchable properties now show up the propertyÆs DisplayName if one is available.
#TTP 336118 - AssetRegistrySearchable properties do not display DisplayName in tooltip

#branch UE4

#proj Editor.ContentBrowser

#add Added an extra bit of logic to the tooltip creation code that grabs the property field and uses the display name from that if one is available.

diff --git a/Engine/Source/Editor/ContentBrowser/Private/AssetViewWidgets.cpp b/Engine/Source/Editor/ContentBrowser/Private/AssetViewWidgets.cpp
--- a/Engine/Source/Editor/ContentBrowser/Private/AssetViewWidgets.cpp
+++ b/Engine/Source/Editor/ContentBrowser/Private/AssetViewWidgets.cpp
@@ -455,10 +455,21 @@ TSharedRef<SToolTip> SAssetViewItem::CreateToolTipWidget() const
 							const int32 SizeOfPrefix = ARRAY_COUNT(StringToRemove);
 							ValueString = ValueString.Mid(SizeOfPrefix - 1, ValueString.Len() - SizeOfPrefix).Replace(TEXT("Engine."), TEXT(""));
 						}
-
-						// We have no type information by this point, so no idea if it's a bool :(
-						const bool bIsBool = false;
-						AddToToolTipInfoBox( InfoBox, FText::FromString(FName::NameToDisplayString(TagIt.Key().ToString(), bIsBool)), FText::FromString(ValueString), bImportant );
+						
+						// Check for DisplayName metadata
+						FText DisplayName;
+						if (UProperty* Field = FindField<UProperty>(AssetClass, TagIt.Key()))
+						{
+							DisplayName = Field->GetDisplayNameText();
+						}
+						else
+						{
+							// We have no type information by this point, so no idea if it's a bool :(
+							const bool bIsBool = false;
+							DisplayName = FText::FromString(FName::NameToDisplayString(TagIt.Key().ToString(), bIsBool));
+						}
+					
+						AddToToolTipInfoBox(InfoBox, DisplayName, FText::FromString(ValueString), bImportant);
 					}
 				}
 			}
