Fixed crash when syncing 'root' directories from the content browser

Missing a trailing slash caused path matching to fail, so added slash and check for root mapping.

TTP# 337724 - CRASH: EDITOR: Syncing the game or editor folders within the editor will cause a crash to occur

diff --git a/Engine/Source/Editor/ContentBrowser/Private/PathContextMenu.cpp b/Engine/Source/Editor/ContentBrowser/Private/PathContextMenu.cpp
--- a/Engine/Source/Editor/ContentBrowser/Private/PathContextMenu.cpp
+++ b/Engine/Source/Editor/ContentBrowser/Private/PathContextMenu.cpp
@@ -659,41 +659,56 @@ void FPathContextMenu::ExecuteSCCSync() const
 {
 	ISourceControlProvider& SourceControlProvider = ISourceControlModule::Get().GetProvider();
 
-	// first attempt to unload all assets under this path
-	TArray<FString> PackageNames;
-	GetPackageNamesInSelectedPaths(PackageNames);
+	TArray<FString> RootPaths;
+	FPackageName::QueryRootContentPaths( RootPaths );
 
-	// Form a list of loaded packages to prompt for save
-	TArray<UPackage*> LoadedPackages;
-	for( const auto& PackageName : PackageNames )
+	// find valid paths on disk
+	TArray<FString> PathsOnDisk;
+	for(const auto& SelectedPath : SelectedPaths)
 	{
-		UPackage* Package = FindPackage(nullptr, *PackageName);
-		if ( Package != nullptr )
+		// note: we make sure to terminate our path here so root directories map correctly.
+		const FString CompletePath = SelectedPath / "";
+
+		const bool bMatchesRoot = RootPaths.ContainsByPredicate([&](const FString& RootPath){ return CompletePath.StartsWith(RootPath); });
+		if(bMatchesRoot)
 		{
-			LoadedPackages.Add(Package);
+			FString PathOnDisk = FPackageName::LongPackageNameToFilename(CompletePath);
+			PathsOnDisk.Add(PathOnDisk);
 		}
 	}
 
-	FText ErrorMessage;
-	PackageTools::UnloadPackages(LoadedPackages, ErrorMessage);
-
-	if(!ErrorMessage.IsEmpty())
-	{
-		FMessageDialog::Open( EAppMsgType::Ok, ErrorMessage );
-	}
-	else
+	if ( PathsOnDisk.Num() > 0 )
 	{
-		TArray<FString> PathsOnDisk;
-		for(const auto& SelectedPath : SelectedPaths)
+		// attempt to unload all assets under this path
+		TArray<FString> PackageNames;
+		GetPackageNamesInSelectedPaths(PackageNames);
+
+		// Form a list of loaded packages to prompt for save
+		TArray<UPackage*> LoadedPackages;
+		for( const auto& PackageName : PackageNames )
 		{
-			FString PathOnDisk = FPackageName::LongPackageNameToFilename(SelectedPath) / "";
-			PathsOnDisk.Add(PathOnDisk);
+			UPackage* Package = FindPackage(nullptr, *PackageName);
+			if ( Package != nullptr )
+			{
+				LoadedPackages.Add(Package);
+			}
 		}
 
-		if ( PathsOnDisk.Num() > 0 )
+		FText ErrorMessage;
+		PackageTools::UnloadPackages(LoadedPackages, ErrorMessage);
+
+		if(!ErrorMessage.IsEmpty())
+		{
+			FMessageDialog::Open( EAppMsgType::Ok, ErrorMessage );
+		}
+		else
 		{
 			SourceControlProvider.Execute(ISourceControlOperation::Create<FSync>(), PathsOnDisk);
-		}		
+		}
+	}
+	else
+	{
+		UE_LOG(LogContentBrowser, Warning, TEXT("Couldn't find any valid paths to sync."))
 	}
 }
 
