Added gameplay asset type category.
Modified FilterList and NewAssetContextMenu to only display categories that contain at least one item.

diff --git a/Engine/Source/Developer/AssetTools/Public/IAssetTypeActions.h b/Engine/Source/Developer/AssetTools/Public/IAssetTypeActions.h
--- a/Engine/Source/Developer/AssetTools/Public/IAssetTypeActions.h
+++ b/Engine/Source/Developer/AssetTools/Public/IAssetTypeActions.h
@@ -13,7 +13,8 @@ namespace EAssetTypeCategories
 		Sounds					= 1 << 3,
 		Physics					= 1 << 4,
 		UI						= 1 << 5,
-		Misc					= 1 << 6
+		Misc					= 1 << 6,
+		Gameplay				= 1 << 7,
 	};
 }
 
diff --git a/Engine/Source/Editor/ContentBrowser/Private/NewAssetContextMenu.cpp b/Engine/Source/Editor/ContentBrowser/Private/NewAssetContextMenu.cpp
--- a/Engine/Source/Editor/ContentBrowser/Private/NewAssetContextMenu.cpp
+++ b/Engine/Source/Editor/ContentBrowser/Private/NewAssetContextMenu.cpp
@@ -5,6 +5,40 @@
 
 #define LOCTEXT_NAMESPACE "ContentBrowser"
 
+struct FFactoryItem
+{
+	UFactory* Factory;
+	FText DisplayName;
+
+	FFactoryItem(UFactory* InFactory, const FText& InDisplayName)
+		: Factory(InFactory), DisplayName(InDisplayName)
+	{}
+};
+
+TArray<FFactoryItem> FindFactoriesInCategory(EAssetTypeCategories::Type AssetTypeCategory)
+{
+	TArray<FFactoryItem> FactoriesInThisCategory;
+	for (TObjectIterator<UClass> It; It; ++It)
+	{
+		UClass* Class = *It;
+		if (Class->IsChildOf(UFactory::StaticClass()) && !Class->HasAnyClassFlags(CLASS_Abstract))
+		{
+			UFactory* Factory = Class->GetDefaultObject<UFactory>();
+			if (Factory->ShouldShowInNewMenu() && Factory->ValidForCurrentGame() && ensure(!Factory->GetDisplayName().IsEmpty()))
+			{
+				uint32 FactoryCategories = Factory->GetMenuCategories();
+
+				if (FactoryCategories & AssetTypeCategory)
+				{
+					new(FactoriesInThisCategory)FFactoryItem(Factory, Factory->GetDisplayName());
+				}
+			}
+		}
+	}
+
+	return FactoriesInThisCategory;
+}
+
 class SFactoryMenuEntry : public SCompoundWidget
 {
 public:
@@ -129,78 +163,99 @@ void FNewAssetContextMenu::MakeContextMenu(FMenuBuilder& MenuBuilder, const FStr
 	}
 	MenuBuilder.EndSection(); //ContentBrowserNewBasicAsset
 
-	MenuBuilder.BeginSection("ContentBrowserNewAdvancedAsset", LOCTEXT("AdvancedAssetsMenuHeading", "Other Assets") );
+	MenuBuilder.BeginSection("ContentBrowserNewAdvancedAsset", LOCTEXT("AdvancedAssetsMenuHeading", "Other Assets"));
 	{
-		MenuBuilder.AddSubMenu(
-			LOCTEXT("AnimationAssetCategory", "Animation"),
-			FText::GetEmpty(),
-			FNewMenuDelegate::CreateStatic(&FNewAssetContextMenu::CreateNewAssetMenuCategory, EAssetTypeCategories::Animation, InPath, InOnNewAssetRequested)
-			);
+		{
+			TArray<FFactoryItem> AnimationFactories = FindFactoriesInCategory(EAssetTypeCategories::Animation);
+			if (AnimationFactories.Num() > 0)
+			{
+				MenuBuilder.AddSubMenu(
+					LOCTEXT("AnimationAssetCategory", "Animation"),
+					FText::GetEmpty(),
+					FNewMenuDelegate::CreateStatic(&FNewAssetContextMenu::CreateNewAssetMenuCategory, EAssetTypeCategories::Animation, InPath, InOnNewAssetRequested)
+					);
+			}
+		}
 
-		MenuBuilder.AddSubMenu(
-			LOCTEXT("MaterialAssetCategory", "Materials & Textures"),
-			FText::GetEmpty(),
-			FNewMenuDelegate::CreateStatic(&FNewAssetContextMenu::CreateNewAssetMenuCategory, EAssetTypeCategories::MaterialsAndTextures, InPath, InOnNewAssetRequested)
-			);
+		{
+			TArray<FFactoryItem> MaterialFactories = FindFactoriesInCategory(EAssetTypeCategories::MaterialsAndTextures);
+			if (MaterialFactories.Num() > 0)
+			{
+				MenuBuilder.AddSubMenu(
+					LOCTEXT("MaterialAssetCategory", "Materials & Textures"),
+					FText::GetEmpty(),
+					FNewMenuDelegate::CreateStatic(&FNewAssetContextMenu::CreateNewAssetMenuCategory, EAssetTypeCategories::MaterialsAndTextures, InPath, InOnNewAssetRequested)
+					);
+			}
+	}
 
-		MenuBuilder.AddSubMenu(
-			LOCTEXT("SoundAssetCategory", "Sounds"),
-			FText::GetEmpty(),
-			FNewMenuDelegate::CreateStatic(&FNewAssetContextMenu::CreateNewAssetMenuCategory, EAssetTypeCategories::Sounds, InPath, InOnNewAssetRequested)
-			);
+		{
+			TArray<FFactoryItem> SoundFactories = FindFactoriesInCategory(EAssetTypeCategories::Sounds);
+			if (SoundFactories.Num() > 0)
+			{
+				MenuBuilder.AddSubMenu(
+					LOCTEXT("SoundAssetCategory", "Sounds"),
+					FText::GetEmpty(),
+					FNewMenuDelegate::CreateStatic(&FNewAssetContextMenu::CreateNewAssetMenuCategory, EAssetTypeCategories::Sounds, InPath, InOnNewAssetRequested)
+					);
+			}
+		}
 
-		MenuBuilder.AddSubMenu(
-			LOCTEXT("PhysicsAssetCategory", "Physics"),
-			FText::GetEmpty(),
-			FNewMenuDelegate::CreateStatic(&FNewAssetContextMenu::CreateNewAssetMenuCategory, EAssetTypeCategories::Physics, InPath, InOnNewAssetRequested)
-			);
+		{
+			TArray<FFactoryItem> PhysicsFactories = FindFactoriesInCategory(EAssetTypeCategories::Physics);
+			if (PhysicsFactories.Num() > 0)
+			{
+				MenuBuilder.AddSubMenu(
+					LOCTEXT("PhysicsAssetCategory", "Physics"),
+					FText::GetEmpty(),
+					FNewMenuDelegate::CreateStatic(&FNewAssetContextMenu::CreateNewAssetMenuCategory, EAssetTypeCategories::Physics, InPath, InOnNewAssetRequested)
+					);
+			}
+		}
 
-		MenuBuilder.AddSubMenu(
-			LOCTEXT("UserInterfaceAssetCategory", "User Interface"),
-			FText::GetEmpty(),
-			FNewMenuDelegate::CreateStatic(&FNewAssetContextMenu::CreateNewAssetMenuCategory, EAssetTypeCategories::UI, InPath, InOnNewAssetRequested)
-			);
+		{
+			TArray<FFactoryItem> UIFactories = FindFactoriesInCategory(EAssetTypeCategories::UI);
+			if (UIFactories.Num() > 0)
+			{
+				MenuBuilder.AddSubMenu(
+					LOCTEXT("UserInterfaceAssetCategory", "User Interface"),
+					FText::GetEmpty(),
+					FNewMenuDelegate::CreateStatic(&FNewAssetContextMenu::CreateNewAssetMenuCategory, EAssetTypeCategories::UI, InPath, InOnNewAssetRequested)
+					);
+			}
+		}
 
-		MenuBuilder.AddSubMenu(
-			LOCTEXT("MiscellaneousAssetCategory", "Miscellaneous"),
-			FText::GetEmpty(),
-			FNewMenuDelegate::CreateStatic(&FNewAssetContextMenu::CreateNewAssetMenuCategory, EAssetTypeCategories::Misc, InPath, InOnNewAssetRequested)
-			);
+		{
+			TArray<FFactoryItem> MiscFactories = FindFactoriesInCategory(EAssetTypeCategories::Misc);
+			if (MiscFactories.Num() > 0)
+			{
+				MenuBuilder.AddSubMenu(
+					LOCTEXT("MiscellaneousAssetCategory", "Miscellaneous"),
+					FText::GetEmpty(),
+					FNewMenuDelegate::CreateStatic(&FNewAssetContextMenu::CreateNewAssetMenuCategory, EAssetTypeCategories::Misc, InPath, InOnNewAssetRequested)
+					);
+			}
+		}
+
+		{
+			TArray<FFactoryItem> GameplayFactories = FindFactoriesInCategory(EAssetTypeCategories::Gameplay);
+			if (GameplayFactories.Num() > 0)
+			{
+				MenuBuilder.AddSubMenu(
+					LOCTEXT("GameplayAssetCategory", "Gameplay"),
+					FText::GetEmpty(),
+					FNewMenuDelegate::CreateStatic(&FNewAssetContextMenu::CreateNewAssetMenuCategory, EAssetTypeCategories::Gameplay, InPath, InOnNewAssetRequested)
+					);
+			}
+		}
 	}
 	MenuBuilder.EndSection(); //ContentBrowserNewAdvancedAsset
 }
 
 void FNewAssetContextMenu::CreateNewAssetMenuCategory(FMenuBuilder& MenuBuilder, EAssetTypeCategories::Type AssetTypeCategory, FString InPath, FOnNewAssetRequested InOnNewAssetRequested)
 {
-	struct FFactoryItem
-	{
-		UFactory* Factory;
-		FText DisplayName;
-
-		FFactoryItem(UFactory* InFactory, const FText& InDisplayName)
-			: Factory(InFactory), DisplayName(InDisplayName)
-		{}
-	};
-
 	// Find UFactory classes that can create new objects in this category.
-	TArray<FFactoryItem> FactoriesInThisCategory;
-	for( TObjectIterator<UClass> It ; It ; ++It )
-	{
-		UClass* Class = *It;
-		if ( Class->IsChildOf(UFactory::StaticClass()) && !Class->HasAnyClassFlags(CLASS_Abstract) )
-		{
-			UFactory* Factory = Class->GetDefaultObject<UFactory>();
-			if ( Factory->ShouldShowInNewMenu() && Factory->ValidForCurrentGame() && ensure(!Factory->GetDisplayName().IsEmpty()) )
-			{
-				uint32 FactoryCategories = Factory->GetMenuCategories();
-
-				if ( FactoryCategories & AssetTypeCategory )
-				{
-					new(FactoriesInThisCategory) FFactoryItem(Factory, Factory->GetDisplayName());
-				}
-			}
-		}
-	}
+	TArray<FFactoryItem> FactoriesInThisCategory = FindFactoriesInCategory(AssetTypeCategory);
 
 	// Sort the list
 	struct FCompareFactoryDisplayNames
diff --git a/Engine/Source/Editor/ContentBrowser/Private/SFilterList.cpp b/Engine/Source/Editor/ContentBrowser/Private/SFilterList.cpp
--- a/Engine/Source/Editor/ContentBrowser/Private/SFilterList.cpp
+++ b/Engine/Source/Editor/ContentBrowser/Private/SFilterList.cpp
@@ -952,6 +952,7 @@ TSharedRef<SWidget> SFilterList::MakeAddFilterMenu(EAssetTypeCategories::Type Me
 	CategoryToMenuMap.Add(EAssetTypeCategories::Physics, FCategoryMenu( LOCTEXT("PhysicsFilter", "Physics"), LOCTEXT("PhysicsFilterTooltip", "Filter by physics assets."), "ContentBrowserFilterPhysicsAsset", LOCTEXT("PhysicsAssetsMenuHeading", "Physics Assets") ) );
 	CategoryToMenuMap.Add(EAssetTypeCategories::UI, FCategoryMenu(LOCTEXT("UIFilter", "User Interface"), LOCTEXT("UIFilterTooltip", "Filter by UI assets."), "ContentBrowserFilterUIAsset", LOCTEXT("UIAssetsMenuHeading", "User Interface Assets")));
 	CategoryToMenuMap.Add(EAssetTypeCategories::Misc, FCategoryMenu( LOCTEXT("MiscFilter", "Miscellaneous"), LOCTEXT("MiscFilterTooltip", "Filter by miscellaneous assets."), "ContentBrowserFilterMiscAsset", LOCTEXT("MiscAssetsMenuHeading", "Misc Assets") ) );
+	CategoryToMenuMap.Add(EAssetTypeCategories::Gameplay, FCategoryMenu(LOCTEXT("GameplayFilter", "Gameplay"), LOCTEXT("GameplayFilterTooltip", "Filter by gameplay assets."), "ContentBrowserFilterGameplayAsset", LOCTEXT("GameplayAssetsMenuHeading", "Gameplay Assets")));
 
 	// Load the asset tools module to get access to the browser type maps
 	FAssetToolsModule& AssetToolsModule = FModuleManager::LoadModuleChecked<FAssetToolsModule>(TEXT("AssetTools"));
@@ -968,7 +969,7 @@ TSharedRef<SWidget> SFilterList::MakeAddFilterMenu(EAssetTypeCategories::Type Me
 	};
 	AssetTypeActionsList.Sort( FCompareIAssetTypeActions() );
 
-	// For every asset type, move it into all the catgories it should appear in
+	// For every asset type, move it into all the categories it should appear in
 	for (int32 ClassIdx = 0; ClassIdx < AssetTypeActionsList.Num(); ++ClassIdx)
 	{
 		const TWeakPtr<IAssetTypeActions>& WeakTypeActions = AssetTypeActionsList[ClassIdx];
@@ -993,6 +994,14 @@ TSharedRef<SWidget> SFilterList::MakeAddFilterMenu(EAssetTypeCategories::Type Me
 		}
 	}
 
+	for (auto MenuIt = CategoryToMenuMap.CreateIterator(); MenuIt; ++MenuIt)
+	{
+		if (MenuIt.Value().Assets.Num() == 0)
+		{
+			CategoryToMenuMap.Remove(MenuIt.Key());
+		}
+	}
+
 	FMenuBuilder MenuBuilder(/*bInShouldCloseWindowAfterMenuSelection=*/true, NULL);
 
 	MenuBuilder.BeginSection( "ContentBrowserResetFilters" );
