UE4 Refactoring. Changed OVERRIDE and FINAL macros to keywords override and final.

diff --git a/Engine/Source/Editor/ContentBrowser/Private/AssetViewWidgets.h b/Engine/Source/Editor/ContentBrowser/Private/AssetViewWidgets.h
--- a/Engine/Source/Editor/ContentBrowser/Private/AssetViewWidgets.h
+++ b/Engine/Source/Editor/ContentBrowser/Private/AssetViewWidgets.h
@@ -16,24 +16,24 @@ namespace FAssetViewModeUtils
 class SAssetTileView : public STileView<TSharedPtr<FAssetViewItem>>
 {
 public:
-	virtual bool SupportsKeyboardFocus() const OVERRIDE { return true; }
-	virtual FReply OnKeyDown( const FGeometry& InGeometry, const FKeyboardEvent& InKeyboardEvent ) OVERRIDE;
+	virtual bool SupportsKeyboardFocus() const override { return true; }
+	virtual FReply OnKeyDown( const FGeometry& InGeometry, const FKeyboardEvent& InKeyboardEvent ) override;
 };
 
 /** The list view mode of the asset view */
 class SAssetListView : public SListView<TSharedPtr<FAssetViewItem>>
 {
 public:
-	virtual bool SupportsKeyboardFocus() const OVERRIDE { return true; }
-	virtual FReply OnKeyDown( const FGeometry& InGeometry, const FKeyboardEvent& InKeyboardEvent ) OVERRIDE;
+	virtual bool SupportsKeyboardFocus() const override { return true; }
+	virtual FReply OnKeyDown( const FGeometry& InGeometry, const FKeyboardEvent& InKeyboardEvent ) override;
 };
 
 /** The columns view mode of the asset view */
 class SAssetColumnView : public SListView<TSharedPtr<FAssetViewItem>>
 {
 public:
-	virtual bool SupportsKeyboardFocus() const OVERRIDE { return true; }
-	virtual FReply OnKeyDown( const FGeometry& InGeometry, const FKeyboardEvent& InKeyboardEvent ) OVERRIDE;
+	virtual bool SupportsKeyboardFocus() const override { return true; }
+	virtual FReply OnKeyDown( const FGeometry& InGeometry, const FKeyboardEvent& InKeyboardEvent ) override;
 };
 
 /** A base class for all asset view items */
@@ -95,12 +95,12 @@ public:
 	/** Performs common initialization logic for all asset view items */
 	void Construct( const FArguments& InArgs );
 
-	virtual void Tick( const FGeometry& AllottedGeometry, const double InCurrentTime, const float InDeltaTime ) OVERRIDE;
-	virtual TSharedPtr<IToolTip> GetToolTip() OVERRIDE;
-	virtual void OnDragEnter( const FGeometry& MyGeometry, const FDragDropEvent& DragDropEvent ) OVERRIDE;
-	virtual void OnDragLeave( const FDragDropEvent& DragDropEvent ) OVERRIDE;
-	virtual FReply OnDragOver( const FGeometry& MyGeometry, const FDragDropEvent& DragDropEvent ) OVERRIDE;
-	virtual FReply OnDrop( const FGeometry& MyGeometry, const FDragDropEvent& DragDropEvent ) OVERRIDE;
+	virtual void Tick( const FGeometry& AllottedGeometry, const double InCurrentTime, const float InDeltaTime ) override;
+	virtual TSharedPtr<IToolTip> GetToolTip() override;
+	virtual void OnDragEnter( const FGeometry& MyGeometry, const FDragDropEvent& DragDropEvent ) override;
+	virtual void OnDragLeave( const FDragDropEvent& DragDropEvent ) override;
+	virtual FReply OnDragOver( const FGeometry& MyGeometry, const FDragDropEvent& DragDropEvent ) override;
+	virtual FReply OnDrop( const FGeometry& MyGeometry, const FDragDropEvent& DragDropEvent ) override;
 
 	/** Returns the color this item should be tinted with */
 	virtual FSlateColor GetAssetColor() const;
@@ -334,7 +334,7 @@ public:
 	void Construct( const FArguments& InArgs );
 
 	/** Handles committing a name change */
-	virtual void OnAssetDataChanged() OVERRIDE;
+	virtual void OnAssetDataChanged() override;
 
 private:
 	/** Returns the size of the thumbnail widget */
@@ -435,7 +435,7 @@ public:
 	void Construct( const FArguments& InArgs );
 
 	/** Handles committing a name change */
-	virtual void OnAssetDataChanged() OVERRIDE;
+	virtual void OnAssetDataChanged() override;
 
 private:
 	/** Returns the size of the thumbnail box widget */
@@ -503,10 +503,10 @@ public:
 	virtual TSharedRef<SWidget> GenerateWidgetForColumn( const FName& ColumnName, FIsSelected InIsSelected );
 
 	/** Handles committing a name change */
-	virtual void OnAssetDataChanged() OVERRIDE;
+	virtual void OnAssetDataChanged() override;
 
 	/** Returns the color this item should be tinted with */
-	virtual FSlateColor GetAssetColor() const OVERRIDE;
+	virtual FSlateColor GetAssetColor() const override;
 
 private:
 	/** Gets the tool tip text for the name */
@@ -553,7 +553,7 @@ public:
 		Content = this->AssetColumnItem;
 	}
 
-	virtual TSharedRef<SWidget> GenerateWidgetForColumn( const FName& ColumnName ) OVERRIDE
+	virtual TSharedRef<SWidget> GenerateWidgetForColumn( const FName& ColumnName ) override
 	{
 		if ( this->AssetColumnItem.IsValid() )
 		{
@@ -566,7 +566,7 @@ public:
 		
 	}
 
-	virtual void Tick( const FGeometry& AllottedGeometry, const double InCurrentTime, const float InDeltaTime ) OVERRIDE
+	virtual void Tick( const FGeometry& AllottedGeometry, const double InCurrentTime, const float InDeltaTime ) override
 	{
 		SMultiColumnTableRow< TSharedPtr<FAssetViewItem> >::Tick(AllottedGeometry, InCurrentTime, InDeltaTime);
 
diff --git a/Engine/Source/Editor/ContentBrowser/Private/ContentBrowserSingleton.h b/Engine/Source/Editor/ContentBrowser/Private/ContentBrowserSingleton.h
--- a/Engine/Source/Editor/ContentBrowser/Private/ContentBrowserSingleton.h
+++ b/Engine/Source/Editor/ContentBrowser/Private/ContentBrowserSingleton.h
@@ -18,15 +18,15 @@ public:
 	virtual ~FContentBrowserSingleton();
 
 	// IContentBrowserSingleton interface
-	virtual TSharedRef<class SWidget> CreateAssetPicker(const FAssetPickerConfig& AssetPickerConfig) OVERRIDE;
-	virtual TSharedRef<class SWidget> CreatePathPicker(const FPathPickerConfig& PathPickerConfig) OVERRIDE;
-	virtual TSharedRef<class SWidget> CreateCollectionPicker(const FCollectionPickerConfig& CollectionPickerConfig) OVERRIDE;
-	virtual bool HasPrimaryContentBrowser() const OVERRIDE;
-	virtual void FocusPrimaryContentBrowser(bool bFocusSearch) OVERRIDE;
-	virtual void CreateNewAsset(const FString& DefaultAssetName, const FString& PackagePath, UClass* AssetClass, UFactory* Factory) OVERRIDE;
-	virtual void SyncBrowserToAssets(const TArray<class FAssetData>& AssetDataList, bool bAllowLockedBrowsers = false) OVERRIDE;
-	virtual void SyncBrowserToAssets(const TArray<UObject*>& AssetList, bool bAllowLockedBrowsers = false) OVERRIDE;
-	virtual void GetSelectedAssets(TArray<FAssetData>& SelectedAssets) OVERRIDE;
+	virtual TSharedRef<class SWidget> CreateAssetPicker(const FAssetPickerConfig& AssetPickerConfig) override;
+	virtual TSharedRef<class SWidget> CreatePathPicker(const FPathPickerConfig& PathPickerConfig) override;
+	virtual TSharedRef<class SWidget> CreateCollectionPicker(const FCollectionPickerConfig& CollectionPickerConfig) override;
+	virtual bool HasPrimaryContentBrowser() const override;
+	virtual void FocusPrimaryContentBrowser(bool bFocusSearch) override;
+	virtual void CreateNewAsset(const FString& DefaultAssetName, const FString& PackagePath, UClass* AssetClass, UFactory* Factory) override;
+	virtual void SyncBrowserToAssets(const TArray<class FAssetData>& AssetDataList, bool bAllowLockedBrowsers = false) override;
+	virtual void SyncBrowserToAssets(const TArray<UObject*>& AssetList, bool bAllowLockedBrowsers = false) override;
+	virtual void GetSelectedAssets(TArray<FAssetData>& SelectedAssets) override;
 
 		/** Gets the content browser singleton as a FContentBrowserSingleton */
 	static FContentBrowserSingleton& Get();
diff --git a/Engine/Source/Editor/ContentBrowser/Private/FrontendFilters.h b/Engine/Source/Editor/ContentBrowser/Private/FrontendFilters.h
--- a/Engine/Source/Editor/ContentBrowser/Private/FrontendFilters.h
+++ b/Engine/Source/Editor/ContentBrowser/Private/FrontendFilters.h
@@ -13,13 +13,13 @@ public:
 	FFrontendFilter_CheckedOut(TSharedPtr<FFrontendFilterCategory> InCategory);
 
 	// FFrontendFilter implementation
-	virtual FString GetName() const OVERRIDE { return TEXT("CheckedOut"); }
-	virtual FText GetDisplayName() const OVERRIDE { return LOCTEXT("FrontendFilter_CheckedOut", "Checked Out"); }
-	virtual FText GetToolTipText() const OVERRIDE { return LOCTEXT("FrontendFilter_CheckedOutTooltip", "Show only assets that you have checked out or pending for add."); }
-	virtual void ActiveStateChanged(bool bActive) OVERRIDE;
+	virtual FString GetName() const override { return TEXT("CheckedOut"); }
+	virtual FText GetDisplayName() const override { return LOCTEXT("FrontendFilter_CheckedOut", "Checked Out"); }
+	virtual FText GetToolTipText() const override { return LOCTEXT("FrontendFilter_CheckedOutTooltip", "Show only assets that you have checked out or pending for add."); }
+	virtual void ActiveStateChanged(bool bActive) override;
 
 	// IFilter implementation
-	virtual bool PassesFilter( AssetFilterType InItem ) const OVERRIDE;
+	virtual bool PassesFilter( AssetFilterType InItem ) const override;
 
 private:
 	
@@ -37,12 +37,12 @@ public:
 	FFrontendFilter_Modified(TSharedPtr<FFrontendFilterCategory> InCategory) : FFrontendFilter(InCategory) {}
 
 	// FFrontendFilter implementation
-	virtual FString GetName() const OVERRIDE { return TEXT("Modified"); }
-	virtual FText GetDisplayName() const OVERRIDE { return LOCTEXT("FrontendFilter_Modified", "Modified"); }
-	virtual FText GetToolTipText() const OVERRIDE { return LOCTEXT("FrontendFilter_ModifiedTooltip", "Show only assets that have been modified and not yet saved."); }
+	virtual FString GetName() const override { return TEXT("Modified"); }
+	virtual FText GetDisplayName() const override { return LOCTEXT("FrontendFilter_Modified", "Modified"); }
+	virtual FText GetToolTipText() const override { return LOCTEXT("FrontendFilter_ModifiedTooltip", "Show only assets that have been modified and not yet saved."); }
 
 	// IFilter implementation
-	virtual bool PassesFilter( AssetFilterType InItem ) const OVERRIDE;
+	virtual bool PassesFilter( AssetFilterType InItem ) const override;
 };
 
 /** A filter that displays blueprints that have replicated properties */
@@ -52,12 +52,12 @@ public:
 	FFrontendFilter_ReplicatedBlueprint(TSharedPtr<FFrontendFilterCategory> InCategory) : FFrontendFilter(InCategory) {}
 
 	// FFrontendFilter implementation
-	virtual FString GetName() const OVERRIDE { return TEXT("ReplicatedBlueprint"); }
-	virtual FText GetDisplayName() const OVERRIDE { return LOCTEXT("FFrontendFilter_ReplicatedBlueprint", "Replicated Blueprints"); }
-	virtual FText GetToolTipText() const OVERRIDE { return LOCTEXT("FFrontendFilter_ReplicatedBlueprintToolTip", "Show only blueprints with replicated properties."); }
+	virtual FString GetName() const override { return TEXT("ReplicatedBlueprint"); }
+	virtual FText GetDisplayName() const override { return LOCTEXT("FFrontendFilter_ReplicatedBlueprint", "Replicated Blueprints"); }
+	virtual FText GetToolTipText() const override { return LOCTEXT("FFrontendFilter_ReplicatedBlueprintToolTip", "Show only blueprints with replicated properties."); }
 
 	// IFilter implementation
-	virtual bool PassesFilter( AssetFilterType InItem ) const OVERRIDE;
+	virtual bool PassesFilter( AssetFilterType InItem ) const override;
 };
 
 /** An inverse filter that allows display of content in developer folders that are not the current user's */
@@ -68,14 +68,14 @@ public:
 	FFrontendFilter_ShowOtherDevelopers(TSharedPtr<FFrontendFilterCategory> InCategory);
 
 	// FFrontendFilter implementation
-	virtual FString GetName() const OVERRIDE { return TEXT("ShowOtherDevelopers"); }
-	virtual FText GetDisplayName() const OVERRIDE { return LOCTEXT("FrontendFilter_ShowOtherDevelopers", "Other Developers"); }
-	virtual FText GetToolTipText() const OVERRIDE { return LOCTEXT("FrontendFilter_ShowOtherDevelopersTooltip", "Allow display of assets in developer folders that aren't yours."); }
-	virtual bool IsInverseFilter() const OVERRIDE { return true; }
-	virtual void SetCurrentFilter(const FARFilter& InFilter) OVERRIDE;
+	virtual FString GetName() const override { return TEXT("ShowOtherDevelopers"); }
+	virtual FText GetDisplayName() const override { return LOCTEXT("FrontendFilter_ShowOtherDevelopers", "Other Developers"); }
+	virtual FText GetToolTipText() const override { return LOCTEXT("FrontendFilter_ShowOtherDevelopersTooltip", "Allow display of assets in developer folders that aren't yours."); }
+	virtual bool IsInverseFilter() const override { return true; }
+	virtual void SetCurrentFilter(const FARFilter& InFilter) override;
 
 	// IFilter implementation
-	virtual bool PassesFilter( AssetFilterType InItem ) const OVERRIDE;
+	virtual bool PassesFilter( AssetFilterType InItem ) const override;
 
 private:
 	FString BaseDeveloperPath;
@@ -91,14 +91,14 @@ public:
 	FFrontendFilter_ShowRedirectors(TSharedPtr<FFrontendFilterCategory> InCategory);
 
 	// FFrontendFilter implementation
-	virtual FString GetName() const OVERRIDE { return TEXT("ShowRedirectors"); }
-	virtual FText GetDisplayName() const OVERRIDE { return LOCTEXT("FrontendFilter_ShowRedirectors", "Show Redirectors"); }
-	virtual FText GetToolTipText() const OVERRIDE { return LOCTEXT("FrontendFilter_ShowRedirectorsToolTip", "Allow display of Redirectors."); }
-	virtual bool IsInverseFilter() const OVERRIDE { return true; }
-	virtual void SetCurrentFilter(const FARFilter& InFilter) OVERRIDE;
+	virtual FString GetName() const override { return TEXT("ShowRedirectors"); }
+	virtual FText GetDisplayName() const override { return LOCTEXT("FrontendFilter_ShowRedirectors", "Show Redirectors"); }
+	virtual FText GetToolTipText() const override { return LOCTEXT("FrontendFilter_ShowRedirectorsToolTip", "Allow display of Redirectors."); }
+	virtual bool IsInverseFilter() const override { return true; }
+	virtual void SetCurrentFilter(const FARFilter& InFilter) override;
 
 	// IFilter implementation
-	virtual bool PassesFilter( AssetFilterType InItem ) const OVERRIDE;
+	virtual bool PassesFilter( AssetFilterType InItem ) const override;
 
 private:
 	bool bAreRedirectorsInBaseFilter;
@@ -114,13 +114,13 @@ public:
 	~FFrontendFilter_InUseByLoadedLevels();
 
 	// FFrontendFilter implementation
-	virtual FString GetName() const OVERRIDE { return TEXT("InUseByLoadedLevels"); }
-	virtual FText GetDisplayName() const OVERRIDE { return LOCTEXT("FrontendFilter_InUseByLoadedLevels", "In Use By Level"); }
-	virtual FText GetToolTipText() const OVERRIDE { return LOCTEXT("FrontendFilter_InUseByLoadedLevelsToolTip", "Show only assets that are currently in use by any loaded level."); }
-	virtual void ActiveStateChanged(bool bActive) OVERRIDE;
+	virtual FString GetName() const override { return TEXT("InUseByLoadedLevels"); }
+	virtual FText GetDisplayName() const override { return LOCTEXT("FrontendFilter_InUseByLoadedLevels", "In Use By Level"); }
+	virtual FText GetToolTipText() const override { return LOCTEXT("FrontendFilter_InUseByLoadedLevelsToolTip", "Show only assets that are currently in use by any loaded level."); }
+	virtual void ActiveStateChanged(bool bActive) override;
 
 	// IFilter implementation
-	virtual bool PassesFilter( AssetFilterType InItem ) const OVERRIDE;
+	virtual bool PassesFilter( AssetFilterType InItem ) const override;
 
 	/** Handler for when maps change in the editor */
 	void OnEditorMapChange( uint32 MapChangeFlags );
diff --git a/Engine/Source/Editor/ContentBrowser/Private/SAssetPicker.h b/Engine/Source/Editor/ContentBrowser/Private/SAssetPicker.h
--- a/Engine/Source/Editor/ContentBrowser/Private/SAssetPicker.h
+++ b/Engine/Source/Editor/ContentBrowser/Private/SAssetPicker.h
@@ -19,8 +19,8 @@ public:
 	void Construct( const FArguments& InArgs );
 
 	// SWidget implementation
-	virtual void Tick( const FGeometry& AllottedGeometry, const double InCurrentTime, const float InDeltaTime ) OVERRIDE;
-	virtual FReply OnKeyDown( const FGeometry& MyGeometry, const FKeyboardEvent& InKeyboardEvent ) OVERRIDE;
+	virtual void Tick( const FGeometry& AllottedGeometry, const double InCurrentTime, const float InDeltaTime ) override;
+	virtual FReply OnKeyDown( const FGeometry& MyGeometry, const FKeyboardEvent& InKeyboardEvent ) override;
 	// End of SWidget implementation
 
 private:
diff --git a/Engine/Source/Editor/ContentBrowser/Private/SAssetView.h b/Engine/Source/Editor/ContentBrowser/Private/SAssetView.h
--- a/Engine/Source/Editor/ContentBrowser/Private/SAssetView.h
+++ b/Engine/Source/Editor/ContentBrowser/Private/SAssetView.h
@@ -216,14 +216,14 @@ public:
 	bool IsRenamingAsset() const;
 
 	// SWidget inherited
-	virtual void Tick( const FGeometry& AllottedGeometry, const double InCurrentTime, const float InDeltaTime ) OVERRIDE;
-	virtual void OnDragLeave( const FDragDropEvent& DragDropEvent ) OVERRIDE;
-	virtual FReply OnDragOver( const FGeometry& MyGeometry, const FDragDropEvent& DragDropEvent ) OVERRIDE;
-	virtual FReply OnDrop( const FGeometry& MyGeometry, const FDragDropEvent& DragDropEvent ) OVERRIDE;
-	virtual FReply OnKeyChar( const FGeometry& MyGeometry,const FCharacterEvent& InCharacterEvent ) OVERRIDE;
-	virtual FReply OnKeyDown( const FGeometry& MyGeometry, const FKeyboardEvent& InKeyboardEvent ) OVERRIDE;
-	virtual FReply OnMouseWheel( const FGeometry& MyGeometry, const FPointerEvent& MouseEvent ) OVERRIDE;
-	virtual void OnKeyboardFocusChanging( const FWeakWidgetPath& PreviousFocusPath, const FWidgetPath& NewWidgetPath ) OVERRIDE;
+	virtual void Tick( const FGeometry& AllottedGeometry, const double InCurrentTime, const float InDeltaTime ) override;
+	virtual void OnDragLeave( const FDragDropEvent& DragDropEvent ) override;
+	virtual FReply OnDragOver( const FGeometry& MyGeometry, const FDragDropEvent& DragDropEvent ) override;
+	virtual FReply OnDrop( const FGeometry& MyGeometry, const FDragDropEvent& DragDropEvent ) override;
+	virtual FReply OnKeyChar( const FGeometry& MyGeometry,const FCharacterEvent& InCharacterEvent ) override;
+	virtual FReply OnKeyDown( const FGeometry& MyGeometry, const FKeyboardEvent& InKeyboardEvent ) override;
+	virtual FReply OnMouseWheel( const FGeometry& MyGeometry, const FPointerEvent& MouseEvent ) override;
+	virtual void OnKeyboardFocusChanging( const FWeakWidgetPath& PreviousFocusPath, const FWidgetPath& NewWidgetPath ) override;
 
 	/** Opens the selected assets or folders, depending on the selection */
 	void OnOpenAssetsOrFolders();
diff --git a/Engine/Source/Editor/ContentBrowser/Private/SCollectionView.h b/Engine/Source/Editor/ContentBrowser/Private/SCollectionView.h
--- a/Engine/Source/Editor/ContentBrowser/Private/SCollectionView.h
+++ b/Engine/Source/Editor/ContentBrowser/Private/SCollectionView.h
@@ -55,9 +55,9 @@ public:
 	/** Loads any settings to config that should be persistent between editor sessions */
 	void LoadSettings(const FString& IniFilename, const FString& IniSection, const FString& SettingsString);
 
-	void Tick( const FGeometry& AllottedGeometry, const double InCurrentTime, const float InDeltaTime ) OVERRIDE;
+	void Tick( const FGeometry& AllottedGeometry, const double InCurrentTime, const float InDeltaTime ) override;
 
-	virtual FReply OnKeyDown( const FGeometry& MyGeometry, const FKeyboardEvent& InKeyboardEvent ) OVERRIDE;
+	virtual FReply OnKeyDown( const FGeometry& MyGeometry, const FKeyboardEvent& InKeyboardEvent ) override;
 
 private:
 
diff --git a/Engine/Source/Editor/ContentBrowser/Private/SContentBrowser.h b/Engine/Source/Editor/ContentBrowser/Private/SContentBrowser.h
--- a/Engine/Source/Editor/ContentBrowser/Private/SContentBrowser.h
+++ b/Engine/Source/Editor/ContentBrowser/Private/SContentBrowser.h
@@ -61,10 +61,10 @@ public:
 	/** Gives keyboard focus to the asset search box */
 	void SetKeyboardFocusOnSearch() const;
 
-	virtual FReply OnKeyDown( const FGeometry& MyGeometry, const FKeyboardEvent& InKeyboardEvent ) OVERRIDE;
-	virtual FReply OnMouseButtonDown( const FGeometry& MyGeometry, const FPointerEvent& MouseEvent ) OVERRIDE;
-	virtual FReply OnPreviewMouseButtonDown( const FGeometry& MyGeometry, const FPointerEvent& MouseEvent ) OVERRIDE;
-	virtual FReply OnMouseButtonDoubleClick( const FGeometry& InMyGeometry, const FPointerEvent& InMouseEvent ) OVERRIDE;
+	virtual FReply OnKeyDown( const FGeometry& MyGeometry, const FKeyboardEvent& InKeyboardEvent ) override;
+	virtual FReply OnMouseButtonDown( const FGeometry& MyGeometry, const FPointerEvent& MouseEvent ) override;
+	virtual FReply OnPreviewMouseButtonDown( const FGeometry& MyGeometry, const FPointerEvent& MouseEvent ) override;
+	virtual FReply OnMouseButtonDoubleClick( const FGeometry& InMyGeometry, const FPointerEvent& InMouseEvent ) override;
 
 private:
 	/** Called to retrieve the text that should be highlighted on assets */
diff --git a/Engine/Source/Editor/ContentBrowser/Private/SFilterList.cpp b/Engine/Source/Editor/ContentBrowser/Private/SFilterList.cpp
--- a/Engine/Source/Editor/ContentBrowser/Private/SFilterList.cpp
+++ b/Engine/Source/Editor/ContentBrowser/Private/SFilterList.cpp
@@ -20,7 +20,7 @@ public:
 		OnFilterMiddleButtonClicked = NewFilterMiddleButtonClicked;
 	}
 
-	virtual FReply OnMouseButtonDoubleClick( const FGeometry& InMyGeometry, const FPointerEvent& InMouseEvent ) OVERRIDE
+	virtual FReply OnMouseButtonDoubleClick( const FGeometry& InMyGeometry, const FPointerEvent& InMouseEvent ) override
 	{
 		if ( InMouseEvent.GetEffectingButton() == EKeys::LeftMouseButton && OnFilterDoubleClicked.IsBound() )
 		{
@@ -32,7 +32,7 @@ public:
 		}
 	}
 
-	virtual FReply OnMouseButtonUp( const FGeometry& InMyGeometry, const FPointerEvent& InMouseEvent ) OVERRIDE
+	virtual FReply OnMouseButtonUp( const FGeometry& InMyGeometry, const FPointerEvent& InMouseEvent ) override
 	{
 		if( InMouseEvent.GetEffectingButton() == EKeys::MiddleMouseButton && OnFilterMiddleButtonClicked.IsBound() )
 		{
diff --git a/Engine/Source/Editor/ContentBrowser/Private/SThumbnailEditModeTools.h b/Engine/Source/Editor/ContentBrowser/Private/SThumbnailEditModeTools.h
--- a/Engine/Source/Editor/ContentBrowser/Private/SThumbnailEditModeTools.h
+++ b/Engine/Source/Editor/ContentBrowser/Private/SThumbnailEditModeTools.h
@@ -17,10 +17,10 @@ public:
 	void Construct( const FArguments& InArgs, const TSharedPtr<FAssetThumbnail>& InAssetThumbnail );
 
 	// SWidget Interface
-	FReply OnMouseButtonDown( const FGeometry& MyGeometry, const FPointerEvent& MouseEvent ) OVERRIDE;
-	FReply OnMouseButtonUp( const FGeometry& MyGeometry, const FPointerEvent& MouseEvent ) OVERRIDE;
-	FReply OnMouseMove( const FGeometry& MyGeometry, const FPointerEvent& MouseEvent ) OVERRIDE;
-	FCursorReply OnCursorQuery( const FGeometry& MyGeometry, const FPointerEvent& CursorEvent ) const OVERRIDE;
+	FReply OnMouseButtonDown( const FGeometry& MyGeometry, const FPointerEvent& MouseEvent ) override;
+	FReply OnMouseButtonUp( const FGeometry& MyGeometry, const FPointerEvent& MouseEvent ) override;
+	FReply OnMouseMove( const FGeometry& MyGeometry, const FPointerEvent& MouseEvent ) override;
+	FCursorReply OnCursorQuery( const FGeometry& MyGeometry, const FPointerEvent& CursorEvent ) const override;
 
 protected:
 	/** Gets the visibility for the primitives toolbar */
diff --git a/Engine/Source/Editor/ContentBrowser/Private/SourcesViewWidgets.h b/Engine/Source/Editor/ContentBrowser/Private/SourcesViewWidgets.h
--- a/Engine/Source/Editor/ContentBrowser/Private/SourcesViewWidgets.h
+++ b/Engine/Source/Editor/ContentBrowser/Private/SourcesViewWidgets.h
@@ -50,11 +50,11 @@ public:
 
 	~SAssetTreeItem();
 
-	virtual void OnDragEnter( const FGeometry& MyGeometry, const FDragDropEvent& DragDropEvent ) OVERRIDE;
-	virtual void OnDragLeave( const FDragDropEvent& DragDropEvent ) OVERRIDE;
-	virtual FReply OnDragOver( const FGeometry& MyGeometry, const FDragDropEvent& DragDropEvent ) OVERRIDE;
-	virtual FReply OnDrop( const FGeometry& MyGeometry, const FDragDropEvent& DragDropEvent ) OVERRIDE;
-	virtual void Tick( const FGeometry& AllottedGeometry, const double InCurrentTime, const float InDeltaTime ) OVERRIDE;
+	virtual void OnDragEnter( const FGeometry& MyGeometry, const FDragDropEvent& DragDropEvent ) override;
+	virtual void OnDragLeave( const FDragDropEvent& DragDropEvent ) override;
+	virtual FReply OnDragOver( const FGeometry& MyGeometry, const FDragDropEvent& DragDropEvent ) override;
+	virtual FReply OnDrop( const FGeometry& MyGeometry, const FDragDropEvent& DragDropEvent ) override;
+	virtual void Tick( const FGeometry& AllottedGeometry, const double InCurrentTime, const float InDeltaTime ) override;
 
 private:
 	/** Handles verifying name changes */
@@ -176,11 +176,11 @@ public:
 
 	~SCollectionListItem();
 
-	virtual void Tick( const FGeometry& AllottedGeometry, const double InCurrentTime, const float InDeltaTime ) OVERRIDE;
-	virtual void OnDragEnter( const FGeometry& MyGeometry, const FDragDropEvent& DragDropEvent ) OVERRIDE;
-	virtual void OnDragLeave( const FDragDropEvent& DragDropEvent ) OVERRIDE;
-	virtual FReply OnDragOver( const FGeometry& MyGeometry, const FDragDropEvent& DragDropEvent ) OVERRIDE;
-	virtual FReply OnDrop( const FGeometry& MyGeometry, const FDragDropEvent& DragDropEvent ) OVERRIDE;
+	virtual void Tick( const FGeometry& AllottedGeometry, const double InCurrentTime, const float InDeltaTime ) override;
+	virtual void OnDragEnter( const FGeometry& MyGeometry, const FDragDropEvent& DragDropEvent ) override;
+	virtual void OnDragLeave( const FDragDropEvent& DragDropEvent ) override;
+	virtual FReply OnDragOver( const FGeometry& MyGeometry, const FDragDropEvent& DragDropEvent ) override;
+	virtual FReply OnDrop( const FGeometry& MyGeometry, const FDragDropEvent& DragDropEvent ) override;
 
 private:
 	/** Handles beginning a name change */
diff --git a/Engine/Source/Editor/ContentBrowser/Public/FrontendFilterBase.h b/Engine/Source/Editor/ContentBrowser/Public/FrontendFilterBase.h
--- a/Engine/Source/Editor/ContentBrowser/Public/FrontendFilterBase.h
+++ b/Engine/Source/Editor/ContentBrowser/Public/FrontendFilterBase.h
@@ -48,7 +48,7 @@ public:
 
 	// IFilter implementation
 	DECLARE_DERIVED_EVENT( FFrontendFilter, IFilter<AssetFilterType>::FChangedEvent, FChangedEvent );
-	virtual FChangedEvent& OnChanged() OVERRIDE { return ChangedEvent; }
+	virtual FChangedEvent& OnChanged() override { return ChangedEvent; }
 
 	TSharedPtr<FFrontendFilterCategory> GetCategory() { return FilterCategory; }
 
diff --git a/Engine/Source/Editor/CurveAssetEditor/Private/CurveAssetEditor.h b/Engine/Source/Editor/CurveAssetEditor/Private/CurveAssetEditor.h
--- a/Engine/Source/Editor/CurveAssetEditor/Private/CurveAssetEditor.h
+++ b/Engine/Source/Editor/CurveAssetEditor/Private/CurveAssetEditor.h
@@ -16,8 +16,8 @@ public:
 
 	virtual ~FCurveAssetEditor() {}
 
-	virtual void RegisterTabSpawners(const TSharedRef<class FTabManager>& TabManager) OVERRIDE;
-	virtual void UnregisterTabSpawners(const TSharedRef<class FTabManager>& TabManager) OVERRIDE;
+	virtual void RegisterTabSpawners(const TSharedRef<class FTabManager>& TabManager) override;
+	virtual void UnregisterTabSpawners(const TSharedRef<class FTabManager>& TabManager) override;
 
 	/**
 	 * Edits the specified table
@@ -29,10 +29,10 @@ public:
 	void InitCurveAssetEditor( const EToolkitMode::Type Mode, const TSharedPtr< class IToolkitHost >& InitToolkitHost, UCurveBase* CurveToEdit );
 
 	/** IToolkit interface */
-	virtual FName GetToolkitFName() const OVERRIDE;
-	virtual FText GetBaseToolkitName() const OVERRIDE;
-	virtual FString GetWorldCentricTabPrefix() const OVERRIDE;
-	virtual FLinearColor GetWorldCentricTabColorScale() const OVERRIDE;
+	virtual FName GetToolkitFName() const override;
+	virtual FText GetBaseToolkitName() const override;
+	virtual FString GetWorldCentricTabPrefix() const override;
+	virtual FLinearColor GetWorldCentricTabColorScale() const override;
 
 private:
 
diff --git a/Engine/Source/Editor/CurveTableEditor/Private/CurveTableEditor.h b/Engine/Source/Editor/CurveTableEditor/Private/CurveTableEditor.h
--- a/Engine/Source/Editor/CurveTableEditor/Private/CurveTableEditor.h
+++ b/Engine/Source/Editor/CurveTableEditor/Private/CurveTableEditor.h
@@ -12,8 +12,8 @@ class FCurveTableEditor :
 {
 
 public:
-	virtual void RegisterTabSpawners(const TSharedRef<class FTabManager>& TabManager) OVERRIDE;
-	virtual void UnregisterTabSpawners(const TSharedRef<class FTabManager>& TabManager) OVERRIDE;
+	virtual void RegisterTabSpawners(const TSharedRef<class FTabManager>& TabManager) override;
+	virtual void UnregisterTabSpawners(const TSharedRef<class FTabManager>& TabManager) override;
 
 	/**
 	 * Edits the specified table
@@ -28,10 +28,10 @@ public:
 	virtual ~FCurveTableEditor();
 
 	/** IToolkit interface */
-	virtual FName GetToolkitFName() const OVERRIDE;
-	virtual FText GetBaseToolkitName() const OVERRIDE;
-	virtual FString GetWorldCentricTabPrefix() const OVERRIDE;
-	virtual FLinearColor GetWorldCentricTabColorScale() const OVERRIDE;
+	virtual FName GetToolkitFName() const override;
+	virtual FText GetBaseToolkitName() const override;
+	virtual FString GetWorldCentricTabPrefix() const override;
+	virtual FLinearColor GetWorldCentricTabColorScale() const override;
 
 	/**	Spawns the tab with the curve table inside */
 	TSharedRef<SDockTab> SpawnTab_CurveTable( const FSpawnTabArgs& Args );
diff --git a/Engine/Source/Editor/DataTableEditor/Private/DataTableEditor.h b/Engine/Source/Editor/DataTableEditor/Private/DataTableEditor.h
--- a/Engine/Source/Editor/DataTableEditor/Private/DataTableEditor.h
+++ b/Engine/Source/Editor/DataTableEditor/Private/DataTableEditor.h
@@ -13,8 +13,8 @@ class FDataTableEditor :
 
 public:
 
-	virtual void RegisterTabSpawners(const TSharedRef<class FTabManager>& TabManager) OVERRIDE;
-	virtual void UnregisterTabSpawners(const TSharedRef<class FTabManager>& TabManager) OVERRIDE;
+	virtual void RegisterTabSpawners(const TSharedRef<class FTabManager>& TabManager) override;
+	virtual void UnregisterTabSpawners(const TSharedRef<class FTabManager>& TabManager) override;
 
 	/**
 	 * Edits the specified table
@@ -32,13 +32,13 @@ public:
 	TSharedPtr<SUniformGridPanel> CreateGridPanel();
 
 	/** IToolkit interface */
-	virtual FName GetToolkitFName() const OVERRIDE;
-	virtual FText GetBaseToolkitName() const OVERRIDE;
-	virtual FString GetWorldCentricTabPrefix() const OVERRIDE;
-	virtual FLinearColor GetWorldCentricTabColorScale() const OVERRIDE;
+	virtual FName GetToolkitFName() const override;
+	virtual FText GetBaseToolkitName() const override;
+	virtual FString GetWorldCentricTabPrefix() const override;
+	virtual FLinearColor GetWorldCentricTabColorScale() const override;
 
 	/* Called when DataTable was reloaded */
-	virtual void OnDataTableReloaded() OVERRIDE;
+	virtual void OnDataTableReloaded() override;
 
 
 private:
diff --git a/Engine/Source/Editor/DestructibleMeshEditor/Classes/DestructibleChunkParamsProxy.h b/Engine/Source/Editor/DestructibleMeshEditor/Classes/DestructibleChunkParamsProxy.h
--- a/Engine/Source/Editor/DestructibleMeshEditor/Classes/DestructibleChunkParamsProxy.h
+++ b/Engine/Source/Editor/DestructibleMeshEditor/Classes/DestructibleChunkParamsProxy.h
@@ -22,7 +22,7 @@ class UDestructibleChunkParamsProxy : public UObject
 
 #if WITH_EDITOR
 	// UObject interface
-	virtual void PostEditChangeProperty(struct FPropertyChangedEvent& PropertyChangedEvent) OVERRIDE;
+	virtual void PostEditChangeProperty(struct FPropertyChangedEvent& PropertyChangedEvent) override;
 	// End of UObject interface
 #endif
 };
\ No newline at end of file
diff --git a/Engine/Source/Editor/DestructibleMeshEditor/Private/DestructibleMeshEditor.h b/Engine/Source/Editor/DestructibleMeshEditor/Private/DestructibleMeshEditor.h
--- a/Engine/Source/Editor/DestructibleMeshEditor/Private/DestructibleMeshEditor.h
+++ b/Engine/Source/Editor/DestructibleMeshEditor/Private/DestructibleMeshEditor.h
@@ -31,7 +31,7 @@ public:
 	/**
 	 * Initialize commands
 	 */
-	virtual void RegisterCommands() OVERRIDE;
+	virtual void RegisterCommands() override;
 
 public:
 };
@@ -54,25 +54,25 @@ public:
 	virtual ~FDestructibleMeshEditor();
 
 	/** IToolkit interface */
-	virtual void RegisterTabSpawners(const TSharedRef<class FTabManager>& TabManager) OVERRIDE;
-	virtual void UnregisterTabSpawners(const TSharedRef<class FTabManager>& TabManager) OVERRIDE;
-	virtual FName GetToolkitFName() const OVERRIDE;
-	virtual FText GetBaseToolkitName() const OVERRIDE;
-	virtual FString GetWorldCentricTabPrefix() const OVERRIDE;
-	virtual FLinearColor GetWorldCentricTabColorScale() const OVERRIDE;
-
-	virtual FString GetDocumentationLink() const OVERRIDE
+	virtual void RegisterTabSpawners(const TSharedRef<class FTabManager>& TabManager) override;
+	virtual void UnregisterTabSpawners(const TSharedRef<class FTabManager>& TabManager) override;
+	virtual FName GetToolkitFName() const override;
+	virtual FText GetBaseToolkitName() const override;
+	virtual FString GetWorldCentricTabPrefix() const override;
+	virtual FLinearColor GetWorldCentricTabColorScale() const override;
+
+	virtual FString GetDocumentationLink() const override
 	{
 		// @Todo: revert to Engine/Physics/Destructibles once docs exist. 
 		return FString(TEXT("Engine/Physics"));
 	}
 
 	/** IDestructibleMeshEditor interface */
-	virtual UDestructibleMesh* GetDestructibleMesh() OVERRIDE;
-	virtual int32 GetCurrentPreviewDepth() const OVERRIDE;
-	virtual void SetCurrentPreviewDepth(uint32 InPreviewDepthDepth) OVERRIDE;
-	virtual void RefreshTool() OVERRIDE;
-	virtual void RefreshViewport() OVERRIDE;
+	virtual UDestructibleMesh* GetDestructibleMesh() override;
+	virtual int32 GetCurrentPreviewDepth() const override;
+	virtual void SetCurrentPreviewDepth(uint32 InPreviewDepthDepth) override;
+	virtual void RefreshTool() override;
+	virtual void RefreshViewport() override;
 
 	void SetSelectedChunks(const TArray<UObject*>& SelectedChunks);
 
diff --git a/Engine/Source/Editor/DestructibleMeshEditor/Private/SDestructibleMeshEditorViewport.cpp b/Engine/Source/Editor/DestructibleMeshEditor/Private/SDestructibleMeshEditorViewport.cpp
--- a/Engine/Source/Editor/DestructibleMeshEditor/Private/SDestructibleMeshEditorViewport.cpp
+++ b/Engine/Source/Editor/DestructibleMeshEditor/Private/SDestructibleMeshEditorViewport.cpp
@@ -27,11 +27,11 @@ public:
 	FDestructibleMeshEditorViewportClient(TWeakPtr<IDestructibleMeshEditor> InDestructibleMeshEditor, FPreviewScene& InPreviewScene);
 
 	// FGCObject interface
-	virtual void AddReferencedObjects( FReferenceCollector& Collector ) OVERRIDE;
+	virtual void AddReferencedObjects( FReferenceCollector& Collector ) override;
 	// End of FGCObject interface
 
 	// FEditorViewportClient interface
-	virtual void Draw(const FSceneView* View,FPrimitiveDrawInterface* PDI) OVERRIDE;
+	virtual void Draw(const FSceneView* View,FPrimitiveDrawInterface* PDI) override;
 	FLinearColor GetBackgroundColor() const { return FLinearColor::Black; }
 	virtual void ProcessClick(class FSceneView& View, class HHitProxy* HitProxy, FKey Key, EInputEvent Event, uint32 HitX, uint32 HitY);
 
diff --git a/Engine/Source/Editor/DestructibleMeshEditor/Private/SDestructibleMeshEditorViewport.h b/Engine/Source/Editor/DestructibleMeshEditor/Private/SDestructibleMeshEditorViewport.h
--- a/Engine/Source/Editor/DestructibleMeshEditor/Private/SDestructibleMeshEditorViewport.h
+++ b/Engine/Source/Editor/DestructibleMeshEditor/Private/SDestructibleMeshEditorViewport.h
@@ -19,14 +19,14 @@ public:
 	~SDestructibleMeshEditorViewport();
 	
 	// FGCObject interface
-	virtual void AddReferencedObjects( FReferenceCollector& Collector ) OVERRIDE;
+	virtual void AddReferencedObjects( FReferenceCollector& Collector ) override;
 	// End of FGCObject interface
 
 	// FNotifyHook interface
-	virtual void NotifyPostChange( const FPropertyChangedEvent& PropertyChangedEvent, class FEditPropertyChain* PropertyThatChanged) OVERRIDE;
+	virtual void NotifyPostChange( const FPropertyChangedEvent& PropertyChangedEvent, class FEditPropertyChain* PropertyThatChanged) override;
 	// End of FNotifyHook interface
 
-	virtual void Tick(const FGeometry& AllottedGeometry, const double InCurrentTime, const float InDeltaTime) OVERRIDE;
+	virtual void Tick(const FGeometry& AllottedGeometry, const double InCurrentTime, const float InDeltaTime) override;
 	void RefreshViewport();
 	
 	/** Component for the preview Destructible mesh. */
diff --git a/Engine/Source/Editor/DetailCustomizations/Private/ActorDetails.cpp b/Engine/Source/Editor/DetailCustomizations/Private/ActorDetails.cpp
--- a/Engine/Source/Editor/DetailCustomizations/Private/ActorDetails.cpp
+++ b/Engine/Source/Editor/DetailCustomizations/Private/ActorDetails.cpp
@@ -159,7 +159,7 @@ public:
 	/** Allowed ChildOf relationship. */
 	TSet< const UClass* > AllowedChildOfRelationship;
 
-	virtual bool IsClassAllowed(const FClassViewerInitializationOptions& InInitOptions, const UClass* InClass, TSharedRef< FClassViewerFilterFuncs > InFilterFuncs ) OVERRIDE
+	virtual bool IsClassAllowed(const FClassViewerInitializationOptions& InInitOptions, const UClass* InClass, TSharedRef< FClassViewerFilterFuncs > InFilterFuncs ) override
 	{
 		EFilterReturn::Type eState = InFilterFuncs->IfInClassesSet(AllowedClasses, InClass);
 		if(eState == EFilterReturn::NoItems)
@@ -186,7 +186,7 @@ public:
 		return false;
 	}
 
-	virtual bool IsUnloadedClassAllowed(const FClassViewerInitializationOptions& InInitOptions, const TSharedRef< const IUnloadedBlueprintData > InUnloadedClassData, TSharedRef< FClassViewerFilterFuncs > InFilterFuncs) OVERRIDE
+	virtual bool IsUnloadedClassAllowed(const FClassViewerInitializationOptions& InInitOptions, const TSharedRef< const IUnloadedBlueprintData > InUnloadedClassData, TSharedRef< FClassViewerFilterFuncs > InFilterFuncs) override
 	{
 		EFilterReturn::Type eState = InFilterFuncs->IfInClassesSet(AllowedClasses, InUnloadedClassData);
 		if(eState == EFilterReturn::NoItems)
diff --git a/Engine/Source/Editor/DetailCustomizations/Private/AmbientSoundDetails.h b/Engine/Source/Editor/DetailCustomizations/Private/AmbientSoundDetails.h
--- a/Engine/Source/Editor/DetailCustomizations/Private/AmbientSoundDetails.h
+++ b/Engine/Source/Editor/DetailCustomizations/Private/AmbientSoundDetails.h
@@ -18,7 +18,7 @@ private:
 	};
 
 	/** IDetailCustomization interface */
-	virtual void CustomizeDetails( IDetailLayoutBuilder& DetailBuilder ) OVERRIDE;
+	virtual void CustomizeDetails( IDetailLayoutBuilder& DetailBuilder ) override;
 
 	bool IsEditSoundCueEnabled() const;
 
diff --git a/Engine/Source/Editor/DetailCustomizations/Private/AnimMontageSegmentDetails.h b/Engine/Source/Editor/DetailCustomizations/Private/AnimMontageSegmentDetails.h
--- a/Engine/Source/Editor/DetailCustomizations/Private/AnimMontageSegmentDetails.h
+++ b/Engine/Source/Editor/DetailCustomizations/Private/AnimMontageSegmentDetails.h
@@ -14,7 +14,7 @@ public:
 	static TSharedRef<IDetailCustomization> MakeInstance();
 
 	/** IDetailCustomization interface */
-	virtual void CustomizeDetails( IDetailLayoutBuilder& DetailBuilder ) OVERRIDE;
+	virtual void CustomizeDetails( IDetailLayoutBuilder& DetailBuilder ) override;
 };
 
 
@@ -43,7 +43,7 @@ public:
 
 	void Construct(const FArguments& InArgs);
 
-	virtual void Tick(const FGeometry& AllottedGeometry, const double InCurrentTime, const float InDeltaTime) OVERRIDE;
+	virtual void Tick(const FGeometry& AllottedGeometry, const double InCurrentTime, const float InDeltaTime) override;
 	void RefreshViewport();
 
 private:
@@ -155,8 +155,8 @@ public:
 	FAnimationSegmentViewportClient(FPreviewScene& InPreviewScene, FFXSystemInterface * FXSystem);
 
 	// FlEditorViewportClient interface
-	virtual FSceneInterface* GetScene() const OVERRIDE;
-	virtual FLinearColor GetBackgroundColor() const OVERRIDE { return FLinearColor::Black; }
+	virtual FSceneInterface* GetScene() const override;
+	virtual FLinearColor GetBackgroundColor() const override { return FLinearColor::Black; }
 
 	// End of FEditorViewportClient
 
diff --git a/Engine/Source/Editor/DetailCustomizations/Private/AnimSequenceDetails.h b/Engine/Source/Editor/DetailCustomizations/Private/AnimSequenceDetails.h
--- a/Engine/Source/Editor/DetailCustomizations/Private/AnimSequenceDetails.h
+++ b/Engine/Source/Editor/DetailCustomizations/Private/AnimSequenceDetails.h
@@ -11,7 +11,7 @@ public:
 	static TSharedRef<IDetailCustomization> MakeInstance();
 
 	/** IDetailCustomization interface */
-	virtual void CustomizeDetails( IDetailLayoutBuilder& DetailBuilder ) OVERRIDE;
+	virtual void CustomizeDetails( IDetailLayoutBuilder& DetailBuilder ) override;
 
 	/** Create an override for the supplide Property */
 	void CreateOverridenProperty(IDetailLayoutBuilder& DetailBuilder, IDetailCategoryBuilder& AdditiveSettingsCategory, TSharedPtr<IPropertyHandle> PropertyHandle, TAttribute<EVisibility> VisibilityAttribute);
@@ -69,7 +69,7 @@ public:
 
 	void Construct(const FArguments& InArgs);
 
-	virtual void Tick(const FGeometry& AllottedGeometry, const double InCurrentTime, const float InDeltaTime) OVERRIDE;
+	virtual void Tick(const FGeometry& AllottedGeometry, const double InCurrentTime, const float InDeltaTime) override;
 	void RefreshViewport();
 
 private:
diff --git a/Engine/Source/Editor/DetailCustomizations/Private/AnimStateNodeDetails.h b/Engine/Source/Editor/DetailCustomizations/Private/AnimStateNodeDetails.h
--- a/Engine/Source/Editor/DetailCustomizations/Private/AnimStateNodeDetails.h
+++ b/Engine/Source/Editor/DetailCustomizations/Private/AnimStateNodeDetails.h
@@ -13,6 +13,6 @@ public:
 	static TSharedRef<IDetailCustomization> MakeInstance();
 
 	/** IDetailCustomization interface */
-	virtual void CustomizeDetails( IDetailLayoutBuilder& DetailBuilder ) OVERRIDE;
+	virtual void CustomizeDetails( IDetailLayoutBuilder& DetailBuilder ) override;
 };
 
diff --git a/Engine/Source/Editor/DetailCustomizations/Private/AttenuationSettingsCustomizations.h b/Engine/Source/Editor/DetailCustomizations/Private/AttenuationSettingsCustomizations.h
--- a/Engine/Source/Editor/DetailCustomizations/Private/AttenuationSettingsCustomizations.h
+++ b/Engine/Source/Editor/DetailCustomizations/Private/AttenuationSettingsCustomizations.h
@@ -8,8 +8,8 @@ public:
 	static TSharedRef<IPropertyTypeCustomization> MakeInstance();
 
 	/** IPropertyTypeCustomization instance */
-	virtual void CustomizeHeader( TSharedRef<class IPropertyHandle> StructPropertyHandle, class FDetailWidgetRow& HeaderRow, IPropertyTypeCustomizationUtils& StructCustomizationUtils ) OVERRIDE;
-	virtual void CustomizeChildren( TSharedRef<IPropertyHandle> StructPropertyHandle, class IDetailChildrenBuilder& ChildBuilder, IPropertyTypeCustomizationUtils& StructCustomizationUtils ) OVERRIDE;
+	virtual void CustomizeHeader( TSharedRef<class IPropertyHandle> StructPropertyHandle, class FDetailWidgetRow& HeaderRow, IPropertyTypeCustomizationUtils& StructCustomizationUtils ) override;
+	virtual void CustomizeChildren( TSharedRef<IPropertyHandle> StructPropertyHandle, class IDetailChildrenBuilder& ChildBuilder, IPropertyTypeCustomizationUtils& StructCustomizationUtils ) override;
 
 protected:
 
diff --git a/Engine/Source/Editor/DetailCustomizations/Private/BlackboardEntryDetails.h b/Engine/Source/Editor/DetailCustomizations/Private/BlackboardEntryDetails.h
--- a/Engine/Source/Editor/DetailCustomizations/Private/BlackboardEntryDetails.h
+++ b/Engine/Source/Editor/DetailCustomizations/Private/BlackboardEntryDetails.h
@@ -9,8 +9,8 @@ public:
 	static TSharedRef<IPropertyTypeCustomization> MakeInstance();
 
 	/** IPropertyTypeCustomization interface */
-	virtual void CustomizeHeader( TSharedRef<class IPropertyHandle> StructPropertyHandle, class FDetailWidgetRow& HeaderRow, IPropertyTypeCustomizationUtils& StructCustomizationUtils ) OVERRIDE;
-	virtual void CustomizeChildren( TSharedRef<class IPropertyHandle> StructPropertyHandle, class IDetailChildrenBuilder& StructBuilder, IPropertyTypeCustomizationUtils& StructCustomizationUtils ) OVERRIDE;
+	virtual void CustomizeHeader( TSharedRef<class IPropertyHandle> StructPropertyHandle, class FDetailWidgetRow& HeaderRow, IPropertyTypeCustomizationUtils& StructCustomizationUtils ) override;
+	virtual void CustomizeChildren( TSharedRef<class IPropertyHandle> StructPropertyHandle, class IDetailChildrenBuilder& StructBuilder, IPropertyTypeCustomizationUtils& StructCustomizationUtils ) override;
 
 private:
 
diff --git a/Engine/Source/Editor/DetailCustomizations/Private/BodyInstanceCustomization.h b/Engine/Source/Editor/DetailCustomizations/Private/BodyInstanceCustomization.h
--- a/Engine/Source/Editor/DetailCustomizations/Private/BodyInstanceCustomization.h
+++ b/Engine/Source/Editor/DetailCustomizations/Private/BodyInstanceCustomization.h
@@ -23,8 +23,8 @@ public:
 	}
 
 	/** IPropertyTypeCustomization interface */
-	virtual void CustomizeHeader( TSharedRef<IPropertyHandle> StructPropertyHandle, class FDetailWidgetRow& HeaderRow, IPropertyTypeCustomizationUtils& StructCustomizationUtils ) OVERRIDE {};
-	virtual void CustomizeChildren( TSharedRef<class IPropertyHandle> StructPropertyHandle, class IDetailChildrenBuilder& StructBuilder, IPropertyTypeCustomizationUtils& StructCustomizationUtils ) OVERRIDE;
+	virtual void CustomizeHeader( TSharedRef<IPropertyHandle> StructPropertyHandle, class FDetailWidgetRow& HeaderRow, IPropertyTypeCustomizationUtils& StructCustomizationUtils ) override {};
+	virtual void CustomizeChildren( TSharedRef<class IPropertyHandle> StructPropertyHandle, class IDetailChildrenBuilder& StructBuilder, IPropertyTypeCustomizationUtils& StructCustomizationUtils ) override;
 
 private:
 	// Profile combo related
diff --git a/Engine/Source/Editor/DetailCustomizations/Private/CollisionProfileDetails.h b/Engine/Source/Editor/DetailCustomizations/Private/CollisionProfileDetails.h
--- a/Engine/Source/Editor/DetailCustomizations/Private/CollisionProfileDetails.h
+++ b/Engine/Source/Editor/DetailCustomizations/Private/CollisionProfileDetails.h
@@ -39,7 +39,7 @@ public:
 	* @param InArgs - The Slate argument list.
 	*/
 	void Construct(const FArguments& InArgs, const TSharedRef<STableViewBase>& InOwnerTableView);
-	virtual TSharedRef<SWidget> GenerateWidgetForColumn(const FName& ColumnName) OVERRIDE;
+	virtual TSharedRef<SWidget> GenerateWidgetForColumn(const FName& ColumnName) override;
 
 private:
 
@@ -86,7 +86,7 @@ public:
 	* @param InArgs - The Slate argument list.
 	*/
 	void Construct(const FArguments& InArgs, const TSharedRef<STableViewBase>& InOwnerTableView);
-	virtual TSharedRef<SWidget> GenerateWidgetForColumn(const FName& ColumnName) OVERRIDE;
+	virtual TSharedRef<SWidget> GenerateWidgetForColumn(const FName& ColumnName) override;
 
 private:
 
@@ -105,7 +105,7 @@ public:
 	static TSharedRef<IDetailCustomization> MakeInstance();
 
 	/** IDetailCustomization interface */
-	virtual void CustomizeDetails( IDetailLayoutBuilder& DetailBuilder ) OVERRIDE;
+	virtual void CustomizeDetails( IDetailLayoutBuilder& DetailBuilder ) override;
 
 private:
 
diff --git a/Engine/Source/Editor/DetailCustomizations/Private/ComponentTransformDetails.h b/Engine/Source/Editor/DetailCustomizations/Private/ComponentTransformDetails.h
--- a/Engine/Source/Editor/DetailCustomizations/Private/ComponentTransformDetails.h
+++ b/Engine/Source/Editor/DetailCustomizations/Private/ComponentTransformDetails.h
@@ -26,13 +26,13 @@ public:
 	 */
 	void CacheTransform();
 
-	virtual void GenerateHeaderRowContent( FDetailWidgetRow& NodeRow ) OVERRIDE {}
-	virtual void GenerateChildContent( IDetailChildrenBuilder& ChildrenBuilder ) OVERRIDE;
-	virtual bool RequiresTick() const OVERRIDE { return true; }
-	virtual FName GetName() const OVERRIDE { return NAME_None; }
-	virtual bool InitiallyCollapsed() const OVERRIDE { return false; }
-	virtual void Tick( float DeltaTime ) OVERRIDE;
-	virtual void SetOnRebuildChildren( FSimpleDelegate OnRebuildChildren ) OVERRIDE{}
+	virtual void GenerateHeaderRowContent( FDetailWidgetRow& NodeRow ) override {}
+	virtual void GenerateChildContent( IDetailChildrenBuilder& ChildrenBuilder ) override;
+	virtual bool RequiresTick() const override { return true; }
+	virtual FName GetName() const override { return NAME_None; }
+	virtual bool InitiallyCollapsed() const override { return false; }
+	virtual void Tick( float DeltaTime ) override;
+	virtual void SetOnRebuildChildren( FSimpleDelegate OnRebuildChildren ) override{}
 
 private:
 	/** @return Whether the transform details panel should be enabled (editable) or not (read-only / greyed out) */
diff --git a/Engine/Source/Editor/DetailCustomizations/Private/CurveStructCustomization.h b/Engine/Source/Editor/DetailCustomizations/Private/CurveStructCustomization.h
--- a/Engine/Source/Editor/DetailCustomizations/Private/CurveStructCustomization.h
+++ b/Engine/Source/Editor/DetailCustomizations/Private/CurveStructCustomization.h
@@ -18,15 +18,15 @@ public:
 	virtual ~FCurveStructCustomization();
 
 	/** IPropertyTypeCustomization interface */
-	virtual void CustomizeHeader( TSharedRef<class IPropertyHandle> InStructPropertyHandle, class FDetailWidgetRow& HeaderRow, IPropertyTypeCustomizationUtils& StructCustomizationUtils ) OVERRIDE;
-	virtual void CustomizeChildren( TSharedRef<class IPropertyHandle> InStructPropertyHandle, class IDetailChildrenBuilder& StructBuilder, IPropertyTypeCustomizationUtils& StructCustomizationUtils ) OVERRIDE;
+	virtual void CustomizeHeader( TSharedRef<class IPropertyHandle> InStructPropertyHandle, class FDetailWidgetRow& HeaderRow, IPropertyTypeCustomizationUtils& StructCustomizationUtils ) override;
+	virtual void CustomizeChildren( TSharedRef<class IPropertyHandle> InStructPropertyHandle, class IDetailChildrenBuilder& StructBuilder, IPropertyTypeCustomizationUtils& StructCustomizationUtils ) override;
 
 	/** FCurveOwnerInterface interface */
-	virtual TArray<FRichCurveEditInfoConst> GetCurves() const OVERRIDE;
-	virtual TArray<FRichCurveEditInfo> GetCurves() OVERRIDE;
-	virtual UObject* GetOwner() OVERRIDE;
-	virtual void ModifyOwner() OVERRIDE;
-	virtual void MakeTransactional() OVERRIDE;
+	virtual TArray<FRichCurveEditInfoConst> GetCurves() const override;
+	virtual TArray<FRichCurveEditInfo> GetCurves() override;
+	virtual UObject* GetOwner() override;
+	virtual void ModifyOwner() override;
+	virtual void MakeTransactional() override;
 
 private:
 	/**
diff --git a/Engine/Source/Editor/DetailCustomizations/Private/CurveTableCustomization.h b/Engine/Source/Editor/DetailCustomizations/Private/CurveTableCustomization.h
--- a/Engine/Source/Editor/DetailCustomizations/Private/CurveTableCustomization.h
+++ b/Engine/Source/Editor/DetailCustomizations/Private/CurveTableCustomization.h
@@ -14,7 +14,7 @@ public:
 	}
 
 	/** IPropertyTypeCustomization interface */
-	virtual void CustomizeHeader( TSharedRef<class IPropertyHandle> InStructPropertyHandle, class FDetailWidgetRow& HeaderRow, IPropertyTypeCustomizationUtils& StructCustomizationUtils ) OVERRIDE
+	virtual void CustomizeHeader( TSharedRef<class IPropertyHandle> InStructPropertyHandle, class FDetailWidgetRow& HeaderRow, IPropertyTypeCustomizationUtils& StructCustomizationUtils ) override
 	{
 		this->StructPropertyHandle = InStructPropertyHandle;
 
@@ -25,7 +25,7 @@ public:
 			];
 	}
 
-	virtual void CustomizeChildren( TSharedRef<class IPropertyHandle> InStructPropertyHandle, class IDetailChildrenBuilder& StructBuilder, IPropertyTypeCustomizationUtils& StructCustomizationUtils ) OVERRIDE
+	virtual void CustomizeChildren( TSharedRef<class IPropertyHandle> InStructPropertyHandle, class IDetailChildrenBuilder& StructBuilder, IPropertyTypeCustomizationUtils& StructCustomizationUtils ) override
 	{
 		/** Get all the existing property handles */
 		CurveTablePropertyHandle = InStructPropertyHandle->GetChildHandle( "CurveTable" );
diff --git a/Engine/Source/Editor/DetailCustomizations/Private/DataTableCategoryCustomization.h b/Engine/Source/Editor/DetailCustomizations/Private/DataTableCategoryCustomization.h
--- a/Engine/Source/Editor/DetailCustomizations/Private/DataTableCategoryCustomization.h
+++ b/Engine/Source/Editor/DetailCustomizations/Private/DataTableCategoryCustomization.h
@@ -16,7 +16,7 @@ public:
 	}
 
 	/** IPropertyTypeCustomization interface */
-	virtual void CustomizeHeader(TSharedRef<class IPropertyHandle> InStructPropertyHandle, class FDetailWidgetRow& HeaderRow, IPropertyTypeCustomizationUtils& StructCustomizationUtils) OVERRIDE
+	virtual void CustomizeHeader(TSharedRef<class IPropertyHandle> InStructPropertyHandle, class FDetailWidgetRow& HeaderRow, IPropertyTypeCustomizationUtils& StructCustomizationUtils) override
 	{
 		this->StructPropertyHandle = InStructPropertyHandle;
 
@@ -27,7 +27,7 @@ public:
 			];
 	}
 
-	virtual void CustomizeChildren(TSharedRef<class IPropertyHandle> InStructPropertyHandle, class IDetailChildrenBuilder& StructBuilder, IPropertyTypeCustomizationUtils& StructCustomizationUtils) OVERRIDE
+	virtual void CustomizeChildren(TSharedRef<class IPropertyHandle> InStructPropertyHandle, class IDetailChildrenBuilder& StructBuilder, IPropertyTypeCustomizationUtils& StructCustomizationUtils) override
 	{
 		/** Get all the existing property handles */
 		DataTablePropertyHandle		= InStructPropertyHandle->GetChildHandle("DataTable");
diff --git a/Engine/Source/Editor/DetailCustomizations/Private/DataTableCustomization.h b/Engine/Source/Editor/DetailCustomizations/Private/DataTableCustomization.h
--- a/Engine/Source/Editor/DetailCustomizations/Private/DataTableCustomization.h
+++ b/Engine/Source/Editor/DetailCustomizations/Private/DataTableCustomization.h
@@ -16,7 +16,7 @@ public:
 	}
 
 	/** IPropertyTypeCustomization interface */
-	virtual void CustomizeHeader( TSharedRef<class IPropertyHandle> InStructPropertyHandle, class FDetailWidgetRow& HeaderRow, IPropertyTypeCustomizationUtils& StructCustomizationUtils ) OVERRIDE
+	virtual void CustomizeHeader( TSharedRef<class IPropertyHandle> InStructPropertyHandle, class FDetailWidgetRow& HeaderRow, IPropertyTypeCustomizationUtils& StructCustomizationUtils ) override
 	{
 		this->StructPropertyHandle = InStructPropertyHandle;
 
@@ -27,7 +27,7 @@ public:
 			];
 	}
 
-	virtual void CustomizeChildren( TSharedRef<class IPropertyHandle> InStructPropertyHandle, class IDetailChildrenBuilder& StructBuilder, IPropertyTypeCustomizationUtils& StructCustomizationUtils ) OVERRIDE
+	virtual void CustomizeChildren( TSharedRef<class IPropertyHandle> InStructPropertyHandle, class IDetailChildrenBuilder& StructBuilder, IPropertyTypeCustomizationUtils& StructCustomizationUtils ) override
 	{
 		/** Get all the existing property handles */
 		DataTablePropertyHandle = InStructPropertyHandle->GetChildHandle( "DataTable" );
diff --git a/Engine/Source/Editor/DetailCustomizations/Private/DialogueStructsCustomizations.cpp b/Engine/Source/Editor/DetailCustomizations/Private/DialogueStructsCustomizations.cpp
--- a/Engine/Source/Editor/DetailCustomizations/Private/DialogueStructsCustomizations.cpp
+++ b/Engine/Source/Editor/DetailCustomizations/Private/DialogueStructsCustomizations.cpp
@@ -58,7 +58,7 @@ public:
 	SLATE_END_ARGS()
 
 	void Construct( const FArguments& InArgs, const TSharedRef<IPropertyHandle>& InPropertyHandle, const TSharedRef<FAssetThumbnailPool>& InAssetThumbnailPool );
-	virtual void Tick( const FGeometry& AllottedGeometry, const double InCurrentTime, const float InDeltaTime ) OVERRIDE;
+	virtual void Tick( const FGeometry& AllottedGeometry, const double InCurrentTime, const float InDeltaTime ) override;
 
 private:
 	void OnSelectionChanged(TSharedPtr<UDialogueVoice*> Speaker, ESelectInfo::Type SelectInfo);
@@ -278,7 +278,7 @@ public:
 	SLATE_END_ARGS()
 
 	void Construct( const FArguments& InArgs, const TSharedRef<IPropertyHandle>& InPropertyHandle, const TSharedRef<FAssetThumbnailPool>& InAssetThumbnailPool );
-	virtual void Tick( const FGeometry& AllottedGeometry, const double InCurrentTime, const float InDeltaTime ) OVERRIDE;
+	virtual void Tick( const FGeometry& AllottedGeometry, const double InCurrentTime, const float InDeltaTime ) override;
 
 private:
 	void OnSelectionChanged(TSharedPtr<FTargetSet> TargetSet, ESelectInfo::Type SelectInfo);
@@ -597,7 +597,7 @@ public:
 	SLATE_END_ARGS()
 
 	void Construct( const FArguments& InArgs, const TSharedRef<IPropertyHandle>& InPropertyHandle, const TSharedRef<FAssetThumbnailPool>& InAssetThumbnailPool );
-	virtual void Tick(  const FGeometry& AllottedGeometry, const double InCurrentTime, const float InDeltaTime ) OVERRIDE;
+	virtual void Tick(  const FGeometry& AllottedGeometry, const double InCurrentTime, const float InDeltaTime ) override;
 
 private:
 	bool IsDialogueWaveValid() const;
diff --git a/Engine/Source/Editor/DetailCustomizations/Private/DialogueStructsCustomizations.h b/Engine/Source/Editor/DetailCustomizations/Private/DialogueStructsCustomizations.h
--- a/Engine/Source/Editor/DetailCustomizations/Private/DialogueStructsCustomizations.h
+++ b/Engine/Source/Editor/DetailCustomizations/Private/DialogueStructsCustomizations.h
@@ -9,8 +9,8 @@ public:
 	static TSharedRef<IPropertyTypeCustomization> MakeInstance();
 
 	/** IPropertyTypeCustomization interface begin */
-	virtual void CustomizeHeader( TSharedRef<IPropertyHandle> StructPropertyHandle, class FDetailWidgetRow& HeaderRow, IPropertyTypeCustomizationUtils& StructCustomizationUtils ) OVERRIDE;
-	virtual void CustomizeChildren( TSharedRef<IPropertyHandle> StructPropertyHandle, class IDetailChildrenBuilder& ChildBuilder, IPropertyTypeCustomizationUtils& StructCustomizationUtils ) OVERRIDE;
+	virtual void CustomizeHeader( TSharedRef<IPropertyHandle> StructPropertyHandle, class FDetailWidgetRow& HeaderRow, IPropertyTypeCustomizationUtils& StructCustomizationUtils ) override;
+	virtual void CustomizeChildren( TSharedRef<IPropertyHandle> StructPropertyHandle, class IDetailChildrenBuilder& ChildBuilder, IPropertyTypeCustomizationUtils& StructCustomizationUtils ) override;
 	/** IPropertyTypeCustomization interface end */
 };
 
@@ -21,7 +21,7 @@ public:
 	static TSharedRef<IPropertyTypeCustomization> MakeInstance();
 
 	/** IPropertyTypeCustomization interface begin */
-	virtual void CustomizeHeader( TSharedRef<IPropertyHandle> StructPropertyHandle, class FDetailWidgetRow& HeaderRow, IPropertyTypeCustomizationUtils& StructCustomizationUtils ) OVERRIDE;
-	virtual void CustomizeChildren( TSharedRef<IPropertyHandle> StructPropertyHandle, class IDetailChildrenBuilder& ChildBuilder, IPropertyTypeCustomizationUtils& StructCustomizationUtils ) OVERRIDE;
+	virtual void CustomizeHeader( TSharedRef<IPropertyHandle> StructPropertyHandle, class FDetailWidgetRow& HeaderRow, IPropertyTypeCustomizationUtils& StructCustomizationUtils ) override;
+	virtual void CustomizeChildren( TSharedRef<IPropertyHandle> StructPropertyHandle, class IDetailChildrenBuilder& ChildBuilder, IPropertyTypeCustomizationUtils& StructCustomizationUtils ) override;
 	/** IPropertyTypeCustomization interface end */
 };
\ No newline at end of file
diff --git a/Engine/Source/Editor/DetailCustomizations/Private/DialogueWaveDetails.h b/Engine/Source/Editor/DetailCustomizations/Private/DialogueWaveDetails.h
--- a/Engine/Source/Editor/DetailCustomizations/Private/DialogueWaveDetails.h
+++ b/Engine/Source/Editor/DetailCustomizations/Private/DialogueWaveDetails.h
@@ -8,13 +8,13 @@ public:
 	FDialogueContextMappingNodeBuilder( IDetailLayoutBuilder* InDetailLayoutBuilder, const TSharedPtr<IPropertyHandle>& InPropertyHandle );
 
 	/** IDetailCustomNodeBuilder interface */
-	virtual void SetOnRebuildChildren( FSimpleDelegate InOnRebuildChildren  ) OVERRIDE { OnRebuildChildren = InOnRebuildChildren; } 
-	virtual bool RequiresTick() const OVERRIDE { return false; }
-	virtual void Tick( float DeltaTime ) OVERRIDE {};
-	virtual void GenerateHeaderRowContent( FDetailWidgetRow& NodeRow ) OVERRIDE;
-	virtual void GenerateChildContent( IDetailChildrenBuilder& ChildrenBuilder ) OVERRIDE;
-	virtual bool InitiallyCollapsed() const OVERRIDE { return true; };
-	virtual FName GetName() const OVERRIDE { return NAME_None; }
+	virtual void SetOnRebuildChildren( FSimpleDelegate InOnRebuildChildren  ) override { OnRebuildChildren = InOnRebuildChildren; } 
+	virtual bool RequiresTick() const override { return false; }
+	virtual void Tick( float DeltaTime ) override {};
+	virtual void GenerateHeaderRowContent( FDetailWidgetRow& NodeRow ) override;
+	virtual void GenerateChildContent( IDetailChildrenBuilder& ChildrenBuilder ) override;
+	virtual bool InitiallyCollapsed() const override { return true; };
+	virtual FName GetName() const override { return NAME_None; }
 
 private:
 	void RemoveContextButton_OnClick();
@@ -40,7 +40,7 @@ public:
 	static TSharedRef<IDetailCustomization> MakeInstance();
 
 	/** ILayoutDetails interface */
-	virtual void CustomizeDetails( class IDetailLayoutBuilder& DetailBuilder ) OVERRIDE;
+	virtual void CustomizeDetails( class IDetailLayoutBuilder& DetailBuilder ) override;
 
 private:
 	FReply AddDialogueContextMapping_OnClicked();
diff --git a/Engine/Source/Editor/DetailCustomizations/Private/DialogueWaveWidgets.cpp b/Engine/Source/Editor/DetailCustomizations/Private/DialogueWaveWidgets.cpp
--- a/Engine/Source/Editor/DetailCustomizations/Private/DialogueWaveWidgets.cpp
+++ b/Engine/Source/Editor/DetailCustomizations/Private/DialogueWaveWidgets.cpp
@@ -550,11 +550,11 @@ public:
 
 	}
 protected:
-	virtual FReply OnMouseButtonDown( const FGeometry& MyGeometry, const FPointerEvent& MouseEvent ) OVERRIDE;
+	virtual FReply OnMouseButtonDown( const FGeometry& MyGeometry, const FPointerEvent& MouseEvent ) override;
 
-	virtual FReply OnMouseButtonDoubleClick( const FGeometry& InMyGeometry, const FPointerEvent& InMouseEvent ) OVERRIDE;
+	virtual FReply OnMouseButtonDoubleClick( const FGeometry& InMyGeometry, const FPointerEvent& InMouseEvent ) override;
 
-	virtual FReply OnMouseButtonUp( const FGeometry& MyGeometry, const FPointerEvent& MouseEvent ) OVERRIDE;
+	virtual FReply OnMouseButtonUp( const FGeometry& MyGeometry, const FPointerEvent& MouseEvent ) override;
 
 private:
 	void DoRemove();
diff --git a/Engine/Source/Editor/DetailCustomizations/Private/DialogueWaveWidgets.h b/Engine/Source/Editor/DetailCustomizations/Private/DialogueWaveWidgets.h
--- a/Engine/Source/Editor/DetailCustomizations/Private/DialogueWaveWidgets.h
+++ b/Engine/Source/Editor/DetailCustomizations/Private/DialogueWaveWidgets.h
@@ -18,7 +18,7 @@ public:
 
 	void Construct( const FArguments& InArgs, const TSharedRef<IPropertyHandle>& InPropertyHandle, const TSharedRef<FAssetThumbnailPool>& InAssetThumbnailPool );
 
-	virtual void Tick( const FGeometry& AllottedGeometry, const double InCurrentTime, const float InDeltaTime ) OVERRIDE;
+	virtual void Tick( const FGeometry& AllottedGeometry, const double InCurrentTime, const float InDeltaTime ) override;
 
 	FText GetDialogueVoiceDescription() const;
 
@@ -89,7 +89,7 @@ public:
 	SLATE_END_ARGS()
 
 	void Construct( const FArguments& InArgs, const TSharedRef<IPropertyHandle>& InPropertyHandle, const TSharedRef<FAssetThumbnailPool>& InAssetThumbnailPool );
-	virtual void Tick( const FGeometry& AllottedGeometry, const double InCurrentTime, const float InDeltaTime ) OVERRIDE;
+	virtual void Tick( const FGeometry& AllottedGeometry, const double InCurrentTime, const float InDeltaTime ) override;
 
 	FText GetDialogueVoiceDescription() const;
 
@@ -114,7 +114,7 @@ public:
 	SLATE_END_ARGS()
 
 	void Construct( const FArguments& InArgs, const TSharedRef<IPropertyHandle>& InPropertyHandle, const TSharedRef<FAssetThumbnailPool>& InAssetThumbnailPool );
-	virtual void Tick( const FGeometry& AllottedGeometry, const double InCurrentTime, const float InDeltaTime ) OVERRIDE;
+	virtual void Tick( const FGeometry& AllottedGeometry, const double InCurrentTime, const float InDeltaTime ) override;
 
 private:
 	bool IsSpeakerValid() const;
diff --git a/Engine/Source/Editor/DetailCustomizations/Private/DirectoryPathStructCustomization.h b/Engine/Source/Editor/DetailCustomizations/Private/DirectoryPathStructCustomization.h
--- a/Engine/Source/Editor/DetailCustomizations/Private/DirectoryPathStructCustomization.h
+++ b/Engine/Source/Editor/DetailCustomizations/Private/DirectoryPathStructCustomization.h
@@ -8,9 +8,9 @@ public:
 	static TSharedRef<IPropertyTypeCustomization> MakeInstance();
 
 	/** IPropertyTypeCustomization interface */
-	virtual void CustomizeHeader( TSharedRef<class IPropertyHandle> StructPropertyHandle, class FDetailWidgetRow& HeaderRow, IPropertyTypeCustomizationUtils& StructCustomizationUtils ) OVERRIDE;
+	virtual void CustomizeHeader( TSharedRef<class IPropertyHandle> StructPropertyHandle, class FDetailWidgetRow& HeaderRow, IPropertyTypeCustomizationUtils& StructCustomizationUtils ) override;
 
-	virtual void CustomizeChildren( TSharedRef<class IPropertyHandle> StructPropertyHandle, class IDetailChildrenBuilder& StructBuilder, IPropertyTypeCustomizationUtils& StructCustomizationUtils ) OVERRIDE;              
+	virtual void CustomizeChildren( TSharedRef<class IPropertyHandle> StructPropertyHandle, class IDetailChildrenBuilder& StructBuilder, IPropertyTypeCustomizationUtils& StructCustomizationUtils ) override;              
 
 private:
 
diff --git a/Engine/Source/Editor/DetailCustomizations/Private/FilePathStructCustomization.h b/Engine/Source/Editor/DetailCustomizations/Private/FilePathStructCustomization.h
--- a/Engine/Source/Editor/DetailCustomizations/Private/FilePathStructCustomization.h
+++ b/Engine/Source/Editor/DetailCustomizations/Private/FilePathStructCustomization.h
@@ -11,9 +11,9 @@ public:
 	static TSharedRef<IPropertyTypeCustomization> MakeInstance();
 
 	/** IPropertyTypeCustomization interface */
-	virtual void CustomizeHeader( TSharedRef<class IPropertyHandle> StructPropertyHandle, class FDetailWidgetRow& HeaderRow, IPropertyTypeCustomizationUtils& StructCustomizationUtils ) OVERRIDE;
+	virtual void CustomizeHeader( TSharedRef<class IPropertyHandle> StructPropertyHandle, class FDetailWidgetRow& HeaderRow, IPropertyTypeCustomizationUtils& StructCustomizationUtils ) override;
 
-	virtual void CustomizeChildren( TSharedRef<class IPropertyHandle> StructPropertyHandle, class IDetailChildrenBuilder& StructBuilder, IPropertyTypeCustomizationUtils& StructCustomizationUtils ) OVERRIDE;              
+	virtual void CustomizeChildren( TSharedRef<class IPropertyHandle> StructPropertyHandle, class IDetailChildrenBuilder& StructBuilder, IPropertyTypeCustomizationUtils& StructCustomizationUtils ) override;              
 
 	static TSharedPtr<SWidget> CreatePickerWidget(TSharedRef<class IPropertyHandle> StructPropertyHandle,  const FString& FileFilterExtension = FString(TEXT("*")), const FOnPathPicked& OnPathPicked = FOnPathPicked());
 
diff --git a/Engine/Source/Editor/DetailCustomizations/Private/GeneralProjectSettingsDetails.h b/Engine/Source/Editor/DetailCustomizations/Private/GeneralProjectSettingsDetails.h
--- a/Engine/Source/Editor/DetailCustomizations/Private/GeneralProjectSettingsDetails.h
+++ b/Engine/Source/Editor/DetailCustomizations/Private/GeneralProjectSettingsDetails.h
@@ -11,7 +11,7 @@ public:
 	static TSharedRef<IDetailCustomization> MakeInstance();
 
 	// IDetailCustomization interface
-	virtual void CustomizeDetails(IDetailLayoutBuilder& DetailBuilder) OVERRIDE;
+	virtual void CustomizeDetails(IDetailLayoutBuilder& DetailBuilder) override;
 	// End of IDetailCustomization interface
 };
 
diff --git a/Engine/Source/Editor/DetailCustomizations/Private/InputSettingsDetails.h b/Engine/Source/Editor/DetailCustomizations/Private/InputSettingsDetails.h
--- a/Engine/Source/Editor/DetailCustomizations/Private/InputSettingsDetails.h
+++ b/Engine/Source/Editor/DetailCustomizations/Private/InputSettingsDetails.h
@@ -21,13 +21,13 @@ public:
 	FActionMappingsNodeBuilder( IDetailLayoutBuilder* InDetailLayoutBuilder, const TSharedPtr<IPropertyHandle>& InPropertyHandle );
 
 	/** IDetailCustomNodeBuilder interface */
-	virtual void SetOnRebuildChildren( FSimpleDelegate InOnRebuildChildren  ) OVERRIDE { OnRebuildChildren = InOnRebuildChildren; } 
-	virtual bool RequiresTick() const OVERRIDE { return true; }
-	virtual void Tick( float DeltaTime ) OVERRIDE;
-	virtual void GenerateHeaderRowContent( FDetailWidgetRow& NodeRow ) OVERRIDE;
-	virtual void GenerateChildContent( IDetailChildrenBuilder& ChildrenBuilder ) OVERRIDE;
-	virtual bool InitiallyCollapsed() const OVERRIDE { return true; };
-	virtual FName GetName() const OVERRIDE { return NAME_None; }
+	virtual void SetOnRebuildChildren( FSimpleDelegate InOnRebuildChildren  ) override { OnRebuildChildren = InOnRebuildChildren; } 
+	virtual bool RequiresTick() const override { return true; }
+	virtual void Tick( float DeltaTime ) override;
+	virtual void GenerateHeaderRowContent( FDetailWidgetRow& NodeRow ) override;
+	virtual void GenerateChildContent( IDetailChildrenBuilder& ChildrenBuilder ) override;
+	virtual bool InitiallyCollapsed() const override { return true; };
+	virtual FName GetName() const override { return NAME_None; }
 
 private:
 	void AddActionMappingButton_OnClick();
@@ -62,13 +62,13 @@ public:
 	FAxisMappingsNodeBuilder( IDetailLayoutBuilder* InDetailLayoutBuilder, const TSharedPtr<IPropertyHandle>& InPropertyHandle );
 
 	/** IDetailCustomNodeBuilder interface */
-	virtual void SetOnRebuildChildren( FSimpleDelegate InOnRebuildChildren  ) OVERRIDE { OnRebuildChildren = InOnRebuildChildren; } 
-	virtual bool RequiresTick() const OVERRIDE { return true; }
-	virtual void Tick( float DeltaTime ) OVERRIDE;
-	virtual void GenerateHeaderRowContent( FDetailWidgetRow& NodeRow ) OVERRIDE;
-	virtual void GenerateChildContent( IDetailChildrenBuilder& ChildrenBuilder ) OVERRIDE;
-	virtual bool InitiallyCollapsed() const OVERRIDE { return true; };
-	virtual FName GetName() const OVERRIDE { return NAME_None; }
+	virtual void SetOnRebuildChildren( FSimpleDelegate InOnRebuildChildren  ) override { OnRebuildChildren = InOnRebuildChildren; } 
+	virtual bool RequiresTick() const override { return true; }
+	virtual void Tick( float DeltaTime ) override;
+	virtual void GenerateHeaderRowContent( FDetailWidgetRow& NodeRow ) override;
+	virtual void GenerateChildContent( IDetailChildrenBuilder& ChildrenBuilder ) override;
+	virtual bool InitiallyCollapsed() const override { return true; };
+	virtual FName GetName() const override { return NAME_None; }
 
 private:
 	void AddAxisMappingButton_OnClick();
@@ -104,5 +104,5 @@ public:
 	static TSharedRef<IDetailCustomization> MakeInstance();
 
 	/** ILayoutDetails interface */
-	virtual void CustomizeDetails( class IDetailLayoutBuilder& DetailBuilder ) OVERRIDE;
+	virtual void CustomizeDetails( class IDetailLayoutBuilder& DetailBuilder ) override;
 };
diff --git a/Engine/Source/Editor/DetailCustomizations/Private/InputStructCustomization.h b/Engine/Source/Editor/DetailCustomizations/Private/InputStructCustomization.h
--- a/Engine/Source/Editor/DetailCustomizations/Private/InputStructCustomization.h
+++ b/Engine/Source/Editor/DetailCustomizations/Private/InputStructCustomization.h
@@ -11,8 +11,8 @@ public:
 	static TSharedRef<IPropertyTypeCustomization> MakeInstance();
 
 	/** IPropertyTypeCustomization interface */
-	virtual void CustomizeHeader( TSharedRef<class IPropertyHandle> InStructPropertyHandle, class FDetailWidgetRow& HeaderRow, IPropertyTypeCustomizationUtils& StructCustomizationUtils ) OVERRIDE;
-	virtual void CustomizeChildren( TSharedRef<class IPropertyHandle> InStructPropertyHandle, class IDetailChildrenBuilder& StructBuilder, IPropertyTypeCustomizationUtils& StructCustomizationUtils ) OVERRIDE;
+	virtual void CustomizeHeader( TSharedRef<class IPropertyHandle> InStructPropertyHandle, class FDetailWidgetRow& HeaderRow, IPropertyTypeCustomizationUtils& StructCustomizationUtils ) override;
+	virtual void CustomizeChildren( TSharedRef<class IPropertyHandle> InStructPropertyHandle, class IDetailChildrenBuilder& StructBuilder, IPropertyTypeCustomizationUtils& StructCustomizationUtils ) override;
 };
 
 /**
@@ -24,8 +24,8 @@ public:
 	static TSharedRef<IPropertyTypeCustomization> MakeInstance();
 
 	/** IPropertyTypeCustomization interface */
-	virtual void CustomizeHeader( TSharedRef<class IPropertyHandle> InStructPropertyHandle, class FDetailWidgetRow& HeaderRow, IPropertyTypeCustomizationUtils& StructCustomizationUtils ) OVERRIDE;
-	virtual void CustomizeChildren( TSharedRef<class IPropertyHandle> InStructPropertyHandle, class IDetailChildrenBuilder& StructBuilder, IPropertyTypeCustomizationUtils& StructCustomizationUtils ) OVERRIDE;
+	virtual void CustomizeHeader( TSharedRef<class IPropertyHandle> InStructPropertyHandle, class FDetailWidgetRow& HeaderRow, IPropertyTypeCustomizationUtils& StructCustomizationUtils ) override;
+	virtual void CustomizeChildren( TSharedRef<class IPropertyHandle> InStructPropertyHandle, class IDetailChildrenBuilder& StructBuilder, IPropertyTypeCustomizationUtils& StructCustomizationUtils ) override;
 
 private:
 	void RemoveActionMappingButton_OnClick();
@@ -43,8 +43,8 @@ public:
 	static TSharedRef<IPropertyTypeCustomization> MakeInstance();
 
 	/** IPropertyTypeCustomization interface */
-	virtual void CustomizeHeader( TSharedRef<class IPropertyHandle> InStructPropertyHandle, class FDetailWidgetRow& HeaderRow, IPropertyTypeCustomizationUtils& StructCustomizationUtils ) OVERRIDE;
-	virtual void CustomizeChildren( TSharedRef<class IPropertyHandle> InStructPropertyHandle, class IDetailChildrenBuilder& StructBuilder, IPropertyTypeCustomizationUtils& StructCustomizationUtils ) OVERRIDE;
+	virtual void CustomizeHeader( TSharedRef<class IPropertyHandle> InStructPropertyHandle, class FDetailWidgetRow& HeaderRow, IPropertyTypeCustomizationUtils& StructCustomizationUtils ) override;
+	virtual void CustomizeChildren( TSharedRef<class IPropertyHandle> InStructPropertyHandle, class IDetailChildrenBuilder& StructBuilder, IPropertyTypeCustomizationUtils& StructCustomizationUtils ) override;
 
 private:
 	void RemoveAxisMappingButton_OnClick();
diff --git a/Engine/Source/Editor/DetailCustomizations/Private/MathStructCustomizations.h b/Engine/Source/Editor/DetailCustomizations/Private/MathStructCustomizations.h
--- a/Engine/Source/Editor/DetailCustomizations/Private/MathStructCustomizations.h
+++ b/Engine/Source/Editor/DetailCustomizations/Private/MathStructCustomizations.h
@@ -18,8 +18,8 @@ public:
 	virtual ~FMathStructCustomization() {}
 
 	/** IPropertyTypeCustomization instance */
-	virtual void CustomizeHeader( TSharedRef<class IPropertyHandle> StructPropertyHandle, class FDetailWidgetRow& HeaderRow, IPropertyTypeCustomizationUtils& StructCustomizationUtils ) OVERRIDE;
-	virtual void CustomizeChildren( TSharedRef<class IPropertyHandle> StructPropertyHandle, class IDetailChildrenBuilder& StructBuilder, IPropertyTypeCustomizationUtils& StructCustomizationUtils ) OVERRIDE;
+	virtual void CustomizeHeader( TSharedRef<class IPropertyHandle> StructPropertyHandle, class FDetailWidgetRow& HeaderRow, IPropertyTypeCustomizationUtils& StructCustomizationUtils ) override;
+	virtual void CustomizeChildren( TSharedRef<class IPropertyHandle> StructPropertyHandle, class IDetailChildrenBuilder& StructBuilder, IPropertyTypeCustomizationUtils& StructCustomizationUtils ) override;
 	
 protected:
 	/**
@@ -132,7 +132,7 @@ public:
 	static TSharedRef<IPropertyTypeCustomization> MakeInstance();
 private:
 	/** FMathStructCustomization interface */
-	virtual void GetSortedChildren( TSharedRef<IPropertyHandle> StructPropertyHandle, TArray< TSharedRef<IPropertyHandle> >& OutChildren ) OVERRIDE;
+	virtual void GetSortedChildren( TSharedRef<IPropertyHandle> StructPropertyHandle, TArray< TSharedRef<IPropertyHandle> >& OutChildren ) override;
 };
 
 /**
@@ -151,9 +151,9 @@ private:
 	{}
 
 	/** FMathStructCustomization interface */
-	virtual void CustomizeHeader( TSharedRef<class IPropertyHandle> StructPropertyHandle, class FDetailWidgetRow& HeaderRow, IPropertyTypeCustomizationUtils& StructCustomizationUtils ) OVERRIDE;
-	virtual void MakeHeaderRow( TSharedRef<class IPropertyHandle>& InStructPropertyHandle, FDetailWidgetRow& Row ) OVERRIDE;
-	virtual void GetSortedChildren( TSharedRef<IPropertyHandle> StructPropertyHandle, TArray< TSharedRef<IPropertyHandle> >& OutChildren ) OVERRIDE;
+	virtual void CustomizeHeader( TSharedRef<class IPropertyHandle> StructPropertyHandle, class FDetailWidgetRow& HeaderRow, IPropertyTypeCustomizationUtils& StructCustomizationUtils ) override;
+	virtual void MakeHeaderRow( TSharedRef<class IPropertyHandle>& InStructPropertyHandle, FDetailWidgetRow& Row ) override;
+	virtual void GetSortedChildren( TSharedRef<IPropertyHandle> StructPropertyHandle, TArray< TSharedRef<IPropertyHandle> >& OutChildren ) override;
 
 	/**
 	 * Get the color used by this struct as a linear color value
diff --git a/Engine/Source/Editor/DetailCustomizations/Private/MathStructProxyCustomizations.h b/Engine/Source/Editor/DetailCustomizations/Private/MathStructProxyCustomizations.h
--- a/Engine/Source/Editor/DetailCustomizations/Private/MathStructProxyCustomizations.h
+++ b/Engine/Source/Editor/DetailCustomizations/Private/MathStructProxyCustomizations.h
@@ -138,10 +138,10 @@ class FMathStructProxyCustomization : public FMathStructCustomization
 {
 public:
 	/** IPropertyTypeCustomization interface */
-	virtual void CustomizeChildren( TSharedRef<class IPropertyHandle> StructPropertyHandle, class IDetailChildrenBuilder& StructBuilder, IPropertyTypeCustomizationUtils& StructCustomizationUtils ) OVERRIDE;
+	virtual void CustomizeChildren( TSharedRef<class IPropertyHandle> StructPropertyHandle, class IDetailChildrenBuilder& StructBuilder, IPropertyTypeCustomizationUtils& StructCustomizationUtils ) override;
 
 	/** FMathStructCustomization interface */
-	virtual void MakeHeaderRow( TSharedRef<class IPropertyHandle>& InStructPropertyHandle, FDetailWidgetRow& Row ) OVERRIDE;
+	virtual void MakeHeaderRow( TSharedRef<class IPropertyHandle>& InStructPropertyHandle, FDetailWidgetRow& Row ) override;
 
 protected:
 
@@ -233,12 +233,12 @@ public:
 	}
 
 	/** FMathStructCustomization interface */
-	virtual void MakeHeaderRow( TSharedRef<class IPropertyHandle>& StructPropertyHandle, FDetailWidgetRow& Row ) OVERRIDE;
+	virtual void MakeHeaderRow( TSharedRef<class IPropertyHandle>& StructPropertyHandle, FDetailWidgetRow& Row ) override;
 
 protected:
 	/** FMathStructProxyCustomization interface */
-	virtual bool CacheValues( TWeakPtr<IPropertyHandle> WeakHandlePtr ) const OVERRIDE;
-	virtual bool FlushValues( TWeakPtr<IPropertyHandle> WeakHandlePtr ) const OVERRIDE;
+	virtual bool CacheValues( TWeakPtr<IPropertyHandle> WeakHandlePtr ) const override;
+	virtual bool FlushValues( TWeakPtr<IPropertyHandle> WeakHandlePtr ) const override;
 
 protected:
 	/** Cached rotation values */
@@ -270,6 +270,6 @@ public:
 
 protected:
 	/** FMathStructProxyCustomization interface */
-	virtual bool CacheValues( TWeakPtr<IPropertyHandle> WeakHandlePtr ) const OVERRIDE;
-	virtual bool FlushValues( TWeakPtr<IPropertyHandle> WeakHandlePtr ) const OVERRIDE;
+	virtual bool CacheValues( TWeakPtr<IPropertyHandle> WeakHandlePtr ) const override;
+	virtual bool FlushValues( TWeakPtr<IPropertyHandle> WeakHandlePtr ) const override;
 };
\ No newline at end of file
diff --git a/Engine/Source/Editor/DetailCustomizations/Private/ParticleModuleDetails.h b/Engine/Source/Editor/DetailCustomizations/Private/ParticleModuleDetails.h
--- a/Engine/Source/Editor/DetailCustomizations/Private/ParticleModuleDetails.h
+++ b/Engine/Source/Editor/DetailCustomizations/Private/ParticleModuleDetails.h
@@ -19,7 +19,7 @@ public:
 	static TSharedRef<IDetailCustomization> MakeInstance();
 
 	/** IDetailCustomization interface */
-	virtual void CustomizeDetails( IDetailLayoutBuilder& DetailBuilder ) OVERRIDE;
+	virtual void CustomizeDetails( IDetailLayoutBuilder& DetailBuilder ) override;
 };
 
 class FParticleModuleSubUVDetails : public FParticleModuleDetailsBase
@@ -29,7 +29,7 @@ public:
 	static TSharedRef<IDetailCustomization> MakeInstance();
 
 	/** IDetailCustomization interface */
-	virtual void CustomizeDetails( IDetailLayoutBuilder& DetailBuilder ) OVERRIDE;
+	virtual void CustomizeDetails( IDetailLayoutBuilder& DetailBuilder ) override;
 };
 
 class FParticleModuleAccelerationDetails : public FParticleModuleDetailsBase
@@ -39,7 +39,7 @@ public:
 	static TSharedRef<IDetailCustomization> MakeInstance();
 
 	/** IDetailCustomization interface */
-	virtual void CustomizeDetails( IDetailLayoutBuilder& DetailBuilder ) OVERRIDE;
+	virtual void CustomizeDetails( IDetailLayoutBuilder& DetailBuilder ) override;
 };
 
 class FParticleModuleAccelerationDragDetails : public FParticleModuleDetailsBase
@@ -49,7 +49,7 @@ public:
 	static TSharedRef<IDetailCustomization> MakeInstance();
 
 	/** IDetailCustomization interface */
-	virtual void CustomizeDetails( IDetailLayoutBuilder& DetailBuilder ) OVERRIDE;
+	virtual void CustomizeDetails( IDetailLayoutBuilder& DetailBuilder ) override;
 };
 
  class FParticleModuleAccelerationDragScaleOverLifeDetails : public FParticleModuleDetailsBase
@@ -59,7 +59,7 @@ public:
  	static TSharedRef<IDetailCustomization> MakeInstance();
  
  	/** IDetailCustomization interface */
- 	virtual void CustomizeDetails( IDetailLayoutBuilder& DetailBuilder ) OVERRIDE;
+ 	virtual void CustomizeDetails( IDetailLayoutBuilder& DetailBuilder ) override;
  };
 
  class FParticleModuleCollisionGPUDetails : public FParticleModuleDetailsBase
@@ -69,7 +69,7 @@ public:
  	static TSharedRef<IDetailCustomization> MakeInstance();
  
  	/** IDetailCustomization interface */
- 	virtual void CustomizeDetails( IDetailLayoutBuilder& DetailBuilder ) OVERRIDE;
+ 	virtual void CustomizeDetails( IDetailLayoutBuilder& DetailBuilder ) override;
  };
 
  class FParticleModuleOrbitDetails : public FParticleModuleDetailsBase
@@ -79,7 +79,7 @@ public:
 	 static TSharedRef<IDetailCustomization> MakeInstance();
 
 	 /** IDetailCustomization interface */
-	 virtual void CustomizeDetails( IDetailLayoutBuilder& DetailBuilder ) OVERRIDE;
+	 virtual void CustomizeDetails( IDetailLayoutBuilder& DetailBuilder ) override;
  };
 
  class FParticleModuleSizeMultiplyLifeDetails : public FParticleModuleDetailsBase
@@ -89,7 +89,7 @@ public:
 	 static TSharedRef<IDetailCustomization> MakeInstance();
 
 	 /** IDetailCustomization interface */
-	 virtual void CustomizeDetails( IDetailLayoutBuilder& DetailBuilder ) OVERRIDE;
+	 virtual void CustomizeDetails( IDetailLayoutBuilder& DetailBuilder ) override;
  };
 
  class FParticleModuleSizeScaleDetails : public FParticleModuleDetailsBase
@@ -99,7 +99,7 @@ public:
 	 static TSharedRef<IDetailCustomization> MakeInstance();
 
 	 /** IDetailCustomization interface */
-	 virtual void CustomizeDetails( IDetailLayoutBuilder& DetailBuilder ) OVERRIDE;
+	 virtual void CustomizeDetails( IDetailLayoutBuilder& DetailBuilder ) override;
  };
 
  class FParticleModuleVectorFieldScaleDetails : public FParticleModuleDetailsBase
@@ -109,7 +109,7 @@ public:
 	 static TSharedRef<IDetailCustomization> MakeInstance();
 
 	 /** IDetailCustomization interface */
-	 virtual void CustomizeDetails( IDetailLayoutBuilder& DetailBuilder ) OVERRIDE;
+	 virtual void CustomizeDetails( IDetailLayoutBuilder& DetailBuilder ) override;
  };
 
  class FParticleModuleVectorFieldScaleOverLifeDetails : public FParticleModuleDetailsBase
@@ -119,6 +119,6 @@ public:
 	 static TSharedRef<IDetailCustomization> MakeInstance();
 
 	 /** IDetailCustomization interface */
-	 virtual void CustomizeDetails( IDetailLayoutBuilder& DetailBuilder ) OVERRIDE;
+	 virtual void CustomizeDetails( IDetailLayoutBuilder& DetailBuilder ) override;
  };
  
\ No newline at end of file
diff --git a/Engine/Source/Editor/DetailCustomizations/Private/ParticleSystemComponentDetails.h b/Engine/Source/Editor/DetailCustomizations/Private/ParticleSystemComponentDetails.h
--- a/Engine/Source/Editor/DetailCustomizations/Private/ParticleSystemComponentDetails.h
+++ b/Engine/Source/Editor/DetailCustomizations/Private/ParticleSystemComponentDetails.h
@@ -9,7 +9,7 @@ public:
 	static TSharedRef<IDetailCustomization> MakeInstance();
 private:
 	/** IDetailCustomization interface */
-	virtual void CustomizeDetails( IDetailLayoutBuilder& InDetailLayout ) OVERRIDE;
+	virtual void CustomizeDetails( IDetailLayoutBuilder& InDetailLayout ) override;
 
 	/** Handles the Auto-Populate button's on click event */
 	FReply OnAutoPopulateClicked();
diff --git a/Engine/Source/Editor/DetailCustomizations/Private/PhysicsSettingsDetails.h b/Engine/Source/Editor/DetailCustomizations/Private/PhysicsSettingsDetails.h
--- a/Engine/Source/Editor/DetailCustomizations/Private/PhysicsSettingsDetails.h
+++ b/Engine/Source/Editor/DetailCustomizations/Private/PhysicsSettingsDetails.h
@@ -41,7 +41,7 @@ public:
 	* @param InArgs - The Slate argument list.
 	*/
 	void Construct(const FArguments& InArgs, const TSharedRef<STableViewBase>& InOwnerTableView);
-	virtual TSharedRef<SWidget> GenerateWidgetForColumn(const FName& ColumnName) OVERRIDE;
+	virtual TSharedRef<SWidget> GenerateWidgetForColumn(const FName& ColumnName) override;
 	FString GetTypeString() const;
 
 	FText GetName() const;
@@ -64,7 +64,7 @@ public:
 	static TSharedRef<IDetailCustomization> MakeInstance();
 
 	/** IDetailCustomization interface */
-	virtual void CustomizeDetails( IDetailLayoutBuilder& DetailBuilder ) OVERRIDE;
+	virtual void CustomizeDetails( IDetailLayoutBuilder& DetailBuilder ) override;
 
 private:
 
diff --git a/Engine/Source/Editor/DetailCustomizations/Private/SkeletalControlNodeDetails.h b/Engine/Source/Editor/DetailCustomizations/Private/SkeletalControlNodeDetails.h
--- a/Engine/Source/Editor/DetailCustomizations/Private/SkeletalControlNodeDetails.h
+++ b/Engine/Source/Editor/DetailCustomizations/Private/SkeletalControlNodeDetails.h
@@ -11,7 +11,7 @@ public:
 	static TSharedRef<IDetailCustomization> MakeInstance();
 
 	// IDetailCustomization interface
-	virtual void CustomizeDetails(class IDetailLayoutBuilder& DetailBuilder) OVERRIDE;
+	virtual void CustomizeDetails(class IDetailLayoutBuilder& DetailBuilder) override;
 	// End of IDetailCustomization interface
 
 protected:
diff --git a/Engine/Source/Editor/DetailCustomizations/Private/SkeletalMeshComponentDetails.cpp b/Engine/Source/Editor/DetailCustomizations/Private/SkeletalMeshComponentDetails.cpp
--- a/Engine/Source/Editor/DetailCustomizations/Private/SkeletalMeshComponentDetails.cpp
+++ b/Engine/Source/Editor/DetailCustomizations/Private/SkeletalMeshComponentDetails.cpp
@@ -13,7 +13,7 @@
 class FAnimBlueprintFilter : public IClassViewerFilter
 {
 public:
-	virtual bool IsClassAllowed(const FClassViewerInitializationOptions& InInitOptions, const UClass* InClass, TSharedRef< class FClassViewerFilterFuncs > InFilterFuncs ) OVERRIDE
+	virtual bool IsClassAllowed(const FClassViewerInitializationOptions& InInitOptions, const UClass* InClass, TSharedRef< class FClassViewerFilterFuncs > InFilterFuncs ) override
 	{
 		if(InFilterFuncs->IfInChildOfClassesSet(AllowedChildrenOfClasses, InClass) != EFilterReturn::Failed)
 		{
@@ -22,7 +22,7 @@ public:
 		return false;
 	}
 
-	virtual bool IsUnloadedClassAllowed(const FClassViewerInitializationOptions& InInitOptions, const TSharedRef< const class IUnloadedBlueprintData > InUnloadedClassData, TSharedRef< class FClassViewerFilterFuncs > InFilterFuncs) OVERRIDE
+	virtual bool IsUnloadedClassAllowed(const FClassViewerInitializationOptions& InInitOptions, const TSharedRef< const class IUnloadedBlueprintData > InUnloadedClassData, TSharedRef< class FClassViewerFilterFuncs > InFilterFuncs) override
 	{
 		return InFilterFuncs->IfInChildOfClassesSet(AllowedChildrenOfClasses, InUnloadedClassData) != EFilterReturn::Failed;
 	}
diff --git a/Engine/Source/Editor/DetailCustomizations/Private/SlateBrushCustomization.cpp b/Engine/Source/Editor/DetailCustomizations/Private/SlateBrushCustomization.cpp
--- a/Engine/Source/Editor/DetailCustomizations/Private/SlateBrushCustomization.cpp
+++ b/Engine/Source/Editor/DetailCustomizations/Private/SlateBrushCustomization.cpp
@@ -271,7 +271,7 @@ private:
 	/**
 	 * SWidget interface
 	 */
-	virtual FReply OnMouseButtonDown( const FGeometry& MyGeometry, const FPointerEvent& MouseEvent ) OVERRIDE
+	virtual FReply OnMouseButtonDown( const FGeometry& MyGeometry, const FPointerEvent& MouseEvent ) override
 	{
 		if( MouseEvent.GetEffectingButton() == EKeys::LeftMouseButton ) 
 		{
@@ -286,7 +286,7 @@ private:
 		}
 	}
 
-	virtual FReply OnMouseButtonUp( const FGeometry& MyGeometry, const FPointerEvent& MouseEvent ) OVERRIDE
+	virtual FReply OnMouseButtonUp( const FGeometry& MyGeometry, const FPointerEvent& MouseEvent ) override
 	{
 		if( MouseEvent.GetEffectingButton() == EKeys::LeftMouseButton && bUserIsResizing )
 		{
@@ -299,7 +299,7 @@ private:
 		}
 	}
 
-	virtual FReply OnMouseMove( const FGeometry& MyGeometry, const FPointerEvent& MouseEvent ) OVERRIDE
+	virtual FReply OnMouseMove( const FGeometry& MyGeometry, const FPointerEvent& MouseEvent ) override
 	{
 		const FVector2D LocalMouseCoordinates( MyGeometry.AbsoluteToLocal( MouseEvent.GetScreenSpacePosition() ) );
 
@@ -330,14 +330,14 @@ private:
 		return FReply::Unhandled();
 	}
 	
-	virtual void OnMouseEnter( const FGeometry& MyGeometry, const FPointerEvent& MouseEvent ) OVERRIDE
+	virtual void OnMouseEnter( const FGeometry& MyGeometry, const FPointerEvent& MouseEvent ) override
 	{
 		const FVector2D LocalMouseCoordinates( MyGeometry.AbsoluteToLocal( MouseEvent.GetScreenSpacePosition() ) );
 		MouseZone = FindMouseZone( LocalMouseCoordinates );
 		SBorder::OnMouseEnter( MyGeometry, MouseEvent );
 	}
 
-	virtual void OnMouseLeave( const FPointerEvent& MouseEvent ) OVERRIDE
+	virtual void OnMouseLeave( const FPointerEvent& MouseEvent ) override
 	{
 		if( !bUserIsResizing )
 		{
@@ -346,7 +346,7 @@ private:
 		}
 	}
 
-	virtual FCursorReply OnCursorQuery( const FGeometry& MyGeometry, const FPointerEvent& CursorEvent ) const OVERRIDE
+	virtual FCursorReply OnCursorQuery( const FGeometry& MyGeometry, const FPointerEvent& CursorEvent ) const override
 	{
 		if( MouseZone == WZ_RightBorder )
 		{
diff --git a/Engine/Source/Editor/DetailCustomizations/Private/SlateSoundCustomization.h b/Engine/Source/Editor/DetailCustomizations/Private/SlateSoundCustomization.h
--- a/Engine/Source/Editor/DetailCustomizations/Private/SlateSoundCustomization.h
+++ b/Engine/Source/Editor/DetailCustomizations/Private/SlateSoundCustomization.h
@@ -9,8 +9,8 @@ public:
 	static TSharedRef<IPropertyTypeCustomization> MakeInstance();
 
 	/** IPropertyTypeCustomization interface */
-	virtual void CustomizeHeader(TSharedRef<class IPropertyHandle> StructPropertyHandle, class FDetailWidgetRow& HeaderRow, IPropertyTypeCustomizationUtils& StructCustomizationUtils) OVERRIDE;
-	virtual void CustomizeChildren(TSharedRef<class IPropertyHandle> StructPropertyHandle, class IDetailChildrenBuilder& StructBuilder, IPropertyTypeCustomizationUtils& StructCustomizationUtils) OVERRIDE;
+	virtual void CustomizeHeader(TSharedRef<class IPropertyHandle> StructPropertyHandle, class FDetailWidgetRow& HeaderRow, IPropertyTypeCustomizationUtils& StructCustomizationUtils) override;
+	virtual void CustomizeChildren(TSharedRef<class IPropertyHandle> StructPropertyHandle, class IDetailChildrenBuilder& StructBuilder, IPropertyTypeCustomizationUtils& StructCustomizationUtils) override;
 
 protected:
 	/** Called when the resource object used by this FSlateSound has been changed */
diff --git a/Engine/Source/Editor/DetailCustomizations/Private/StringAssetReferenceCustomization.h b/Engine/Source/Editor/DetailCustomizations/Private/StringAssetReferenceCustomization.h
--- a/Engine/Source/Editor/DetailCustomizations/Private/StringAssetReferenceCustomization.h
+++ b/Engine/Source/Editor/DetailCustomizations/Private/StringAssetReferenceCustomization.h
@@ -14,8 +14,8 @@ public:
 	}
 
 	/** IPropertyTypeCustomization interface */
-	virtual void CustomizeHeader( TSharedRef<class IPropertyHandle> InStructPropertyHandle, class FDetailWidgetRow& HeaderRow, IPropertyTypeCustomizationUtils& StructCustomizationUtils ) OVERRIDE;
-	virtual void CustomizeChildren( TSharedRef<class IPropertyHandle> InStructPropertyHandle, class IDetailChildrenBuilder& StructBuilder, IPropertyTypeCustomizationUtils& StructCustomizationUtils ) OVERRIDE;
+	virtual void CustomizeHeader( TSharedRef<class IPropertyHandle> InStructPropertyHandle, class FDetailWidgetRow& HeaderRow, IPropertyTypeCustomizationUtils& StructCustomizationUtils ) override;
+	virtual void CustomizeChildren( TSharedRef<class IPropertyHandle> InStructPropertyHandle, class IDetailChildrenBuilder& StructBuilder, IPropertyTypeCustomizationUtils& StructCustomizationUtils ) override;
 
 private:
 	bool OnShouldFilterAsset( const class FAssetData& InAssetData ) const;
diff --git a/Engine/Source/Editor/DetailCustomizations/Private/StringClassReferenceCustomization.h b/Engine/Source/Editor/DetailCustomizations/Private/StringClassReferenceCustomization.h
--- a/Engine/Source/Editor/DetailCustomizations/Private/StringClassReferenceCustomization.h
+++ b/Engine/Source/Editor/DetailCustomizations/Private/StringClassReferenceCustomization.h
@@ -14,8 +14,8 @@ public:
 	}
 
 	/** IPropertyTypeCustomization interface */
-	virtual void CustomizeHeader(TSharedRef<class IPropertyHandle> InStructPropertyHandle, class FDetailWidgetRow& HeaderRow, IPropertyTypeCustomizationUtils& StructCustomizationUtils) OVERRIDE;
-	virtual void CustomizeChildren(TSharedRef<class IPropertyHandle> InStructPropertyHandle, class IDetailChildrenBuilder& StructBuilder, IPropertyTypeCustomizationUtils& StructCustomizationUtils) OVERRIDE;
+	virtual void CustomizeHeader(TSharedRef<class IPropertyHandle> InStructPropertyHandle, class FDetailWidgetRow& HeaderRow, IPropertyTypeCustomizationUtils& StructCustomizationUtils) override;
+	virtual void CustomizeChildren(TSharedRef<class IPropertyHandle> InStructPropertyHandle, class IDetailChildrenBuilder& StructBuilder, IPropertyTypeCustomizationUtils& StructCustomizationUtils) override;
 
 private:
 	/** @return The class currently set on this reference */
diff --git a/Engine/Source/Editor/DetailCustomizations/Private/VehicleTransmissionDataCustomization.h b/Engine/Source/Editor/DetailCustomizations/Private/VehicleTransmissionDataCustomization.h
--- a/Engine/Source/Editor/DetailCustomizations/Private/VehicleTransmissionDataCustomization.h
+++ b/Engine/Source/Editor/DetailCustomizations/Private/VehicleTransmissionDataCustomization.h
@@ -12,8 +12,8 @@ public:
 	}
 
 	/** IPropertyTypeCustomization interface */
-	virtual void CustomizeHeader(TSharedRef<IPropertyHandle> StructPropertyHandle, class FDetailWidgetRow& HeaderRow, IPropertyTypeCustomizationUtils& StructCustomizationUtils) OVERRIDE;
-	virtual void CustomizeChildren( TSharedRef<class IPropertyHandle> StructPropertyHandle, class IDetailChildrenBuilder& StructBuilder, IPropertyTypeCustomizationUtils& StructCustomizationUtils ) OVERRIDE;
+	virtual void CustomizeHeader(TSharedRef<IPropertyHandle> StructPropertyHandle, class FDetailWidgetRow& HeaderRow, IPropertyTypeCustomizationUtils& StructCustomizationUtils) override;
+	virtual void CustomizeChildren( TSharedRef<class IPropertyHandle> StructPropertyHandle, class IDetailChildrenBuilder& StructBuilder, IPropertyTypeCustomizationUtils& StructCustomizationUtils ) override;
 
 private:
 	FVehicleTransmissionData * SelectedTransmission;
diff --git a/Engine/Source/Editor/DetailCustomizations/Private/WheeledVehicleMovementComponent4WDetails.h b/Engine/Source/Editor/DetailCustomizations/Private/WheeledVehicleMovementComponent4WDetails.h
--- a/Engine/Source/Editor/DetailCustomizations/Private/WheeledVehicleMovementComponent4WDetails.h
+++ b/Engine/Source/Editor/DetailCustomizations/Private/WheeledVehicleMovementComponent4WDetails.h
@@ -11,7 +11,7 @@ public:
 	static TSharedRef<IDetailCustomization> MakeInstance();
 
 	/** IDetailCustomization interface */
-	virtual void CustomizeDetails( IDetailLayoutBuilder& DetailBuilder ) OVERRIDE;
+	virtual void CustomizeDetails( IDetailLayoutBuilder& DetailBuilder ) override;
 	
 private:
 
@@ -19,11 +19,11 @@ private:
 	{
 		FSteeringCurveEditor(UWheeledVehicleMovementComponent4W * InVehicle = NULL);
 		/** FCurveOwnerInterface interface */
-		virtual TArray<FRichCurveEditInfoConst> GetCurves() const OVERRIDE;
-		virtual TArray<FRichCurveEditInfo> GetCurves() OVERRIDE;
-		virtual UObject* GetOwner() OVERRIDE;
-		virtual void ModifyOwner() OVERRIDE;
-		virtual void MakeTransactional() OVERRIDE;
+		virtual TArray<FRichCurveEditInfoConst> GetCurves() const override;
+		virtual TArray<FRichCurveEditInfo> GetCurves() override;
+		virtual UObject* GetOwner() override;
+		virtual void ModifyOwner() override;
+		virtual void MakeTransactional() override;
 
 	private:
 		UWheeledVehicleMovementComponent4W * VehicleComponent;
@@ -35,11 +35,11 @@ private:
 	{
 		FTorqueCurveEditor(UWheeledVehicleMovementComponent4W * InVehicle = NULL);
 		/** FCurveOwnerInterface interface */
-		virtual TArray<FRichCurveEditInfoConst> GetCurves() const OVERRIDE;
-		virtual TArray<FRichCurveEditInfo> GetCurves() OVERRIDE;
-		virtual UObject* GetOwner() OVERRIDE;
-		virtual void ModifyOwner() OVERRIDE;
-		virtual void MakeTransactional() OVERRIDE;
+		virtual TArray<FRichCurveEditInfoConst> GetCurves() const override;
+		virtual TArray<FRichCurveEditInfo> GetCurves() override;
+		virtual UObject* GetOwner() override;
+		virtual void ModifyOwner() override;
+		virtual void MakeTransactional() override;
 
 	private:
 		UWheeledVehicleMovementComponent4W * VehicleComponent;
diff --git a/Engine/Source/Editor/DetailCustomizations/Public/DetailCustomizations.h b/Engine/Source/Editor/DetailCustomizations/Public/DetailCustomizations.h
--- a/Engine/Source/Editor/DetailCustomizations/Public/DetailCustomizations.h
+++ b/Engine/Source/Editor/DetailCustomizations/Public/DetailCustomizations.h
@@ -22,7 +22,7 @@ public:
 	 *
 	 * @return	Whether the module supports shutdown separate from the rest of the engine.
 	 */
-	virtual bool SupportsDynamicReloading() OVERRIDE
+	virtual bool SupportsDynamicReloading() override
 	{
 		return true;
 	}
diff --git a/Engine/Source/Editor/DeviceProfileEditor/Private/DeviceProfileConsoleVariableColumn.cpp b/Engine/Source/Editor/DeviceProfileEditor/Private/DeviceProfileConsoleVariableColumn.cpp
--- a/Engine/Source/Editor/DeviceProfileEditor/Private/DeviceProfileConsoleVariableColumn.cpp
+++ b/Engine/Source/Editor/DeviceProfileEditor/Private/DeviceProfileConsoleVariableColumn.cpp
@@ -46,39 +46,39 @@ public:
 public:
 
 	/** Begin IPropertyTableCellPresenter interface */
-	virtual TSharedRef<class SWidget> ConstructDisplayWidget() OVERRIDE;
+	virtual TSharedRef<class SWidget> ConstructDisplayWidget() override;
 
-	virtual bool RequiresDropDown() OVERRIDE
+	virtual bool RequiresDropDown() override
 	{
 		return false;
 	}
 
-	virtual TSharedRef< class SWidget > ConstructEditModeCellWidget() OVERRIDE
+	virtual TSharedRef< class SWidget > ConstructEditModeCellWidget() override
 	{
 		return ConstructDisplayWidget();
 	}
 
-	virtual TSharedRef< class SWidget > ConstructEditModeDropDownWidget() OVERRIDE
+	virtual TSharedRef< class SWidget > ConstructEditModeDropDownWidget() override
 	{
 		return SNullWidget::NullWidget;
 	}
 
-	virtual TSharedRef< class SWidget > WidgetToFocusOnEdit() OVERRIDE
+	virtual TSharedRef< class SWidget > WidgetToFocusOnEdit() override
 	{
 		return SNullWidget::NullWidget;
 	}
 
-	virtual bool HasReadOnlyEditMode() OVERRIDE
+	virtual bool HasReadOnlyEditMode() override
 	{
 		return true;
 	}
 
-	virtual FString GetValueAsString() OVERRIDE
+	virtual FString GetValueAsString() override
 	{
 		return TEXT("");
 	}
 
-	virtual FText GetValueAsText() OVERRIDE
+	virtual FText GetValueAsText() override
 	{
 		return FText::FromString(TEXT(""));
 	}
