UE4: Many content browser perf improvements and cleanups.

Removed all FAssetPickerConfig options that were not used by anything.
Package dirty refreshes are now managed entirely by FFrontendFilter_Modified, and only apply while the filter is enabled.
List refreshes have once again been split into SLOW source item refreshes and FAST frontend refreshes. Only use the slow refresh when source items change!
Replaced some usage of RemoveAt with RemoveAtSwap when appropriate. A lot of this code is very performance sensitive.
Text filtering is now done via a FFrontendFilter.
Consolidated the "show other developers" filter into a single FFrontendFilter.
There is no such thing as "Dynamic Filters" anymore. Only FrontendFilters.

diff --git a/Engine/Source/Editor/ContentBrowser/Private/AssetViewWidgets.cpp b/Engine/Source/Editor/ContentBrowser/Private/AssetViewWidgets.cpp
--- a/Engine/Source/Editor/ContentBrowser/Private/AssetViewWidgets.cpp
+++ b/Engine/Source/Editor/ContentBrowser/Private/AssetViewWidgets.cpp
@@ -127,8 +127,6 @@ void SAssetViewItem::Construct( const FArguments& InArgs )
 	OnItemDestroyed = InArgs._OnItemDestroyed;
 	ShouldAllowToolTip = InArgs._ShouldAllowToolTip;
 	ThumbnailEditMode = InArgs._ThumbnailEditMode;
-	LabelVisibility = InArgs._LabelVisibility;
-	ConstructToolTipForAsset = InArgs._ConstructToolTip;
 	OnAssetsDragDropped = InArgs._OnAssetsDragDropped;
 	OnPathsDragDropped = InArgs._OnPathsDragDropped;
 	OnFilesDragDropped = InArgs._OnFilesDragDropped;
@@ -401,11 +399,6 @@ TSharedRef<SToolTip> SAssetViewItem::CreateToolTipWidget() const
 			const FAssetData& AssetData = StaticCastSharedPtr<FAssetViewAsset>(AssetItem)->Data;
 			UClass* AssetClass = FindObject<UClass>(ANY_PACKAGE, *AssetData.AssetClass.ToString());
 
-			if ( ConstructToolTipForAsset.IsBound() )
-			{
-				return ConstructToolTipForAsset.Execute( AssetData );
-			}
-
 			// The tooltip contains the name, class, path, and asset registry tags
 			const FText NameText = FText::FromName( AssetData.AssetName );
 			const FText ClassText = FText::Format( LOCTEXT("ClassName", "({0})"), FText::FromName( AssetData.AssetClass ) );
@@ -891,7 +884,6 @@ void SAssetListItem::Construct( const FArguments& InArgs )
 		.OnItemDestroyed(InArgs._OnItemDestroyed)
 		.ShouldAllowToolTip(InArgs._ShouldAllowToolTip)
 		.ThumbnailEditMode(InArgs._ThumbnailEditMode)
-		.ConstructToolTip(InArgs._ConstructToolTip)
 		.OnAssetsDragDropped(InArgs._OnAssetsDragDropped)
 		.OnPathsDragDropped(InArgs._OnPathsDragDropped)
 		.OnFilesDragDropped(InArgs._OnFilesDragDropped)
@@ -1102,8 +1094,6 @@ void SAssetTileItem::Construct( const FArguments& InArgs )
 		.OnItemDestroyed(InArgs._OnItemDestroyed)
 		.ShouldAllowToolTip(InArgs._ShouldAllowToolTip)
 		.ThumbnailEditMode(InArgs._ThumbnailEditMode)
-		.LabelVisibility(InArgs._LabelVisibility)
-		.ConstructToolTip(InArgs._ConstructToolTip)
 		.OnAssetsDragDropped(InArgs._OnAssetsDragDropped)
 		.OnPathsDragDropped(InArgs._OnPathsDragDropped)
 		.OnFilesDragDropped(InArgs._OnFilesDragDropped)
@@ -1292,7 +1282,6 @@ void SAssetColumnItem::Construct( const FArguments& InArgs )
 		.OnRenameCommit(InArgs._OnRenameCommit)
 		.OnVerifyRenameCommit(InArgs._OnVerifyRenameCommit)
 		.OnItemDestroyed(InArgs._OnItemDestroyed)
-		.ConstructToolTip(InArgs._ConstructToolTip )
 		.OnAssetsDragDropped(InArgs._OnAssetsDragDropped)
 		.OnPathsDragDropped(InArgs._OnPathsDragDropped)
 		.OnFilesDragDropped(InArgs._OnFilesDragDropped)
diff --git a/Engine/Source/Editor/ContentBrowser/Private/AssetViewWidgets.h b/Engine/Source/Editor/ContentBrowser/Private/AssetViewWidgets.h
--- a/Engine/Source/Editor/ContentBrowser/Private/AssetViewWidgets.h
+++ b/Engine/Source/Editor/ContentBrowser/Private/AssetViewWidgets.h
@@ -47,8 +47,6 @@ public:
 	SLATE_BEGIN_ARGS( SAssetViewItem )
 		: _ShouldAllowToolTip(true)
 		, _ThumbnailEditMode(false)
-		, _LabelVisibility(EVisibility::Visible)
-		, _ConstructToolTip()
 		{}
 
 		/** Data for the asset this item represents */
@@ -72,12 +70,6 @@ public:
 		/** If true, display the thumbnail edit mode UI */
 		SLATE_ATTRIBUTE( bool, ThumbnailEditMode )
 
-		/** The visibility state of the labels used in asset view items */
-		SLATE_ATTRIBUTE( EVisibility, LabelVisibility )
-
-		/** The tooltip to display when the cursor hovers over this widget */
-		SLATE_EVENT( FConstructToolTipForAsset, ConstructToolTip )
-
 		/** Delegate for when assets are dropped on this item, if it is a folder */
 		SLATE_EVENT( FOnAssetsDragDropped, OnAssetsDragDropped )
 
@@ -212,9 +204,6 @@ protected:
 	/** If true, display the thumbnail edit mode UI */
 	TAttribute<bool> ThumbnailEditMode;
 
-	/** The visibility state of the labels used in asset view items */
-	TAttribute<EVisibility> LabelVisibility;
-
 	/**
 	 * A map of class names to their important asset registry tags and values.
 	 * Important tags will be featured in the tooltip if the value matches.
@@ -228,9 +217,6 @@ protected:
 	/** Cached flag describing if the package is dirty */
 	bool bPackageDirty;
 
-	/** The callback for external code to construct the tooltip used for this asset */
-	FConstructToolTipForAsset ConstructToolTipForAsset;
-
 	/** Flag indicating whether we have requested initial source control state */
 	bool bSourceControlStateRequested;
 
@@ -298,9 +284,6 @@ public:
 		/** Called when any asset item is destroyed. Used in thumbnail management */
 		SLATE_EVENT( FOnItemDestroyed, OnItemDestroyed )
 
-		/** The tooltip to display when the cursor hovers over this widget */
-		SLATE_EVENT( FConstructToolTipForAsset, ConstructToolTip )
-
 		/** If false, the tooltip will not be displayed */
 		SLATE_ATTRIBUTE( bool, ShouldAllowToolTip )
 
@@ -360,7 +343,6 @@ public:
 		, _ItemWidth(16)
 		, _ShouldAllowToolTip(true)
 		, _ThumbnailEditMode(false)
-		, _LabelVisibility(EVisibility::Visible)
 		, _ThumbnailLabel( EThumbnailLabel::ClassName )
 		, _ThumbnailHintColorAndOpacity( FLinearColor( 0.0f, 0.0f, 0.0f, 0.0f ) )
 		, _AllowThumbnailHintLabel(true)
@@ -399,9 +381,6 @@ public:
 		/** Called when any asset item is destroyed. Used in thumbnail management */
 		SLATE_EVENT( FOnItemDestroyed, OnItemDestroyed )
 
-		/** The tooltip to display when the cursor hovers over this widget */
-		SLATE_EVENT( FConstructToolTipForAsset, ConstructToolTip )
-
 		/** If false, the tooltip will not be displayed */
 		SLATE_ATTRIBUTE( bool, ShouldAllowToolTip )
 
@@ -411,9 +390,6 @@ public:
 		/** If true, the thumbnail in this item can be edited */
 		SLATE_ATTRIBUTE( bool, ThumbnailEditMode )
 
-		/** The visibility state of the labels used in asset view items */
-		SLATE_ATTRIBUTE( EVisibility, LabelVisibility )
-
 		/** Whether the item is selected in the view */
 		SLATE_ARGUMENT( FIsSelected, IsSelected )
 
@@ -479,9 +455,6 @@ public:
 		/** Called when any asset item is destroyed. Used in thumbnail management, though it may be used for more so It is in column items for consistency. */
 		SLATE_EVENT( FOnItemDestroyed, OnItemDestroyed )
 
-		/** The tooltip to display when the cursor hovers over this widget */
-		SLATE_EVENT( FConstructToolTipForAsset, ConstructToolTip )
-
 		/** The string in the title to highlight (used when searching by string) */
 		SLATE_ATTRIBUTE( FText, HighlightText )
 
diff --git a/Engine/Source/Editor/ContentBrowser/Private/ContentBrowserPCH.h b/Engine/Source/Editor/ContentBrowser/Private/ContentBrowserPCH.h
--- a/Engine/Source/Editor/ContentBrowser/Private/ContentBrowserPCH.h
+++ b/Engine/Source/Editor/ContentBrowser/Private/ContentBrowserPCH.h
@@ -33,7 +33,6 @@ DECLARE_DELEGATE_OneParam(FOnAssetRenameCommitted, const TArray<FAssetData>& /*A
 #include "CollectionViewUtils.h"
 #include "SCollectionView.h"
 #include "AssetViewSortManager.h"
-#include "OtherDevelopersAssetFilter.h"
 #include "SAssetView.h"
 #include "SAssetPicker.h"
 #include "SPathPicker.h"
diff --git a/Engine/Source/Editor/ContentBrowser/Private/FrontendFilters.cpp b/Engine/Source/Editor/ContentBrowser/Private/FrontendFilters.cpp
--- a/Engine/Source/Editor/ContentBrowser/Private/FrontendFilters.cpp
+++ b/Engine/Source/Editor/ContentBrowser/Private/FrontendFilters.cpp
@@ -6,6 +6,48 @@
 #include "ObjectTools.h"
 
 
+/////////////////////////////////////////
+// FFrontendFilter_Text
+/////////////////////////////////////////
+
+void AssetDataToString(AssetFilterType Asset, OUT TArray< FString >& Array)
+{
+	Array.Add(Asset.GetExportTextName());
+}
+
+FFrontendFilter_Text::FFrontendFilter_Text()
+	: FFrontendFilter(nullptr)
+	, TextFilter( TTextFilter<AssetFilterType>::FItemToStringArray::CreateStatic( &AssetDataToString ) )
+{
+	TextFilter.OnChanged().AddRaw(this, &FFrontendFilter_Text::HandleOnChangedEvent);
+}
+
+FFrontendFilter_Text::~FFrontendFilter_Text()
+{
+	TextFilter.OnChanged().RemoveAll(this);
+}
+
+bool FFrontendFilter_Text::PassesFilter(AssetFilterType InItem) const
+{
+	return TextFilter.PassesFilter(InItem);
+}
+
+FText FFrontendFilter_Text::GetRawFilterText() const
+{
+	return TextFilter.GetRawFilterText();
+}
+
+void FFrontendFilter_Text::SetRawFilterText(const FText& InFilterText)
+{
+	return TextFilter.SetRawFilterText(InFilterText);
+}
+
+void FFrontendFilter_Text::HandleOnChangedEvent()
+{
+	BroadcastChangedEvent();
+}
+
+
 /////////////////////////////////////////
 // FFrontendFilter_CheckedOut
 /////////////////////////////////////////
@@ -51,6 +93,23 @@ void FFrontendFilter_CheckedOut::SourceControlOperationComplete(const FSourceCon
 // FFrontendFilter_Modified
 /////////////////////////////////////////
 
+FFrontendFilter_Modified::FFrontendFilter_Modified(TSharedPtr<FFrontendFilterCategory> InCategory)
+	: FFrontendFilter(InCategory)
+	, bIsCurrentlyActive(false)
+{
+	UPackage::PackageDirtyStateUpdatedEvent.AddRaw(this, &FFrontendFilter_Modified::OnPackageDirtyStateUpdated);
+}
+
+FFrontendFilter_Modified::~FFrontendFilter_Modified()
+{
+	UPackage::PackageDirtyStateUpdatedEvent.RemoveAll(this);
+}
+
+void FFrontendFilter_Modified::ActiveStateChanged(bool bActive)
+{
+	bIsCurrentlyActive = bActive;
+}
+
 bool FFrontendFilter_Modified::PassesFilter( AssetFilterType InItem ) const
 {
 	UPackage* Package = FindPackage(NULL, *InItem.PackageName.ToString());
@@ -63,6 +122,14 @@ bool FFrontendFilter_Modified::PassesFilter( AssetFilterType InItem ) const
 	return false;
 }
 
+void FFrontendFilter_Modified::OnPackageDirtyStateUpdated(UPackage* Package)
+{
+	if (bIsCurrentlyActive)
+	{
+		BroadcastChangedEvent();
+	}
+}
+
 /////////////////////////////////////////
 // FFrontendFilter_ReplicatedBlueprint
 /////////////////////////////////////////
@@ -79,10 +146,11 @@ bool FFrontendFilter_ReplicatedBlueprint::PassesFilter( AssetFilterType InItem )
 
 FFrontendFilter_ShowOtherDevelopers::FFrontendFilter_ShowOtherDevelopers(TSharedPtr<FFrontendFilterCategory> InCategory)
 	: FFrontendFilter(InCategory)
+	, BaseDeveloperPath(FPackageName::FilenameToLongPackageName(FPaths::GameDevelopersDir()))
+	, UserDeveloperPath(FPackageName::FilenameToLongPackageName(FPaths::GameUserDeveloperDir()))
+	, bIsOnlyOneDeveloperPathSelected(false)
+	, bShowOtherDeveloperAssets(false)
 {
-	BaseDeveloperPath = FPackageName::FilenameToLongPackageName(FPaths::GameDevelopersDir());
-	UserDeveloperPath = FPackageName::FilenameToLongPackageName(FPaths::GameUserDeveloperDir());
-	bIsOnlyOneDeveloperPathSelected = false;
 }
 
 void FFrontendFilter_ShowOtherDevelopers::SetCurrentFilter(const FARFilter& InFilter)
@@ -103,19 +171,23 @@ void FFrontendFilter_ShowOtherDevelopers::SetCurrentFilter(const FARFilter& InFi
 
 bool FFrontendFilter_ShowOtherDevelopers::PassesFilter( AssetFilterType InItem ) const
 {
-	// Never hide developer assets when a single developer folder is selected.
-	if ( !bIsOnlyOneDeveloperPathSelected )
+	// Pass all assets if other developer assets are allowed
+	if ( !bShowOtherDeveloperAssets )
 	{
-		// If selecting multiple folders, the Developers folder/parent folder, or "All Assets", hide assets which are found in the development folder unless they are in the current user's folder
-		const FString PackagePath = InItem.PackagePath.ToString() + TEXT("/");
-		const bool bPackageInDeveloperFolder = PackagePath.StartsWith(BaseDeveloperPath) && PackagePath.Len() != BaseDeveloperPath.Len();
-
-		if ( bPackageInDeveloperFolder )
+		// Never hide developer assets when a single developer folder is selected.
+		if ( !bIsOnlyOneDeveloperPathSelected )
 		{
-			const bool bPackageInUserDeveloperFolder = PackagePath.StartsWith(UserDeveloperPath);
-			if ( !bPackageInUserDeveloperFolder )
+			// If selecting multiple folders, the Developers folder/parent folder, or "All Assets", hide assets which are found in the development folder unless they are in the current user's folder
+			const FString PackagePath = InItem.PackagePath.ToString() + TEXT("/");
+			const bool bPackageInDeveloperFolder = PackagePath.StartsWith(BaseDeveloperPath) && PackagePath.Len() != BaseDeveloperPath.Len();
+
+			if ( bPackageInDeveloperFolder )
 			{
-				return false;
+				const bool bPackageInUserDeveloperFolder = PackagePath.StartsWith(UserDeveloperPath);
+				if ( !bPackageInUserDeveloperFolder )
+				{
+					return false;
+				}
 			}
 		}
 	}
@@ -123,6 +195,20 @@ bool FFrontendFilter_ShowOtherDevelopers::PassesFilter( AssetFilterType InItem )
 	return true;
 }
 
+void FFrontendFilter_ShowOtherDevelopers::SetShowOtherDeveloperAssets(bool bValue)
+{
+	if ( bShowOtherDeveloperAssets != bValue )
+	{
+		bShowOtherDeveloperAssets = bValue;
+		BroadcastChangedEvent();
+	}
+}
+
+bool FFrontendFilter_ShowOtherDevelopers::GetShowOtherDeveloperAssets() const
+{
+	return bShowOtherDeveloperAssets;
+}
+
 /////////////////////////////////////////
 // FFrontendFilter_ShowRedirectors
 /////////////////////////////////////////
diff --git a/Engine/Source/Editor/ContentBrowser/Private/FrontendFilters.h b/Engine/Source/Editor/ContentBrowser/Private/FrontendFilters.h
--- a/Engine/Source/Editor/ContentBrowser/Private/FrontendFilters.h
+++ b/Engine/Source/Editor/ContentBrowser/Private/FrontendFilters.h
@@ -6,6 +6,36 @@
 
 #define LOCTEXT_NAMESPACE "ContentBrowser"
 
+/** A filter for text search */
+class FFrontendFilter_Text : public FFrontendFilter
+{
+public:
+	FFrontendFilter_Text();
+	~FFrontendFilter_Text();
+
+	// FFrontendFilter implementation
+	virtual FString GetName() const override { return TEXT("TextFilter"); }
+	virtual FText GetDisplayName() const override { return LOCTEXT("FrontendFilter_Text", "Text"); }
+	virtual FText GetToolTipText() const override { return LOCTEXT("FrontendFilter_TextTooltip", "Show only assets that match the input text"); }
+
+	// IFilter implementation
+	virtual bool PassesFilter(AssetFilterType InItem) const override;
+
+public:
+	/** Returns the unsanitized and unsplit filter terms */
+	FText GetRawFilterText() const;
+
+	/** Set the Text to be used as the Filter's restrictions */
+	void SetRawFilterText(const FText& InFilterText);
+
+private:
+
+	/** Handler for the internal text filter */
+	void HandleOnChangedEvent();
+
+	TTextFilter<AssetFilterType> TextFilter;
+};
+
 /** A filter that displays only checked out assets */
 class FFrontendFilter_CheckedOut : public FFrontendFilter, public TSharedFromThis<FFrontendFilter_CheckedOut>
 {
@@ -34,15 +64,24 @@ private:
 class FFrontendFilter_Modified : public FFrontendFilter
 {
 public:
-	FFrontendFilter_Modified(TSharedPtr<FFrontendFilterCategory> InCategory) : FFrontendFilter(InCategory) {}
+	FFrontendFilter_Modified(TSharedPtr<FFrontendFilterCategory> InCategory);
+	~FFrontendFilter_Modified();
 
 	// FFrontendFilter implementation
 	virtual FString GetName() const override { return TEXT("Modified"); }
 	virtual FText GetDisplayName() const override { return LOCTEXT("FrontendFilter_Modified", "Modified"); }
 	virtual FText GetToolTipText() const override { return LOCTEXT("FrontendFilter_ModifiedTooltip", "Show only assets that have been modified and not yet saved."); }
+	virtual void ActiveStateChanged(bool bActive) override;
 
 	// IFilter implementation
 	virtual bool PassesFilter( AssetFilterType InItem ) const override;
+
+private:
+
+	/** Handler for when a package's dirty state has changed */
+	void OnPackageDirtyStateUpdated(UPackage* Package);
+
+	bool bIsCurrentlyActive;
 };
 
 /** A filter that displays blueprints that have replicated properties */
@@ -77,10 +116,18 @@ public:
 	// IFilter implementation
 	virtual bool PassesFilter( AssetFilterType InItem ) const override;
 
+public:
+	/** Sets if we should filter out assets from other developers */
+	void SetShowOtherDeveloperAssets(bool bValue);
+
+	/** Gets if we should filter out assets from other developers */
+	bool GetShowOtherDeveloperAssets() const;
+
 private:
 	FString BaseDeveloperPath;
 	FString UserDeveloperPath;
 	bool bIsOnlyOneDeveloperPathSelected;
+	bool bShowOtherDeveloperAssets;
 };
 
 /** An inverse filter that allows display of object redirectors */
diff --git a/Engine/Source/Editor/ContentBrowser/Private/OtherDevelopersAssetFilter.cpp b/Engine/Source/Editor/ContentBrowser/Private/OtherDevelopersAssetFilter.cpp
deleted file mode 100644
diff --git a/Engine/Source/Editor/ContentBrowser/Private/OtherDevelopersAssetFilter.h b/Engine/Source/Editor/ContentBrowser/Private/OtherDevelopersAssetFilter.h
deleted file mode 100644
diff --git a/Engine/Source/Editor/ContentBrowser/Private/SAssetPicker.cpp b/Engine/Source/Editor/ContentBrowser/Private/SAssetPicker.cpp
--- a/Engine/Source/Editor/ContentBrowser/Private/SAssetPicker.cpp
+++ b/Engine/Source/Editor/ContentBrowser/Private/SAssetPicker.cpp
@@ -4,48 +4,27 @@
 
 #define LOCTEXT_NAMESPACE "ContentBrowser"
 
-void AssetToString( AssetFilterType Asset, OUT TArray< FString >& Array )
-{
-	Array.Add( Asset.GetExportTextName() );
-}
-
 void SAssetPicker::Construct( const FArguments& InArgs )
 {
-	FilterCollection = MakeShareable( new AssetFilterCollectionType() );
-	
-	for (int Index = 0; Index < InArgs._AssetPickerConfig.ExtraFilters->Num(); Index++)
-	{
-		FilterCollection->Add( InArgs._AssetPickerConfig.ExtraFilters->GetFilterAtIndex( Index ) );
-	}
+	TSharedPtr<AssetFilterCollectionType> FrontendFilters = MakeShareable(new AssetFilterCollectionType());
 
 	BindCommands();
 
 	OnAssetsActivated = InArgs._AssetPickerConfig.OnAssetsActivated;
 	OnAssetSelected = InArgs._AssetPickerConfig.OnAssetSelected;
-	OnAssetClicked = InArgs._AssetPickerConfig.OnAssetClicked;
 	OnAssetDoubleClicked = InArgs._AssetPickerConfig.OnAssetDoubleClicked;
 	OnAssetEnterPressed = InArgs._AssetPickerConfig.OnAssetEnterPressed;
 	bPendingFocusNextFrame = InArgs._AssetPickerConfig.bFocusSearchBoxWhenOpened;
 	DefaultFilterMenuExpansion = InArgs._AssetPickerConfig.DefaultFilterMenuExpansion;
 
-	if ( InArgs._AssetPickerConfig.ClearSelectionDelegate != NULL )
+	for (auto DelegateIt = InArgs._AssetPickerConfig.GetCurrentSelectionDelegates.CreateConstIterator(); DelegateIt; ++DelegateIt)
 	{
-		(*InArgs._AssetPickerConfig.ClearSelectionDelegate) = FClearSelectionDelegate::CreateSP( this, &SAssetPicker::ClearSelection ) ;
-	}
-
-	for ( auto AssetIt = InArgs._AssetPickerConfig.GetCurrentSelectionDelegates.CreateConstIterator(); AssetIt; ++AssetIt )
-	{
-		if ( (*AssetIt) != NULL )
+		if ((*DelegateIt) != NULL)
 		{
-			(**AssetIt) = FGetCurrentSelectionDelegate::CreateSP( this, &SAssetPicker::GetCurrentSelection );
+			(**DelegateIt) = FGetCurrentSelectionDelegate::CreateSP(this, &SAssetPicker::GetCurrentSelection);
 		}
 	}
 
-	if ( InArgs._AssetPickerConfig.AdjustSelectionDelegate != NULL )
-	{
-		(*InArgs._AssetPickerConfig.AdjustSelectionDelegate) = FAdjustSelectionDelegate::CreateSP( this, &SAssetPicker::AdjustSelection ) ;
-	}
-
 	TSharedRef<SVerticalBox> VerticalBox = SNew(SVerticalBox);
 
 	ChildSlot
@@ -53,19 +32,18 @@ void SAssetPicker::Construct( const FArguments& InArgs )
 		VerticalBox
 	];
 
-	TAttribute< FText > HighlightText = InArgs._AssetPickerConfig.HighlightedText;
-	TAttribute< EVisibility > LabelVisibility = InArgs._AssetPickerConfig.LabelVisibility;
+	TAttribute< FText > HighlightText;
 	EThumbnailLabel::Type ThumbnailLabel = InArgs._AssetPickerConfig.ThumbnailLabel;
 
 	// Search box
 	if (!InArgs._AssetPickerConfig.bAutohideSearchBar)
 	{
-		TextFilter = MakeShareable( new TTextFilter< AssetFilterType >( TTextFilter< AssetFilterType >::FItemToStringArray::CreateStatic( &AssetToString ) ) );
-		FilterCollection->Add( TextFilter );
+		TextFilter = MakeShareable( new FFrontendFilter_Text() );
+		FrontendFilters->Add( TextFilter );
 		HighlightText = TAttribute< FText >( this, &SAssetPicker::GetHighlightedText );
 
-		OtherDevelopersFilter = MakeShareable(new FOtherDevelopersAssetFilter());
-		FilterCollection->Add( OtherDevelopersFilter );
+		OtherDevelopersFilter = MakeShareable( new FFrontendFilter_ShowOtherDevelopers(nullptr) );
+		FrontendFilters->Add( OtherDevelopersFilter );
 
 		TSharedRef<SHorizontalBox> HorizontalBox = SNew(SHorizontalBox);
 
@@ -166,8 +144,6 @@ void SAssetPicker::Construct( const FArguments& InArgs )
 	InitialBackendFilter = InArgs._AssetPickerConfig.Filter;
 	InitialBackendFilter.PackagePaths.Empty();
 
-	TSharedPtr<AssetFilterCollectionType> FrontendFilters = MakeShareable(new AssetFilterCollectionType());
-
 	if(InArgs._AssetPickerConfig.bAddFilterUI)
 	{
 		// Filters
@@ -199,7 +175,6 @@ void SAssetPicker::Construct( const FArguments& InArgs )
 		SAssignNew(AssetViewPtr, SAssetView)
 		.SelectionMode( InArgs._AssetPickerConfig.SelectionMode )
 		.OnShouldFilterAsset(InArgs._AssetPickerConfig.OnShouldFilterAsset)
-		.OnAssetClicked(InArgs._AssetPickerConfig.OnAssetClicked)
 		.OnAssetSelected(InArgs._AssetPickerConfig.OnAssetSelected)
 		.OnAssetsActivated(this, &SAssetPicker::HandleAssetsActivated)
 		.OnGetAssetContextMenu(InArgs._AssetPickerConfig.OnGetAssetContextMenu)
@@ -210,27 +185,22 @@ void SAssetPicker::Construct( const FArguments& InArgs )
 		.InitialViewType(InArgs._AssetPickerConfig.InitialAssetViewType)
 		.InitialAssetSelection(InArgs._AssetPickerConfig.InitialAssetSelection)
 		.ThumbnailScale(InArgs._AssetPickerConfig.ThumbnailScale)
-		.OnThumbnailScaleChanged(InArgs._AssetPickerConfig.ThumbnailScaleChanged)
 		.ShowBottomToolbar(InArgs._AssetPickerConfig.bShowBottomToolbar)
 		.OnAssetTagWantsToBeDisplayed(InArgs._AssetPickerConfig.OnAssetTagWantsToBeDisplayed)
-		.OnAssetDragged( InArgs._AssetPickerConfig.OnAssetDragged )
 		.AllowDragging( InArgs._AssetPickerConfig.bAllowDragging )
 		.CanShowClasses( InArgs._AssetPickerConfig.bCanShowClasses )
-		.CanShowFolders( InArgs._AssetPickerConfig.bCanShowFolders )
+		.CanShowFolders( false )
 		.CanShowOnlyAssetsInSelectedFolders( InArgs._AssetPickerConfig.bCanShowOnlyAssetsInSelectedFolders )
 		.CanShowRealTimeThumbnails( InArgs._AssetPickerConfig.bCanShowRealTimeThumbnails )
 		.CanShowDevelopersFolder( InArgs._AssetPickerConfig.bCanShowDevelopersFolder )
 		.PreloadAssetsForContextMenu( InArgs._AssetPickerConfig.bPreloadAssetsForContextMenu )
-		.DynamicFilters( FilterCollection )
 		.HighlightedText( HighlightText )
-		.LabelVisibility( LabelVisibility )
 		.ThumbnailLabel( ThumbnailLabel )
-		.ConstructToolTipForAsset( InArgs._AssetPickerConfig.ConstructToolTipForAsset )
 		.AssetShowWarningText( InArgs._AssetPickerConfig.AssetShowWarningText)
 		.AllowFocusOnSync(false)	// Stop the asset view from stealing focus (we're in control of that)
 	];
 
-	AssetViewPtr->RequestListRefresh();
+	AssetViewPtr->RequestSlowFullListRefresh();
 }
 
 void SAssetPicker::Tick(const FGeometry& AllottedGeometry, const double InCurrentTime, const float InDeltaTime)
@@ -343,21 +313,11 @@ void SAssetPicker::HandleAssetsActivated(const TArray<FAssetData>& ActivatedAsse
 	OnAssetsActivated.ExecuteIfBound( ActivatedAssets, ActivationMethod );
 }
 
-void SAssetPicker::ClearSelection()
-{
-	AssetViewPtr->ClearSelection();
-}
-
 TArray< FAssetData > SAssetPicker::GetCurrentSelection()
 {
 	return AssetViewPtr->GetSelectedAssets();
 }
 
-void SAssetPicker::AdjustSelection(const int32 direction)
-{
-	AssetViewPtr->AdjustActiveSelection(direction);
-}
-
 FText SAssetPicker::GetShowOtherDevelopersToolTip() const
 {
 	if (OtherDevelopersFilter->GetShowOtherDeveloperAssets())
diff --git a/Engine/Source/Editor/ContentBrowser/Private/SAssetPicker.h b/Engine/Source/Editor/ContentBrowser/Private/SAssetPicker.h
--- a/Engine/Source/Editor/ContentBrowser/Private/SAssetPicker.h
+++ b/Engine/Source/Editor/ContentBrowser/Private/SAssetPicker.h
@@ -42,15 +42,9 @@ private:
 	/** Handler for when the user double clicks, presses enter, or presses space on an asset */
 	void HandleAssetsActivated(const TArray<FAssetData>& ActivatedAssets, EAssetTypeActivationMethod::Type ActivationMethod);
 
-	/** Clears the current asset selection in the view */
-	void ClearSelection();
-
 	/** @return The currently selected asset */
 	TArray< FAssetData > GetCurrentSelection();
 
-	/** Adjust the asset data selection, +1 to increment or -1 to decrement */
-	void AdjustSelection(const int32 direction);
-
 	/** @return The text to highlight on the assets  */
 	FText GetHighlightedText() const;
 
@@ -85,9 +79,6 @@ private:
 	/** Called to when an asset is selected or the none button is pressed */
 	FOnAssetSelected OnAssetSelected;
 
-	/** Called to when an asset is clicked */
-	FOnAssetClicked OnAssetClicked;
-
 	/** Called to when an asset is double clicked */
 	FOnAssetDoubleClicked OnAssetDoubleClicked;
 
@@ -102,8 +93,8 @@ private:
 
 	/** Filters needed for filtering the assets */
 	TSharedPtr< AssetFilterCollectionType > FilterCollection;
-	TSharedPtr< TTextFilter< AssetFilterType > > TextFilter;
-	TSharedPtr< FOtherDevelopersAssetFilter > OtherDevelopersFilter;
+	TSharedPtr< FFrontendFilter_Text > TextFilter;
+	TSharedPtr< FFrontendFilter_ShowOtherDevelopers > OtherDevelopersFilter;
 
 	EAssetTypeCategories::Type DefaultFilterMenuExpansion;
 
diff --git a/Engine/Source/Editor/ContentBrowser/Private/SAssetView.cpp b/Engine/Source/Editor/ContentBrowser/Private/SAssetView.cpp
--- a/Engine/Source/Editor/ContentBrowser/Private/SAssetView.cpp
+++ b/Engine/Source/Editor/ContentBrowser/Private/SAssetView.cpp
@@ -83,9 +83,6 @@ void SAssetView::Construct( const FArguments& InArgs )
 	FCoreDelegates::OnAssetLoaded.AddSP(this, &SAssetView::OnAssetLoaded);
 	FCoreDelegates::OnObjectPropertyChanged.AddSP(this, &SAssetView::OnObjectPropertyChanged);
 
-	// Listen for when packages' dirty states are updated
-	UPackage::PackageDirtyStateUpdatedEvent.AddSP(this, &SAssetView::OnPackageDirtyStateUpdated);
-
 	// Listen for when view settings are changed
 	UContentBrowserSettings::OnSettingChanged().AddSP(this, &SAssetView::HandleSettingChanged);
 
@@ -110,7 +107,6 @@ void SAssetView::Construct( const FArguments& InArgs )
 	TileViewThumbnailPadding = 5;
 	TileViewNameHeight = 36;
 	ThumbnailScaleSliderValue = InArgs._ThumbnailScale; 
-	ThumbnailScaleChanged = InArgs._OnThumbnailScaleChanged;
 
 	if ( !ThumbnailScaleSliderValue.IsBound() )
 	{
@@ -139,11 +135,6 @@ void SAssetView::Construct( const FArguments& InArgs )
 
 	SourcesData = InArgs._InitialSourcesData;
 	BackendFilter = InArgs._InitialBackendFilter;
-	DynamicFilters = InArgs._DynamicFilters;
-	if ( DynamicFilters.IsValid() )
-	{
-		DynamicFilters->OnChanged().AddSP( this, &SAssetView::OnDynamicFiltersChanged );
-	}
 
 	FrontendFilters = InArgs._FrontendFilters;
 	if ( FrontendFilters.IsValid() )
@@ -152,7 +143,6 @@ void SAssetView::Construct( const FArguments& InArgs )
 	}
 
 	OnShouldFilterAsset = InArgs._OnShouldFilterAsset;
-	OnAssetClicked = InArgs._OnAssetClicked;
 	OnAssetSelected = InArgs._OnAssetSelected;
 	OnAssetsActivated = InArgs._OnAssetsActivated;
 	OnGetAssetContextMenu = InArgs._OnGetAssetContextMenu;
@@ -161,12 +151,9 @@ void SAssetView::Construct( const FArguments& InArgs )
 	OnFindInAssetTreeRequested = InArgs._OnFindInAssetTreeRequested;
 	OnAssetRenameCommitted = InArgs._OnAssetRenameCommitted;
 	OnAssetTagWantsToBeDisplayed = InArgs._OnAssetTagWantsToBeDisplayed;
-	OnAssetDragged = InArgs._OnAssetDragged;
 	HighlightedText = InArgs._HighlightedText;
-	LabelVisibility = InArgs._LabelVisibility;
 	ThumbnailLabel = InArgs._ThumbnailLabel;
 	AllowThumbnailHintLabel = InArgs._AllowThumbnailHintLabel;
-	ConstructToolTipForAsset = InArgs._ConstructToolTipForAsset;
 	AssetShowWarningText = InArgs._AssetShowWarningText;
 	bAllowDragging = InArgs._AllowDragging;
 	bAllowFocusOnSync = InArgs._AllowFocusOnSync;
@@ -391,7 +378,7 @@ void SAssetView::SetSourcesData(const FSourcesData& InSourcesData)
 {
 	// Update the path and collection lists
 	SourcesData = InSourcesData;
-	bRefreshSourceItemsRequested = true;
+	RequestSlowFullListRefresh();
 	ClearSelection();
 }
 
@@ -409,7 +396,7 @@ void SAssetView::SetBackendFilter(const FARFilter& InBackendFilter)
 {
 	// Update the path and collection lists
 	BackendFilter = InBackendFilter;
-	bRefreshSourceItemsRequested = true;
+	RequestSlowFullListRefresh();
 }
 
 void SAssetView::OnCreateNewFolder(const FString& FolderName, const FString& FolderPath)
@@ -634,9 +621,14 @@ TArray<FString> SAssetView::GetSelectedFolders() const
 	return SelectedFolders;
 }
 
-void SAssetView::RequestListRefresh()
+void SAssetView::RequestSlowFullListRefresh()
+{
+	bSlowFullListRefreshRequested = true;
+}
+
+void SAssetView::RequestQuickFrontendListRefresh()
 {
-	bRefreshSourceItemsRequested = true;
+	bQuickFrontendListRefreshRequested = true;
 }
 
 void SAssetView::SaveSettings(const FString& IniFilename, const FString& IniSection, const FString& SettingsString) const
@@ -797,15 +789,22 @@ void SAssetView::Tick(const FGeometry& AllottedGeometry, const double InCurrentT
 		bPendingUpdateThumbnails = false;
 	}
 
-	if ( bRefreshSourceItemsRequested )
+	if ( bSlowFullListRefreshRequested || bQuickFrontendListRefreshRequested )
 	{
 		ResetQuickJump();
-		RefreshSourceItems();
+		
+		if ( bSlowFullListRefreshRequested )
+		{
+			RefreshSourceItems();
+		}
+
 		RefreshFilteredItems();
 		RefreshFolders();
 		// Don't sync to selection if we are just going to do it below
 		SortList(!PendingSyncAssets.Num());
-		bRefreshSourceItemsRequested = false;
+
+		bSlowFullListRefreshRequested = false;
+		bQuickFrontendListRefreshRequested = false;
 	}
 
 	if ( QueriedAssetItems.Num() > 0 )
@@ -1360,12 +1359,12 @@ void SAssetView::RefreshSourceItems()
 		if ( Item.AssetClass == UObjectRedirector::StaticClass()->GetFName() && !Item.IsUAsset() )
 		{
 			// Do not show redirectors if they are not the main asset in the uasset file.
-			Items.RemoveAt(AssetIdx);
+			Items.RemoveAtSwap(AssetIdx);
 		}
 		else if ( !bDisplayEngine && ContentBrowserUtils::IsEngineFolder(Item.PackagePath.ToString()) )
 		{
 			// If this is an engine folder, and we don't want to show them, remove
-			Items.RemoveAt(AssetIdx);
+			Items.RemoveAtSwap(AssetIdx);
 		}
 	}
 }
@@ -1585,6 +1584,7 @@ void SAssetView::RefreshFolders()
 				if(!Folders.Contains(*SubPathIt))
 				{
 					FilteredAssetItems.Add(MakeShareable(new FAssetViewFolder(*SubPathIt)));
+					RefreshList();
 					Folders.Add(*SubPathIt);
 					bPendingSortFilteredItems = true;
 				}
@@ -1736,7 +1736,6 @@ void SAssetView::ProcessRecentlyAddedAssets()
 						{
 							FilteredAssetItems.Add(MakeShareable(new FAssetViewAsset(AssetItems[AddedAssetIdx])));
 							bPendingSortFilteredItems = true;
-							bRefreshSourceItemsRequested = true;
 
 							RefreshList();
 						}
@@ -1813,6 +1812,7 @@ void SAssetView::OnAssetRegistryPathAdded(const FString& Path)
 						if(!Folders.Contains(NewSubFolder))
 						{
 							FilteredAssetItems.Add(MakeShareable(new FAssetViewFolder(NewSubFolder)));
+							RefreshList();
 							Folders.Add(NewSubFolder);
 							bPendingSortFilteredItems = true;
 						}
@@ -1927,34 +1927,20 @@ void SAssetView::OnObjectPropertyChanged(UObject* Asset, FPropertyChangedEvent&
 	}
 }
 
-void SAssetView::OnPackageDirtyStateUpdated(UPackage* Package)
-{
-	RequestListRefresh();
-}
-
-void SAssetView::OnDynamicFiltersChanged()
-{
-	ResetQuickJump();
-	RefreshFilteredItems();
-	RefreshFolders();
-	SortList();
-}
-
 void SAssetView::OnFrontendFiltersChanged()
 {
-	bRefreshSourceItemsRequested = true;
+	RequestQuickFrontendListRefresh();
 }
 
 bool SAssetView::IsFrontendFilterActive() const
 {
-	return ( FrontendFilters.IsValid() && FrontendFilters->Num() > 0 ) || ( DynamicFilters.IsValid() && DynamicFilters->Num() > 0 );
+	return ( FrontendFilters.IsValid() && FrontendFilters->Num() > 0 );
 }
 
 bool SAssetView::PassesCurrentFrontendFilter(const FAssetData& Item) const
 {
 	// Check the frontend filters list
-	if ( ( FrontendFilters.IsValid() && !FrontendFilters->PassesAllFilters(Item) ) ||
-		 ( DynamicFilters.IsValid() && !DynamicFilters->PassesAllFilters( Item ) ) )
+	if ( FrontendFilters.IsValid() && !FrontendFilters->PassesAllFilters(Item) )
 	{
 		return false;
 	}
@@ -1996,7 +1982,7 @@ void SAssetView::RunAssetsThroughBackendFilter(TArray<FAssetData>& InOutAssetDat
 
 				if ( !CollectionObjectPaths.Contains( AssetData.ObjectPath ) )
 				{
-					InOutAssetDataList.RemoveAt(AssetDataIdx);
+					InOutAssetDataList.RemoveAtSwap(AssetDataIdx);
 				}
 			}
 		}
@@ -2225,7 +2211,7 @@ void SAssetView::ToggleShowOnlyAssetsInSelectedFolders()
 {
 	check( CanShowOnlyAssetsInSelectedFolders() );
 	GetMutableDefault<UContentBrowserSettings>()->ShowOnlyAssetsInSelectedFolders = !GetDefault<UContentBrowserSettings>()->ShowOnlyAssetsInSelectedFolders;
-	bRefreshSourceItemsRequested = true;
+	RequestSlowFullListRefresh();
 }
 
 bool SAssetView::CanShowOnlyAssetsInSelectedFolders() const
@@ -2242,7 +2228,6 @@ void SAssetView::ToggleRealTimeThumbnails()
 {
 	check( CanShowRealTimeThumbnails() );
 	GetMutableDefault<UContentBrowserSettings>()->RealTimeThumbnails = !GetDefault<UContentBrowserSettings>()->RealTimeThumbnails;
-	bRefreshSourceItemsRequested = true;
 }
 
 bool SAssetView::CanShowRealTimeThumbnails() const
@@ -2495,7 +2480,6 @@ TSharedRef<ITableRow> SAssetView::MakeListViewWidget(TSharedPtr<FAssetViewItem>
 			.OnItemDestroyed(this, &SAssetView::AssetItemWidgetDestroyed)
 			.ShouldAllowToolTip(this, &SAssetView::ShouldAllowToolTips)
 			.HighlightText(HighlightedText)
-			.ConstructToolTip( ConstructToolTipForAsset )
 			.IsSelected( FIsSelected::CreateSP(TableRowWidget.Get(), &STableRow<TSharedPtr<FAssetViewItem>>::IsSelectedExclusively) );
 
 		TableRowWidget->SetContent(Item);
@@ -2539,7 +2523,6 @@ TSharedRef<ITableRow> SAssetView::MakeListViewWidget(TSharedPtr<FAssetViewItem>
 			.ShouldAllowToolTip(this, &SAssetView::ShouldAllowToolTips)
 			.HighlightText(HighlightedText)
 			.ThumbnailEditMode(this, &SAssetView::IsThumbnailEditMode)
-			.ConstructToolTip( ConstructToolTipForAsset )
 			.ThumbnailLabel( ThumbnailLabel )
 			.ThumbnailHintColorAndOpacity( this, &SAssetView::GetThumbnailHintColorAndOpacity )
 			.AllowThumbnailHintLabel( AllowThumbnailHintLabel )
@@ -2579,8 +2562,6 @@ TSharedRef<ITableRow> SAssetView::MakeTileViewWidget(TSharedPtr<FAssetViewItem>
 			.OnItemDestroyed(this, &SAssetView::AssetItemWidgetDestroyed)
 			.ShouldAllowToolTip(this, &SAssetView::ShouldAllowToolTips)
 			.HighlightText( HighlightedText )
-			.LabelVisibility( LabelVisibility )
-			.ConstructToolTip( ConstructToolTipForAsset )
 			.IsSelected( FIsSelected::CreateSP(TableRowWidget.Get(), &STableRow<TSharedPtr<FAssetViewItem>>::IsSelectedExclusively) )
 			.OnAssetsDragDropped(this, &SAssetView::OnAssetsDragDropped)
 			.OnPathsDragDropped(this, &SAssetView::OnPathsDragDropped)
@@ -2627,10 +2608,8 @@ TSharedRef<ITableRow> SAssetView::MakeTileViewWidget(TSharedPtr<FAssetViewItem>
 			.ShouldAllowToolTip(this, &SAssetView::ShouldAllowToolTips)
 			.HighlightText( HighlightedText )
 			.ThumbnailEditMode(this, &SAssetView::IsThumbnailEditMode)
-			.LabelVisibility( LabelVisibility )
 			.ThumbnailLabel( ThumbnailLabel )
 			.ThumbnailHintColorAndOpacity( this, &SAssetView::GetThumbnailHintColorAndOpacity )
-			.ConstructToolTip( ConstructToolTipForAsset )
 			.AllowThumbnailHintLabel( AllowThumbnailHintLabel )
 			.IsSelected( FIsSelected::CreateSP(TableRowWidget.Get(), &STableRow<TSharedPtr<FAssetViewItem>>::IsSelectedExclusively) );
 
@@ -2660,7 +2639,6 @@ TSharedRef<ITableRow> SAssetView::MakeColumnViewWidget(TSharedPtr<FAssetViewItem
 				.OnVerifyRenameCommit(this, &SAssetView::AssetVerifyRenameCommit)
 				.OnItemDestroyed(this, &SAssetView::AssetItemWidgetDestroyed)
 				.HighlightText( HighlightedText )
-				.ConstructToolTip( ConstructToolTipForAsset )
 				.OnAssetsDragDropped(this, &SAssetView::OnAssetsDragDropped)
 				.OnPathsDragDropped(this, &SAssetView::OnPathsDragDropped)
 				.OnFilesDragDropped(this, &SAssetView::OnFilesDragDropped)
@@ -3058,18 +3036,7 @@ FReply SAssetView::OnDraggingAssetItem( const FGeometry& MyGeometry, const FPoin
 			
 			if ( InAssetData.Num() > 0 )
 			{
-				FReply Reply = FReply::Unhandled();
-				if ( OnAssetDragged.IsBound() )
-				{
-					Reply = OnAssetDragged.Execute( InAssetData );
-				}
-
-				if ( !Reply.IsEventHandled() )
-				{
-					Reply = FReply::Handled().BeginDragDrop(FAssetDragDropOp::New(InAssetData));
-				}
-
-				return Reply;
+				return FReply::Handled().BeginDragDrop(FAssetDragDropOp::New(InAssetData));
 			}
 		}
 		else
@@ -3254,7 +3221,7 @@ void SAssetView::AssetRenameCommit(const TSharedPtr<FAssetViewItem>& Item, const
 		{
 			// Clearing the rename box on a newly created asset cancels the entire creation process
 			FilteredAssetItems.Remove(Item);
-			bRefreshSourceItemsRequested = true;
+			RefreshList();
 		}
 		else
 		{
@@ -3280,7 +3247,7 @@ void SAssetView::AssetRenameCommit(const TSharedPtr<FAssetViewItem>& Item, const
 
 			// remove this temp item - a new one will have been added by the asset registry callback
 			FilteredAssetItems.Remove(Item);
-			bRefreshSourceItemsRequested = true;
+			RefreshList();
 
 			if(!bSuccess)
 			{
@@ -3320,7 +3287,7 @@ void SAssetView::AssetRenameCommit(const TSharedPtr<FAssetViewItem>& Item, const
 				}
 			}
 
-			bRefreshSourceItemsRequested = true;
+			RequestQuickFrontendListRefresh();
 		}		
 	}
 	else
@@ -3335,7 +3302,7 @@ void SAssetView::AssetRenameCommit(const TSharedPtr<FAssetViewItem>& Item, const
 		{
 			// Sort in the new item
 			bPendingSortFilteredItems = true;
-			bRefreshSourceItemsRequested = true;
+			RequestQuickFrontendListRefresh();
 
 			// Refresh the thumbnail
 			const TSharedPtr<FAssetThumbnail>* AssetThumbnail = RelevantThumbnails.Find(StaticCastSharedPtr<FAssetViewAsset>(Item));
@@ -3480,15 +3447,7 @@ float SAssetView::GetThumbnailScale() const
 
 void SAssetView::SetThumbnailScale( float NewValue )
 {
-	if ( ThumbnailScaleSliderValue.IsBound() )
-	{
-		ThumbnailScaleChanged.ExecuteIfBound( NewValue );
-	}
-	else
-	{
-		ThumbnailScaleSliderValue = NewValue;
-	}
-
+	ThumbnailScaleSliderValue = NewValue;
 	RefreshList();
 }
 
@@ -3504,14 +3463,7 @@ float SAssetView::GetListViewItemHeight() const
 
 float SAssetView::GetTileViewItemHeight() const
 {
-	float Height = GetTileViewItemBaseHeight() * FillScale;
-
-	if ( LabelVisibility.Get() != EVisibility::Collapsed )
-	{
-		Height += TileViewNameHeight;
-	}
-
-	return Height;
+	return TileViewNameHeight + GetTileViewItemBaseHeight() * FillScale;
 }
 
 float SAssetView::GetTileViewItemBaseHeight() const
@@ -3721,7 +3673,7 @@ void SAssetView::SetUserSearching(bool bInSearching)
 {
 	if(bUserSearching != bInSearching)
 	{
-		bRefreshSourceItemsRequested = true;
+		RequestSlowFullListRefresh();
 	}
 	bUserSearching = bInSearching;
 }
@@ -3734,7 +3686,7 @@ void SAssetView::HandleSettingChanged(FName PropertyName)
 		(PropertyName == "DisplayEngineFolder") ||
 		(PropertyName == NAME_None))	// @todo: Needed if PostEditChange was called manually, for now
 	{
-		bRefreshSourceItemsRequested = true;
+		RequestSlowFullListRefresh();
 	}
 }
 
diff --git a/Engine/Source/Editor/ContentBrowser/Private/SAssetView.h b/Engine/Source/Editor/ContentBrowser/Private/SAssetView.h
--- a/Engine/Source/Editor/ContentBrowser/Private/SAssetView.h
+++ b/Engine/Source/Editor/ContentBrowser/Private/SAssetView.h
@@ -15,7 +15,6 @@ class SAssetView : public SCompoundWidget
 public:
 	SLATE_BEGIN_ARGS( SAssetView )
 		: _AreRealTimeThumbnailsAllowed(true)
-		, _LabelVisibility(EVisibility::Visible)
 		, _ThumbnailLabel( EThumbnailLabel::ClassName )
 		, _AllowThumbnailHintLabel(true)
 		, _InitialViewType(EAssetViewType::Tile)
@@ -37,15 +36,9 @@ public:
 		/** Called to check if an asset should be filtered out by external code */
 		SLATE_EVENT( FOnShouldFilterAsset, OnShouldFilterAsset )
 
-		/** Called to when an asset is clicked */
-		SLATE_EVENT( FOnAssetClicked, OnAssetClicked )
-
 		/** Called to when an asset is selected */
 		SLATE_EVENT( FOnAssetSelected, OnAssetSelected )
 
-		/** Called when an asset has begun being dragged by the user */
-		SLATE_EVENT( FOnAssetDragged, OnAssetDragged )
-
 		/** Called when the user double clicks, presses enter, or presses space on an asset */
 		SLATE_EVENT( FOnAssetsActivated, OnAssetsActivated )
 
@@ -64,9 +57,6 @@ public:
 		/** Called when the user has committed a rename of one or more assets */
 		SLATE_EVENT( FOnAssetRenameCommitted, OnAssetRenameCommitted )
 
-		/** The tooltip to display when the cursor hovers over an asset */
-		SLATE_EVENT( FConstructToolTipForAsset, ConstructToolTipForAsset )
-
 		/** The warning text to display when there are no assets to show */
 		SLATE_ATTRIBUTE( FText, AssetShowWarningText )
 
@@ -76,9 +66,6 @@ public:
 		/** Attribute to determine what text should be highlighted */
 		SLATE_ATTRIBUTE( FText, HighlightedText )
 
-		/** Whether to display the name of the asset as a label below the thumbnail */
-		SLATE_ATTRIBUTE( EVisibility, LabelVisibility )
-
 		/** What the label on the asset thumbnails should be */
 		SLATE_ARGUMENT( EThumbnailLabel::Type, ThumbnailLabel )
 
@@ -88,9 +75,6 @@ public:
 		/** The filter collection used to further filter down assets returned from the backend */
 		SLATE_ARGUMENT( TSharedPtr<AssetFilterCollectionType>, FrontendFilters )
 
-		/** The filter collection used to further filter down assets that have passed both the backend and frontend filters */
-		SLATE_ARGUMENT( TSharedPtr<AssetFilterCollectionType>, DynamicFilters )
-
 		/** The initial base sources filter */
 		SLATE_ARGUMENT( FSourcesData, InitialSourcesData )
 
@@ -106,9 +90,6 @@ public:
 		/** The thumbnail scale. [0-1] where 0.5 is no scale */
 		SLATE_ATTRIBUTE( float, ThumbnailScale )
 
-		/** Called when the Thumbnail scale changes and is currently bound to a delegate */
-		SLATE_EVENT( FOnThumbnailScaleChanged, OnThumbnailScaleChanged )
-
 		/** Should the toolbar indicating number of selected assets, mode switch buttons, etc... be shown? */
 		SLATE_ARGUMENT( bool, ShowBottomToolbar )
 
@@ -197,8 +178,11 @@ public:
 	/** Returns all the folders currently selected in the view */
 	TArray<FString> GetSelectedFolders() const;
 
-	/** Requests that the asset view refreshes its visible items. */
-	void RequestListRefresh();
+	/** Requests that the asset view refreshes all it's source items. This is slow and should only be used if the source items change. */
+	void RequestSlowFullListRefresh();
+
+	/** Requests that the asset view refreshes only items that are filtered through frontend sources. This should be used when possible. */
+	void RequestQuickFrontendListRefresh();
 
 	/** Saves any settings to config that should be persistent between editor sessions */
 	void SaveSettings(const FString& IniFilename, const FString& IniSection, const FString& SettingsString) const;
@@ -317,12 +301,6 @@ private:
 	/** Handler for when an asset's property has changed */
 	void OnObjectPropertyChanged(UObject* Asset, FPropertyChangedEvent& PropertyChangedEvent);
 
-	/** Handler for when a package's dirty state has changed */
-	void OnPackageDirtyStateUpdated(UPackage* Package);
-
-	/** Handler for when any dynamic filters have been changed */
-	void OnDynamicFiltersChanged();
-
 	/** Handler for when any frontend filters have been changed */
 	void OnFrontendFiltersChanged();
 
@@ -640,10 +618,12 @@ private:
 	FSourcesData SourcesData;
 	FARFilter BackendFilter;
 	TSharedPtr<AssetFilterCollectionType> FrontendFilters;
-	TSharedPtr<AssetFilterCollectionType> DynamicFilters;
 
-	/** If true, the source items will be refreshed next frame */
-	bool bRefreshSourceItemsRequested;
+	/** If true, the source items will be refreshed next frame. Very slow. */
+	bool bSlowFullListRefreshRequested;
+
+	/** If true, the frontend items will be refreshed next frame. Much faster. */
+	bool bQuickFrontendListRefreshRequested;
 
 	/** The list of assets to sync next frame */
 	TSet<FName> PendingSyncAssets;
@@ -654,9 +634,6 @@ private:
 	/** Called to check if an asset should be filtered out by external code */
 	FOnShouldFilterAsset OnShouldFilterAsset;
 
-	/** Called when an asset was clicked on in the list */
-	FOnAssetClicked OnAssetClicked;
-
 	/** Called when an asset was selected in the list */
 	FOnAssetSelected OnAssetSelected;
 
@@ -680,9 +657,6 @@ private:
 
 	/** Called to check if an asset tag should be display in details view. */
 	FOnShouldDisplayAssetTag OnAssetTagWantsToBeDisplayed;
-
-	/** Called when an asset has begun being dragged by the user */
-	FOnAssetDragged OnAssetDragged;
 	
 	/** When true, filtered list items will be sorted next tick. Provided another sort hasn't happened recently or we are renaming an asset */
 	bool bPendingSortFilteredItems;
@@ -722,7 +696,6 @@ private:
 
 	/** The current value for the thumbnail scale from the thumbnail slider */
 	TAttribute< float > ThumbnailScaleSliderValue;
-	FOnThumbnailScaleChanged ThumbnailScaleChanged;
 
 	/** The max and min thumbnail scales as a fraction of the rendered size */
 	float MinThumbnailScale;
@@ -788,9 +761,6 @@ private:
 	/** The text to highlight on the assets */
 	TAttribute< FText > HighlightedText;
 
-	/** The visibility setting for the label below the thumbnail */
-	TAttribute< EVisibility > LabelVisibility;
-
 	/** What the label on the thumbnails should be */
 	EThumbnailLabel::Type ThumbnailLabel;
 
@@ -803,9 +773,6 @@ private:
 	/** The current thumbnail hint color and opacity*/
 	FLinearColor ThumbnailHintColorAndOpacity;
 
-	/** A callback for external code to construct the tooltip for an asset */
-	FConstructToolTipForAsset ConstructToolTipForAsset;
-
 	/** The text to show when there are no assets to show */
 	TAttribute< FText > AssetShowWarningText;
 
diff --git a/Engine/Source/Editor/ContentBrowser/Private/SContentBrowser.cpp b/Engine/Source/Editor/ContentBrowser/Private/SContentBrowser.cpp
--- a/Engine/Source/Editor/ContentBrowser/Private/SContentBrowser.cpp
+++ b/Engine/Source/Editor/ContentBrowser/Private/SContentBrowser.cpp
@@ -11,11 +11,6 @@
 
 #define LOCTEXT_NAMESPACE "ContentBrowser"
 
-void AssetDataToString( AssetFilterType Asset, OUT TArray< FString >& Array )
-{
-	Array.Add( Asset.GetExportTextName() );
-}
-
 const FString SContentBrowser::SettingsIniSection = TEXT("ContentBrowser");
 
 SContentBrowser::~SContentBrowser()
@@ -27,8 +22,6 @@ SContentBrowser::~SContentBrowser()
 BEGIN_SLATE_FUNCTION_BUILD_OPTIMIZATION
 void SContentBrowser::Construct( const FArguments& InArgs, const FName& InInstanceName )
 {
-	TextFilter = MakeShareable( new TTextFilter< AssetFilterType >( TTextFilter< AssetFilterType >::FItemToStringArray::CreateStatic( &AssetDataToString ) ) );
-
 	if ( InArgs._ContainingTab.IsValid() )
 	{
 		// For content browsers that are placed in tabs, save settings when the tab is closing.
@@ -49,9 +42,8 @@ void SContentBrowser::Construct( const FArguments& InArgs, const FName& InInstan
 	PathContextMenu = MakeShareable(new FPathContextMenu( AsShared() ));
 	PathContextMenu->SetOnNewAssetRequested( FNewAssetContextMenu::FOnNewAssetRequested::CreateSP(this, &SContentBrowser::NewAssetRequested) );
 
-	TSharedPtr<AssetFilterCollectionType> FrontendFilters = MakeShareable(new AssetFilterCollectionType());
-	TSharedPtr<AssetFilterCollectionType> ExtraFilters = MakeShareable(new AssetFilterCollectionType());
-	ExtraFilters->Add( TextFilter );
+	FrontendFilters = MakeShareable(new AssetFilterCollectionType());
+	TextFilter = MakeShareable( new FFrontendFilter_Text() );
 
 	FContentBrowserCommands::Register();
 
@@ -544,7 +536,6 @@ void SContentBrowser::Construct( const FArguments& InArgs, const FName& InInstan
 							.OnAssetRenameCommitted(this, &SContentBrowser::OnAssetRenameCommitted)
 							.AreRealTimeThumbnailsAllowed(this, &SContentBrowser::IsHovered)
 							.FrontendFilters(FrontendFilters)
-							.DynamicFilters(ExtraFilters)
 							.HighlightedText(this, &SContentBrowser::GetHighlightedText)
 							.AllowThumbnailEditMode(true)
 							.AllowThumbnailHintLabel(false)
@@ -1064,10 +1055,12 @@ void SContentBrowser::OnSearchBoxChanged(const FText& InSearchText)
 	TextFilter->SetRawFilterText( InSearchText );
 	if(InSearchText.IsEmpty())
 	{
+		FrontendFilters->Remove(TextFilter);
 		AssetViewPtr->SetUserSearching(false);
 	}
 	else
 	{
+		FrontendFilters->Add(TextFilter);
 		AssetViewPtr->SetUserSearching(true);
 	}
 
@@ -1081,10 +1074,12 @@ void SContentBrowser::OnSearchBoxCommitted(const FText& InSearchText, ETextCommi
 	TextFilter->SetRawFilterText( InSearchText );
 	if(InSearchText.IsEmpty())
 	{
+		FrontendFilters->Remove(TextFilter);
 		AssetViewPtr->SetUserSearching(false);
 	}
 	else
 	{
+		FrontendFilters->Add(TextFilter);
 		AssetViewPtr->SetUserSearching(true);
 	}
 }
@@ -1725,7 +1720,7 @@ void SContentBrowser::OnDuplicateRequested(const TWeakObjectPtr<UObject>& Origin
 
 void SContentBrowser::OnAssetViewRefreshRequested()
 {
-	AssetViewPtr->RequestListRefresh();
+	AssetViewPtr->RequestSlowFullListRefresh();
 }
 
 void SContentBrowser::HandleSettingChanged(FName PropertyName)
diff --git a/Engine/Source/Editor/ContentBrowser/Private/SContentBrowser.h b/Engine/Source/Editor/ContentBrowser/Private/SContentBrowser.h
--- a/Engine/Source/Editor/ContentBrowser/Private/SContentBrowser.h
+++ b/Engine/Source/Editor/ContentBrowser/Private/SContentBrowser.h
@@ -329,8 +329,11 @@ private:
 	/** True if source should not be changed from an outside source */
 	bool bIsLocked;
 
+	/** The list of FrontendFilters currently applied to the asset view */
+	TSharedPtr<AssetFilterCollectionType> FrontendFilters;
+
 	/** The text filter to use on the assets */
-	TSharedPtr< TTextFilter< AssetFilterType > > TextFilter;
+	TSharedPtr< FFrontendFilter_Text > TextFilter;
 
 	/** Commands handled by this widget */
 	TSharedPtr< FUICommandList > Commands;
diff --git a/Engine/Source/Editor/ContentBrowser/Private/SFilterList.cpp b/Engine/Source/Editor/ContentBrowser/Private/SFilterList.cpp
--- a/Engine/Source/Editor/ContentBrowser/Private/SFilterList.cpp
+++ b/Engine/Source/Editor/ContentBrowser/Private/SFilterList.cpp
@@ -387,7 +387,12 @@ void SFilterList::Construct( const FArguments& InArgs )
 	for(auto Iter = InArgs._ExtraFrontendFilters.CreateConstIterator(); Iter; ++Iter)
 	{
 		TSharedRef<FFrontendFilter> Filter = (*Iter);
-		AllFrontendFilterCategories.AddUnique( Filter->GetCategory() );
+		TSharedPtr<FFrontendFilterCategory> Category = Filter->GetCategory();
+		if ( Category.IsValid() )
+		{
+			AllFrontendFilterCategories.AddUnique( Category );
+		}
+
 		AllFrontendFilters.Add(Filter);
 	}
 
diff --git a/Engine/Source/Editor/ContentBrowser/Public/IContentBrowserSingleton.h b/Engine/Source/Editor/ContentBrowser/Public/IContentBrowserSingleton.h
--- a/Engine/Source/Editor/ContentBrowser/Public/IContentBrowserSingleton.h
+++ b/Engine/Source/Editor/ContentBrowser/Public/IContentBrowserSingleton.h
@@ -32,36 +32,20 @@ struct FAssetPickerConfig
 	/** The selection mode the picker should use */
 	ESelectionMode::Type SelectionMode;
 
-	/** A pointer to an existing delegate which the AssetView will register a function which clears the current selection */
-	FClearSelectionDelegate* ClearSelectionDelegate;
-
 	/** An array of pointers to existing delegates which the AssetView will register a function which returns the current selection */
 	TArray<FGetCurrentSelectionDelegate*> GetCurrentSelectionDelegates;
 
-	/** A pointer to an existing delegate which the AssetView will register a function which increments the selection */
-	FAdjustSelectionDelegate* AdjustSelectionDelegate;
-
 	/** The asset registry filter to use to cull results */
 	FARFilter Filter;
 
-	/** Dynamic filters that get applied after the backend and frontend filters*/
-	TSharedRef< AssetFilterCollectionType > ExtraFilters;
-
 	/** Custom front end filters to be displayed */
 	TArray< TSharedRef<class FFrontendFilter> > ExtraFrontendFilters;
 
-	/** The text to highlight on assets */
-	TAttribute< FText > HighlightedText;
-
-	/** The label visibility on asset items */
-	TAttribute< EVisibility > LabelVisibility;
-
 	/** The contents of the label on the thumbnail */
 	EThumbnailLabel::Type ThumbnailLabel;
 
 	/** The default scale for thumbnails. [0-1] range */
 	TAttribute< float > ThumbnailScale;
-	FOnThumbnailScaleChanged ThumbnailScaleChanged;
 
 	/** Only display results in these collections */
 	TArray<FCollectionNameType> Collections;
@@ -72,12 +56,6 @@ struct FAssetPickerConfig
 	/** The delegate that fires when an asset was selected */
 	FOnAssetSelected OnAssetSelected;
 
-	/** The delegate that fires when an asset is dragged */
-	FOnAssetDragged OnAssetDragged;
-
-	/** The delegate that fires when an asset is clicked on */
-	FOnAssetClicked OnAssetClicked;
-
 	/** The delegate that fires when an asset is double clicked */
 	FOnAssetDoubleClicked OnAssetDoubleClicked;
 
@@ -97,9 +75,6 @@ struct FAssetPickerConfig
 	    determine if it should be displayed or not.  If it returns true or isn't bound, the tag will be displayed normally. */
 	FOnShouldDisplayAssetTag OnAssetTagWantsToBeDisplayed;
 
-	/** The delegate that fires to construct the tooltip for the specified asset */
-	FConstructToolTipForAsset ConstructToolTipForAsset;
-
 	/** The default view mode */
 	EAssetViewType::Type InitialAssetViewType;
 
@@ -124,9 +99,6 @@ struct FAssetPickerConfig
 	/** Indicates if this view is allowed to show classes */
 	bool bCanShowClasses;
 
-	/** Indicates if the 'Show Folders' option should be visible */
-	bool bCanShowFolders;
-
 	/** Indicates if the 'Show Only Assets In Selection' option should be visible */
 	bool bCanShowOnlyAssetsInSelectedFolders;
 
@@ -144,15 +116,11 @@ struct FAssetPickerConfig
 
 	/** Override the default filter context menu layout */
 	EAssetTypeCategories::Type DefaultFilterMenuExpansion;
+
 	FAssetPickerConfig()
 		: SelectionMode( ESelectionMode::Multi )
-		, ClearSelectionDelegate( NULL )
-		, AdjustSelectionDelegate( NULL )
-		, ExtraFilters( MakeShareable( new AssetFilterCollectionType() ) )
-		, HighlightedText()
 		, ThumbnailLabel( EThumbnailLabel::ClassName )
 		, ThumbnailScale(0.25f) // A reasonable scale
-		, ThumbnailScaleChanged()
 		, InitialAssetViewType(EAssetViewType::Tile)
 		, bFocusSearchBoxWhenOpened(true)
 		, bAllowNullSelection(false)
@@ -160,7 +128,6 @@ struct FAssetPickerConfig
 		, bAutohideSearchBar(false)
 		, bAllowDragging(true)
 		, bCanShowClasses(true)
-		, bCanShowFolders(false)
 		, bCanShowOnlyAssetsInSelectedFolders(false)
 		, bCanShowRealTimeThumbnails(false)
 		, bCanShowDevelopersFolder(false)
