UE4: If an asset has an asset type actions associated with it, the user-friendly name is now displayed as the type instead of the classname.

diff --git a/Engine/Source/Editor/ContentBrowser/Private/AssetViewWidgets.cpp b/Engine/Source/Editor/ContentBrowser/Private/AssetViewWidgets.cpp
--- a/Engine/Source/Editor/ContentBrowser/Private/AssetViewWidgets.cpp
+++ b/Engine/Source/Editor/ContentBrowser/Private/AssetViewWidgets.cpp
@@ -133,10 +133,8 @@ void SAssetViewItem::Construct( const FArguments& InArgs )
 
 	bDraggedOver = false;
 
-	CachePackageName();
-	AssetPackage = FindObjectSafe<UPackage>(NULL, *CachedPackageName);
 	bPackageDirty = false;
-	UpdatePackageDirtyState();
+	OnAssetDataChanged();
 
 	AssetItem->OnAssetDataChanged.AddSP(this, &SAssetViewItem::OnAssetDataChanged);
 
@@ -352,6 +350,17 @@ void SAssetViewItem::OnAssetDataChanged()
 	AssetPackage = FindObjectSafe<UPackage>(NULL, *CachedPackageName);
 	UpdatePackageDirtyState();
 
+	AssetTypeActions.Reset();
+	if ( AssetItem->GetType() != EAssetItemType::Folder )
+	{
+		UClass* AssetClass = FindObject<UClass>(ANY_PACKAGE, *StaticCastSharedPtr<FAssetViewAsset>(AssetItem)->Data.AssetClass.ToString());
+		if (AssetClass)
+		{
+			FAssetToolsModule& AssetToolsModule = FModuleManager::LoadModuleChecked<FAssetToolsModule>(TEXT("AssetTools"));
+			AssetTypeActions = AssetToolsModule.Get().GetAssetTypeActionsForClass(AssetClass).Pin();
+		}
+	}
+
 	if ( InlineRenameWidget.IsValid() )
 	{
 		InlineRenameWidget->SetText( GetNameText() );
@@ -362,6 +371,30 @@ void SAssetViewItem::DirtyStateChanged()
 {
 }
 
+FText SAssetViewItem::GetAssetClassText() const
+{
+	if (AssetItem.IsValid())
+	{
+		if (AssetItem->GetType() != EAssetItemType::Folder)
+		{
+			if (AssetTypeActions.IsValid())
+			{
+				return AssetTypeActions.Pin()->GetName();
+			}
+			else
+			{
+				return FText::FromName(StaticCastSharedPtr<FAssetViewAsset>(AssetItem)->Data.AssetClass);
+			}
+		}
+		else
+		{
+			return LOCTEXT("FolderName", "Folder");
+		}
+	}
+
+	return FText();
+}
+
 const FSlateBrush* SAssetViewItem::GetSCCStateImage() const
 {
 	return SCCStateBrush;
@@ -401,7 +434,7 @@ TSharedRef<SToolTip> SAssetViewItem::CreateToolTipWidget() const
 
 			// The tooltip contains the name, class, path, and asset registry tags
 			const FText NameText = FText::FromName( AssetData.AssetName );
-			const FText ClassText = FText::Format( LOCTEXT("ClassName", "({0})"), FText::FromName( AssetData.AssetClass ) );
+			const FText ClassText = FText::Format( LOCTEXT("ClassName", "({0})"), GetAssetClassText() );
 
 
 			// Create a box to hold every line of info in the body of the tooltip
@@ -668,24 +701,12 @@ FString SAssetViewItem::GetCheckedOutByOtherText() const
 
 FText SAssetViewItem::GetAssetUserDescription() const
 {
-	if (AssetItem.IsValid())
+	if (AssetItem.IsValid() && AssetTypeActions.IsValid())
 	{
 		if (AssetItem->GetType() != EAssetItemType::Folder)
 		{
 			const FAssetData& AssetData = StaticCastSharedPtr<FAssetViewAsset>(AssetItem)->Data;
-			FText AssetClass = FText::FromName(AssetData.AssetClass);
-
-			if (UClass* AssetClassObj = FindObject<UClass>(ANY_PACKAGE, *AssetClass.ToString()))
-			{
-				FAssetToolsModule& AssetToolsModule = FModuleManager::LoadModuleChecked<FAssetToolsModule>(TEXT("AssetTools"));
-				TWeakPtr<IAssetTypeActions> AssetTypeActions = AssetToolsModule.Get().GetAssetTypeActionsForClass(AssetClassObj);
-
-				if (AssetTypeActions.IsValid())
-				{
-					FText Description = AssetTypeActions.Pin()->GetAssetDescription(AssetData);
-					return Description;
-				}
-			}
+			return AssetTypeActions.Pin()->GetAssetDescription(AssetData);
 		}
 	}
 
@@ -819,6 +840,10 @@ FSlateColor SAssetViewItem::GetAssetColor() const
 				return *Color.Get();
 			}
 		}
+		else if(AssetTypeActions.IsValid())
+		{
+			return AssetTypeActions.Pin()->GetTypeColor().ReinterpretAsLinear();
+		}
 	}
 	return ContentBrowserUtils::GetDefaultColor();
 }
@@ -893,24 +918,8 @@ void SAssetListItem::Construct( const FArguments& InArgs )
 	ItemHeight = InArgs._ItemHeight;
 
 	const float ThumbnailPadding = InArgs._ThumbnailPadding;
-
-	FText AssetClass;
 	bool bIsDeveloperFolder = false;
 
-	if(AssetItem.IsValid())
-	{
-		if(AssetItem->GetType() != EAssetItemType::Folder)
-		{
-			const FAssetData& AssetData = StaticCastSharedPtr<FAssetViewAsset>(AssetItem)->Data;
-			AssetClass = FText::FromName(AssetData.AssetClass);
-		}
-		else
-		{
-			AssetClass = LOCTEXT("FolderName", "Folder");
-			bIsDeveloperFolder = StaticCastSharedPtr<FAssetViewFolder>(AssetItem)->bDeveloperFolder;
-		}
-	}
-
 	TSharedPtr<SWidget> Thumbnail;
 	if ( AssetItem.IsValid() && AssetThumbnail.IsValid() )
 	{
@@ -1034,9 +1043,9 @@ void SAssetListItem::Construct( const FArguments& InArgs )
 				.Padding(0, 1)
 				[
 					// Class
-					SNew(STextBlock)
+					SAssignNew(ClassText, STextBlock)
 					.Font(FEditorStyle::GetFontStyle("ContentBrowser.AssetListViewClassFont"))
-					.Text(AssetClass)
+					.Text(GetAssetClassText())
 				]
 			]
 		]
@@ -1060,6 +1069,11 @@ void SAssetListItem::OnAssetDataChanged()
 {
 	SAssetViewItem::OnAssetDataChanged();
 
+	if (ClassText.IsValid())
+	{
+		ClassText->SetText(GetAssetClassText());
+	}
+
 	SetToolTip( CreateToolTipWidget() );
 }
 
@@ -1425,33 +1439,6 @@ void SAssetColumnItem::OnAssetDataChanged()
 	}
 }
 
-FSlateColor SAssetColumnItem::GetAssetColor() const
-{
-	FSlateColor AssetTypeColor = SAssetViewItem::GetAssetColor();
-
-	if(AssetItem.IsValid())
-	{
-		if(AssetItem->GetType() != EAssetItemType::Folder)
-		{
-			const FAssetData& AssetData = StaticCastSharedPtr<FAssetViewAsset>(AssetItem)->Data;
-			FText AssetClass = FText::FromName(AssetData.AssetClass);
-
-			if (UClass* AssetClassObj = FindObject<UClass>(ANY_PACKAGE, *AssetClass.ToString()))
-			{
-				FAssetToolsModule& AssetToolsModule = FModuleManager::LoadModuleChecked<FAssetToolsModule>(TEXT("AssetTools"));
-				TWeakPtr<IAssetTypeActions> AssetTypeActions = AssetToolsModule.Get().GetAssetTypeActionsForClass(AssetClassObj);
-
-				if (AssetTypeActions.IsValid())
-				{
-					AssetTypeColor = AssetTypeActions.Pin()->GetTypeColor().ReinterpretAsLinear();
-				}
-			}
-		}
-	}
-
-	return AssetTypeColor;
-}
-
 FString SAssetColumnItem::GetAssetNameToolTipText() const
 {
 	if ( AssetItem.IsValid() )
@@ -1482,23 +1469,6 @@ FString SAssetColumnItem::GetAssetNameToolTipText() const
 	}
 }
 
-FText SAssetColumnItem::GetAssetClassText() const
-{
-	if ( AssetItem.IsValid() )
-	{
-		if(AssetItem->GetType() != EAssetItemType::Folder)
-		{
-			return FText::FromName(StaticCastSharedPtr<FAssetViewAsset>(AssetItem)->Data.AssetClass);
-		}
-		else
-		{
-			return LOCTEXT("FolderName", "Folder");
-		}
-	}
-
-	return FText();
-}
-
 FText SAssetColumnItem::GetAssetPathText() const
 {
 	if ( AssetItem.IsValid() )
diff --git a/Engine/Source/Editor/ContentBrowser/Private/AssetViewWidgets.h b/Engine/Source/Editor/ContentBrowser/Private/AssetViewWidgets.h
--- a/Engine/Source/Editor/ContentBrowser/Private/AssetViewWidgets.h
+++ b/Engine/Source/Editor/ContentBrowser/Private/AssetViewWidgets.h
@@ -122,6 +122,9 @@ protected:
 	/** Notification for when the dirty flag changes */
 	virtual void DirtyStateChanged();
 
+	/** Gets the name of the class of this asset */
+	FText GetAssetClassText() const;
+
 	/** Gets the brush for the source control indicator image */
 	const FSlateBrush* GetSCCStateImage() const;
 
@@ -177,6 +180,9 @@ protected:
 	/** The package containing the asset that this item represents */
 	TWeakObjectPtr<UPackage> AssetPackage;
 
+	/** The asset type actions associated with the asset that this item represents */
+	TWeakPtr<IAssetTypeActions> AssetTypeActions;
+
 	/** The cached name of the package containing the asset that this item represents */
 	FString CachedPackageName;
 
@@ -332,6 +338,9 @@ private:
 
 	/** The height allowed for this item */
 	TAttribute<float> ItemHeight;
+
+	/** The text block containing the class name */
+	TSharedPtr<STextBlock> ClassText;
 };
 
 /** An item in the asset tile view */
@@ -478,16 +487,10 @@ public:
 	/** Handles committing a name change */
 	virtual void OnAssetDataChanged() override;
 
-	/** Returns the color this item should be tinted with */
-	virtual FSlateColor GetAssetColor() const override;
-
 private:
 	/** Gets the tool tip text for the name */
 	FString GetAssetNameToolTipText() const;
 
-	/** Gets the name of the class of this asset */
-	FText GetAssetClassText() const;
-
 	/** Gets the path to this asset */
 	FText GetAssetPathText() const;
 
