Highlight any matched filter terms in the path or the type in the tooltip of an asset item, or in the list or column views of the content browser.
#ttp 333654 - Content Browser: Search text field currently filters the assets using folder names as well.
#branch UE4
#proj Editor.ContentBrowser
#add Cache HighlightText attribute in SAssetViewItem, and pass it as the parameter to HighlightText in the various STextBlock widgets which require it.

diff --git a/Engine/Source/Editor/ContentBrowser/Private/AssetViewWidgets.cpp b/Engine/Source/Editor/ContentBrowser/Private/AssetViewWidgets.cpp
--- a/Engine/Source/Editor/ContentBrowser/Private/AssetViewWidgets.cpp
+++ b/Engine/Source/Editor/ContentBrowser/Private/AssetViewWidgets.cpp
@@ -127,6 +127,7 @@ void SAssetViewItem::Construct( const FArguments& InArgs )
 	OnItemDestroyed = InArgs._OnItemDestroyed;
 	ShouldAllowToolTip = InArgs._ShouldAllowToolTip;
 	ThumbnailEditMode = InArgs._ThumbnailEditMode;
+	HighlightText = InArgs._HighlightText;
 	OnAssetsDragDropped = InArgs._OnAssetsDragDropped;
 	OnPathsDragDropped = InArgs._OnPathsDragDropped;
 	OnFilesDragDropped = InArgs._OnFilesDragDropped;
@@ -536,6 +537,7 @@ TSharedRef<SToolTip> SAssetViewItem::CreateToolTipWidget() const
 							.VAlign(VAlign_Center)
 							[
 								SNew(STextBlock).Text(ClassText)
+								.HighlightText(HighlightText)
 							]
 						]
 
@@ -738,6 +740,7 @@ void SAssetViewItem::AddToToolTipInfoBox(const TSharedRef<SVerticalBox>& InfoBox
 		[
 			SNew(STextBlock) .Text( Value )
 			.ColorAndOpacity(bImportant ? ImportantStyle.GetForegroundColor() : FSlateColor::UseForeground())
+			.HighlightText((Key.ToString() == TEXT("Path")) ? HighlightText : FText())
 		]
 	];
 }
@@ -909,6 +912,7 @@ void SAssetListItem::Construct( const FArguments& InArgs )
 		.OnItemDestroyed(InArgs._OnItemDestroyed)
 		.ShouldAllowToolTip(InArgs._ShouldAllowToolTip)
 		.ThumbnailEditMode(InArgs._ThumbnailEditMode)
+		.HighlightText(InArgs._HighlightText)
 		.OnAssetsDragDropped(InArgs._OnAssetsDragDropped)
 		.OnPathsDragDropped(InArgs._OnPathsDragDropped)
 		.OnFilesDragDropped(InArgs._OnFilesDragDropped)
@@ -1046,6 +1050,7 @@ void SAssetListItem::Construct( const FArguments& InArgs )
 					SAssignNew(ClassText, STextBlock)
 					.Font(FEditorStyle::GetFontStyle("ContentBrowser.AssetListViewClassFont"))
 					.Text(GetAssetClassText())
+					.HighlightText(InArgs._HighlightText)
 				]
 			]
 		]
@@ -1108,6 +1113,7 @@ void SAssetTileItem::Construct( const FArguments& InArgs )
 		.OnItemDestroyed(InArgs._OnItemDestroyed)
 		.ShouldAllowToolTip(InArgs._ShouldAllowToolTip)
 		.ThumbnailEditMode(InArgs._ThumbnailEditMode)
+		.HighlightText(InArgs._HighlightText)
 		.OnAssetsDragDropped(InArgs._OnAssetsDragDropped)
 		.OnPathsDragDropped(InArgs._OnPathsDragDropped)
 		.OnFilesDragDropped(InArgs._OnFilesDragDropped)
@@ -1296,6 +1302,7 @@ void SAssetColumnItem::Construct( const FArguments& InArgs )
 		.OnRenameCommit(InArgs._OnRenameCommit)
 		.OnVerifyRenameCommit(InArgs._OnVerifyRenameCommit)
 		.OnItemDestroyed(InArgs._OnItemDestroyed)
+		.HighlightText(InArgs._HighlightText)
 		.OnAssetsDragDropped(InArgs._OnAssetsDragDropped)
 		.OnPathsDragDropped(InArgs._OnPathsDragDropped)
 		.OnFilesDragDropped(InArgs._OnFilesDragDropped)
@@ -1401,13 +1408,15 @@ TSharedRef<SWidget> SAssetColumnItem::GenerateWidgetForColumn( const FName& Colu
 	{
 		Content = SAssignNew(ClassText, STextBlock)
 			.ToolTipText( this, &SAssetColumnItem::GetAssetClassText )
-			.Text( GetAssetClassText() );
+			.Text( GetAssetClassText() )
+			.HighlightText( HighlightText );
 	}
 	else if ( ColumnName == "Path" )
 	{
 		Content = SAssignNew(PathText, STextBlock)
 			.ToolTipText( this, &SAssetColumnItem::GetAssetPathText )
-			.Text( GetAssetPathText() );
+			.Text( GetAssetPathText() )
+			.HighlightText( HighlightText );
 	}
 	else
 	{
diff --git a/Engine/Source/Editor/ContentBrowser/Private/AssetViewWidgets.h b/Engine/Source/Editor/ContentBrowser/Private/AssetViewWidgets.h
--- a/Engine/Source/Editor/ContentBrowser/Private/AssetViewWidgets.h
+++ b/Engine/Source/Editor/ContentBrowser/Private/AssetViewWidgets.h
@@ -70,6 +70,9 @@ public:
 		/** If true, display the thumbnail edit mode UI */
 		SLATE_ATTRIBUTE( bool, ThumbnailEditMode )
 
+		/** The string in the title to highlight (used when searching by string) */
+		SLATE_ATTRIBUTE(FText, HighlightText)
+
 		/** Delegate for when assets are dropped on this item, if it is a folder */
 		SLATE_EVENT( FOnAssetsDragDropped, OnAssetsDragDropped )
 
@@ -210,6 +213,9 @@ protected:
 	/** If true, display the thumbnail edit mode UI */
 	TAttribute<bool> ThumbnailEditMode;
 
+	/** The substring to be highlighted in the name and tooltip path */
+	TAttribute<FText> HighlightText;
+
 	/**
 	 * A map of class names to their important asset registry tags and values.
 	 * Important tags will be featured in the tooltip if the value matches.
