Prohibited classes from being dragged into a Collection.
#ttp 339146 - EDITOR: Content Browser: Can't delete classes from collections
#branch UE4
#proj Editor.ContentBrowser
#add Added extra checks in SAssetView::OnDrop and SAssetView::OnDragOver to prevent class assets from being added to Collections.

diff --git a/Engine/Source/Editor/ContentBrowser/Private/SAssetView.cpp b/Engine/Source/Editor/ContentBrowser/Private/SAssetView.cpp
--- a/Engine/Source/Editor/ContentBrowser/Private/SAssetView.cpp
+++ b/Engine/Source/Editor/ContentBrowser/Private/SAssetView.cpp
@@ -1063,9 +1063,14 @@ FReply SAssetView::OnDragOver( const FGeometry& MyGeometry, const FDragDropEvent
 				CollectionManagerModule.Get().GetObjectsInCollection( SourcesData.Collections[0].Name, SourcesData.Collections[0].Type, ObjectPaths );
 
 				bool IsValidDrop = false;
-				for (int Index = 0; Index < AssetDatas.Num(); Index++)
+				for (const auto& AssetData : AssetDatas)
 				{
-					if ( !ObjectPaths.Contains( AssetDatas[Index].ObjectPath ) )
+					if (AssetData.GetClass()->IsChildOf(UClass::StaticClass()))
+					{
+						continue;
+					}
+
+					if ( !ObjectPaths.Contains( AssetData.ObjectPath ) )
 					{
 						IsValidDrop = true;
 						break;
@@ -1109,13 +1114,19 @@ FReply SAssetView::OnDrop( const FGeometry& MyGeometry, const FDragDropEvent& Dr
 		if ( SelectedAssetDatas.Num() > 0 )
 		{
 			TArray< FName > ObjectPaths;
-			for (int Index = 0; Index < SelectedAssetDatas.Num(); Index++)
+			for (const auto& AssetData : SelectedAssetDatas)
 			{
-				ObjectPaths.Add( SelectedAssetDatas[Index].ObjectPath );
+				if (!AssetData.GetClass()->IsChildOf(UClass::StaticClass()))
+				{
+					ObjectPaths.Add(AssetData.ObjectPath);
+				}
 			}
 
-			FCollectionManagerModule& CollectionManagerModule = FModuleManager::LoadModuleChecked<FCollectionManagerModule>(TEXT("CollectionManager"));
-			CollectionManagerModule.Get().AddToCollection( SourcesData.Collections[0].Name, SourcesData.Collections[0].Type, ObjectPaths );
+			if (ObjectPaths.Num() > 0)
+			{
+				FCollectionManagerModule& CollectionManagerModule = FModuleManager::LoadModuleChecked<FCollectionManagerModule>(TEXT("CollectionManager"));
+				CollectionManagerModule.Get().AddToCollection(SourcesData.Collections[0].Name, SourcesData.Collections[0].Type, ObjectPaths);
+			}
 
 			return FReply::Handled();
 		}
