Created Blueprint category for new assets and filters.
#ttp 339216 - Feature Request: Content Browser: Create Blueprint catagory for new assets
#branch UE4
#proj Developer.AssetTools, Editor.Blutility, Editor.ContentBrowser, Editor.GameplayAbilitiesEditor, Editor.UnrealEd
#add Reassigned asset categories, putting all Blueprint type assets under a new Blueprint category.
#add Added new submenu in FNewAssetContextMenu::MakeContextMenu().
#add Added new filter map entry in SFilterList::MakeAddFilterMenu().

diff --git a/Engine/Source/Developer/AssetTools/Private/AssetTypeActions/AssetTypeActions_Blueprint.h b/Engine/Source/Developer/AssetTools/Private/AssetTypeActions/AssetTypeActions_Blueprint.h
--- a/Engine/Source/Developer/AssetTools/Private/AssetTypeActions/AssetTypeActions_Blueprint.h
+++ b/Engine/Source/Developer/AssetTools/Private/AssetTypeActions/AssetTypeActions_Blueprint.h
@@ -12,7 +12,7 @@ public:
 	virtual bool HasActions ( const TArray<UObject*>& InObjects ) const override { return true; }
 	virtual void GetActions( const TArray<UObject*>& InObjects, FMenuBuilder& MenuBuilder ) override;
 	virtual void OpenAssetEditor( const TArray<UObject*>& InObjects, TSharedPtr<class IToolkitHost> EditWithinLevelEditor = TSharedPtr<IToolkitHost>() ) override;
-	virtual uint32 GetCategories() override { return EAssetTypeCategories::Basic; }
+	virtual uint32 GetCategories() override { return EAssetTypeCategories::Blueprint | EAssetTypeCategories::Basic; }
 	virtual void PerformAssetDiff(UObject* Asset1, UObject* Asset2, const struct FRevisionInfo& OldRevision, const struct FRevisionInfo& NewRevision) const override;
 	virtual class UThumbnailInfo* GetThumbnailInfo(UObject* Asset) const override;
 	virtual FText GetAssetDescription(const FAssetData& AssetData) const override;
diff --git a/Engine/Source/Developer/AssetTools/Private/AssetTypeActions/AssetTypeActions_BlueprintGeneratedClass.h b/Engine/Source/Developer/AssetTools/Private/AssetTypeActions/AssetTypeActions_BlueprintGeneratedClass.h
--- a/Engine/Source/Developer/AssetTools/Private/AssetTypeActions/AssetTypeActions_BlueprintGeneratedClass.h
+++ b/Engine/Source/Developer/AssetTools/Private/AssetTypeActions/AssetTypeActions_BlueprintGeneratedClass.h
@@ -12,7 +12,7 @@ public:
 	virtual bool HasActions ( const TArray<UObject*>& InObjects ) const override { return true; }
 	virtual void GetActions( const TArray<UObject*>& InObjects, FMenuBuilder& MenuBuilder ) override;
 	virtual void OpenAssetEditor( const TArray<UObject*>& InObjects, TSharedPtr<class IToolkitHost> EditWithinLevelEditor = TSharedPtr<IToolkitHost>() ) override;
-	virtual uint32 GetCategories() override { return EAssetTypeCategories::Misc; }
+	virtual uint32 GetCategories() override { return EAssetTypeCategories::Blueprint; }
 	virtual void PerformAssetDiff(UObject* Asset1, UObject* Asset2, const struct FRevisionInfo& OldRevision, const struct FRevisionInfo& NewRevision) const override;
 	virtual class UThumbnailInfo* GetThumbnailInfo(UObject* Asset) const override;
 
diff --git a/Engine/Source/Developer/AssetTools/Private/AssetTypeActions/AssetTypeActions_Enum.h b/Engine/Source/Developer/AssetTools/Private/AssetTypeActions/AssetTypeActions_Enum.h
--- a/Engine/Source/Developer/AssetTools/Private/AssetTypeActions/AssetTypeActions_Enum.h
+++ b/Engine/Source/Developer/AssetTools/Private/AssetTypeActions/AssetTypeActions_Enum.h
@@ -9,7 +9,7 @@ public:
 	virtual FText GetName() const override { return NSLOCTEXT("AssetTypeActions", "AssetTypeActions_Enum", "Enumeration"); }
 	virtual FColor GetTypeColor() const override { return FColor(255, 200, 200); }
 	virtual UClass* GetSupportedClass() const override { return UUserDefinedEnum::StaticClass(); }
-	virtual uint32 GetCategories() override { return EAssetTypeCategories::Misc; }
+	virtual uint32 GetCategories() override { return EAssetTypeCategories::Blueprint; }
 
 	virtual void OpenAssetEditor( const TArray<UObject*>& InObjects, TSharedPtr<class IToolkitHost> EditWithinLevelEditor = TSharedPtr<IToolkitHost>() ) override;
 };
diff --git a/Engine/Source/Developer/AssetTools/Private/AssetTypeActions/AssetTypeActions_Struct.h b/Engine/Source/Developer/AssetTools/Private/AssetTypeActions/AssetTypeActions_Struct.h
--- a/Engine/Source/Developer/AssetTools/Private/AssetTypeActions/AssetTypeActions_Struct.h
+++ b/Engine/Source/Developer/AssetTools/Private/AssetTypeActions/AssetTypeActions_Struct.h
@@ -9,7 +9,7 @@ public:
 	virtual FText GetName() const override { return NSLOCTEXT("AssetTypeActions", "AssetTypeActions_Struct", "Structure"); }
 	virtual FColor GetTypeColor() const override { return FColor(0, 0, 255); }
 	virtual UClass* GetSupportedClass() const override { return UUserDefinedStruct::StaticClass(); }
-	virtual uint32 GetCategories() override { return EAssetTypeCategories::Misc; }
+	virtual uint32 GetCategories() override { return EAssetTypeCategories::Blueprint; }
 
 	virtual FText GetAssetDescription(const FAssetData& AssetData) const override;
 	virtual void OpenAssetEditor( const TArray<UObject*>& InObjects, TSharedPtr<class IToolkitHost> EditWithinLevelEditor = TSharedPtr<IToolkitHost>() ) override;
diff --git a/Engine/Source/Developer/AssetTools/Public/IAssetTypeActions.h b/Engine/Source/Developer/AssetTools/Public/IAssetTypeActions.h
--- a/Engine/Source/Developer/AssetTools/Public/IAssetTypeActions.h
+++ b/Engine/Source/Developer/AssetTools/Public/IAssetTypeActions.h
@@ -15,6 +15,8 @@ namespace EAssetTypeCategories
 		UI						= 1 << 5,
 		Misc					= 1 << 6,
 		Gameplay				= 1 << 7,
+		Blueprint				= 1 << 8
+		// Last allowed value is 1 << 31
 	};
 }
 
diff --git a/Engine/Source/Editor/Blutility/Private/AssetTypeActions_EditorUtilityBlueprint.cpp b/Engine/Source/Editor/Blutility/Private/AssetTypeActions_EditorUtilityBlueprint.cpp
--- a/Engine/Source/Editor/Blutility/Private/AssetTypeActions_EditorUtilityBlueprint.cpp
+++ b/Engine/Source/Editor/Blutility/Private/AssetTypeActions_EditorUtilityBlueprint.cpp
@@ -130,7 +130,7 @@ void FAssetTypeActions_EditorUtilityBlueprint::OpenAssetEditor(const TArray<UObj
 
 uint32 FAssetTypeActions_EditorUtilityBlueprint::GetCategories()
 {
-	return EAssetTypeCategories::Misc;
+	return EAssetTypeCategories::Blueprint;
 }
 
 void FAssetTypeActions_EditorUtilityBlueprint::ExecuteEdit(FWeakBlueprintPointerArray Objects)
diff --git a/Engine/Source/Editor/ContentBrowser/Private/NewAssetContextMenu.cpp b/Engine/Source/Editor/ContentBrowser/Private/NewAssetContextMenu.cpp
--- a/Engine/Source/Editor/ContentBrowser/Private/NewAssetContextMenu.cpp
+++ b/Engine/Source/Editor/ContentBrowser/Private/NewAssetContextMenu.cpp
@@ -177,6 +177,18 @@ void FNewAssetContextMenu::MakeContextMenu(FMenuBuilder& MenuBuilder, const FStr
 			}
 		}
 
+		{
+			TArray<FFactoryItem> BlueprintFactories = FindFactoriesInCategory(EAssetTypeCategories::Blueprint);
+			if (BlueprintFactories.Num() > 0)
+			{
+				MenuBuilder.AddSubMenu(
+					LOCTEXT("BlueprintAssetCategory", "Blueprints"),
+					FText::GetEmpty(),
+					FNewMenuDelegate::CreateStatic(&FNewAssetContextMenu::CreateNewAssetMenuCategory, EAssetTypeCategories::Blueprint, InPath, InOnNewAssetRequested)
+					);
+			}
+		}
+
 		{
 			TArray<FFactoryItem> MaterialFactories = FindFactoriesInCategory(EAssetTypeCategories::MaterialsAndTextures);
 			if (MaterialFactories.Num() > 0)
diff --git a/Engine/Source/Editor/ContentBrowser/Private/SFilterList.cpp b/Engine/Source/Editor/ContentBrowser/Private/SFilterList.cpp
--- a/Engine/Source/Editor/ContentBrowser/Private/SFilterList.cpp
+++ b/Engine/Source/Editor/ContentBrowser/Private/SFilterList.cpp
@@ -952,7 +952,8 @@ TSharedRef<SWidget> SFilterList::MakeAddFilterMenu(EAssetTypeCategories::Type Me
 	TMap<EAssetTypeCategories::Type, FCategoryMenu> CategoryToMenuMap;
 	CategoryToMenuMap.Add(EAssetTypeCategories::Basic, FCategoryMenu( LOCTEXT("BasicFilter", "Basic"), LOCTEXT("BasicFilterTooltip", "Filter by basic assets."), "ContentBrowserFilterBasicAsset", LOCTEXT("BasicAssetsMenuHeading", "Basic Assets") ) );
 	CategoryToMenuMap.Add(EAssetTypeCategories::Animation, FCategoryMenu( LOCTEXT("AnimationFilter", "Animation"), LOCTEXT("AnimationFilterTooltip", "Filter by animation assets."), "ContentBrowserFilterAnimationAsset", LOCTEXT("AnimationAssetsMenuHeading", "Animation Assets") ) );
-	CategoryToMenuMap.Add(EAssetTypeCategories::MaterialsAndTextures, FCategoryMenu( LOCTEXT("MaterialFilter", "Materials & Textures"), LOCTEXT("MaterialFilterTooltip", "Filter by material and texture assets."), "ContentBrowserFilterMaterialAsset", LOCTEXT("MaterialAssetsMenuHeading", "Material Assets") ) );
+	CategoryToMenuMap.Add(EAssetTypeCategories::Blueprint, FCategoryMenu( LOCTEXT("BlueprintFilter", "Blueprints"), LOCTEXT("BlueprintFilterTooltip", "Filter by blueprint assets."), "ContentBrowserFilterBlueprintAsset", LOCTEXT("BlueprintAssetsMenuHeading", "Blueprint Assets") ) );
+	CategoryToMenuMap.Add(EAssetTypeCategories::MaterialsAndTextures, FCategoryMenu(LOCTEXT("MaterialFilter", "Materials & Textures"), LOCTEXT("MaterialFilterTooltip", "Filter by material and texture assets."), "ContentBrowserFilterMaterialAsset", LOCTEXT("MaterialAssetsMenuHeading", "Material Assets")));
 	CategoryToMenuMap.Add(EAssetTypeCategories::Sounds, FCategoryMenu( LOCTEXT("SoundFilter", "Sounds"), LOCTEXT("SoundFilterTooltip", "Filter by sound assets."), "ContentBrowserFilterSoundAsset", LOCTEXT("SoundAssetsMenuHeading", "Sound Assets") ) );
 	CategoryToMenuMap.Add(EAssetTypeCategories::Physics, FCategoryMenu( LOCTEXT("PhysicsFilter", "Physics"), LOCTEXT("PhysicsFilterTooltip", "Filter by physics assets."), "ContentBrowserFilterPhysicsAsset", LOCTEXT("PhysicsAssetsMenuHeading", "Physics Assets") ) );
 	CategoryToMenuMap.Add(EAssetTypeCategories::UI, FCategoryMenu(LOCTEXT("UIFilter", "User Interface"), LOCTEXT("UIFilterTooltip", "Filter by UI assets."), "ContentBrowserFilterUIAsset", LOCTEXT("UIAssetsMenuHeading", "User Interface Assets")));
diff --git a/Engine/Source/Editor/GameplayAbilitiesEditor/Private/AssetTypeActions_GameplayAbilitiesBlueprint.h b/Engine/Source/Editor/GameplayAbilitiesEditor/Private/AssetTypeActions_GameplayAbilitiesBlueprint.h
--- a/Engine/Source/Editor/GameplayAbilitiesEditor/Private/AssetTypeActions_GameplayAbilitiesBlueprint.h
+++ b/Engine/Source/Editor/GameplayAbilitiesEditor/Private/AssetTypeActions_GameplayAbilitiesBlueprint.h
@@ -13,7 +13,7 @@ public:
 	virtual FColor GetTypeColor() const override { return FColor(255,255,255); }
 	virtual UClass* GetSupportedClass() const override;
 	virtual void OpenAssetEditor(const TArray<UObject*>& InObjects, TSharedPtr<class IToolkitHost> EditWithinLevelEditor = TSharedPtr<IToolkitHost>()) override;
-	virtual uint32 GetCategories() override { return EAssetTypeCategories::Gameplay; }
+	virtual uint32 GetCategories() override { return EAssetTypeCategories::Blueprint | EAssetTypeCategories::Gameplay; }
 	// End IAssetTypeActions Implementation
 
 private:
diff --git a/Engine/Source/Editor/UnrealEd/Private/Factories/EditorFactories.cpp b/Engine/Source/Editor/UnrealEd/Private/Factories/EditorFactories.cpp
--- a/Engine/Source/Editor/UnrealEd/Private/Factories/EditorFactories.cpp
+++ b/Engine/Source/Editor/UnrealEd/Private/Factories/EditorFactories.cpp
@@ -6248,8 +6248,7 @@ FName UBlueprintMacroFactory::GetNewAssetThumbnailOverride() const
 
 uint32 UBlueprintMacroFactory::GetMenuCategories() const
 {
-	// Force this factory into the misc category, since it does not belong in the top menu
-	return EAssetTypeCategories::Misc;
+	return EAssetTypeCategories::Blueprint;
 }
 
 UObject* UBlueprintMacroFactory::FactoryCreateNew(UClass* Class, UObject* InParent, FName Name, EObjectFlags Flags, UObject* Context, FFeedbackContext* Warn, FName CallingContext)
@@ -6309,8 +6308,7 @@ FName UBlueprintFunctionLibraryFactory::GetNewAssetThumbnailOverride() const
 
 uint32 UBlueprintFunctionLibraryFactory::GetMenuCategories() const
 {
-	// Force this factory into the misc category, since it does not belong in the top menu
-	return EAssetTypeCategories::Misc;
+	return EAssetTypeCategories::Blueprint;
 }
 
 UObject* UBlueprintFunctionLibraryFactory::FactoryCreateNew(UClass* Class, UObject* InParent, FName Name, EObjectFlags Flags, UObject* Context, FFeedbackContext* Warn, FName CallingContext)
@@ -6365,8 +6363,7 @@ FName UBlueprintInterfaceFactory::GetNewAssetThumbnailOverride() const
 
 uint32 UBlueprintInterfaceFactory::GetMenuCategories() const
 {
-	// Force this factory into the misc category, since it does not belong in the top menu
-	return EAssetTypeCategories::Misc;
+	return EAssetTypeCategories::Blueprint;
 }
 
 UObject* UBlueprintInterfaceFactory::FactoryCreateNew(UClass* Class, UObject* InParent, FName Name, EObjectFlags Flags, UObject* Context, FFeedbackContext* Warn, FName CallingContext)
