Asset View now supports secondary sorting

#UDN: Feature suggestion; Content browser can sort by a second column

#branch UE4

#change
Added two new images for secondary sort ascending and descending.
(Updated styles where needed).

Added new enum EColumnSortPriority: currently only lists Primary and Secondary, but can easily accommodate more*.
(FOnSortModeChanged was modified to also have the priority as a param, fixedup existing usage).

SHeaderRow now also has a SortPriority attribute to specify the SortMode order
(Defaults to Primary if unused).

Modified TUniquePtr so that assigning one from another worked correctly.
(Reviewed by [REDACTED_PERSON]).

SAssetView is the only table that has been modified, so far, to take advantage of the secondary sort. SetMajorityAssetType has been updated to correctly filter out all those sorts which are no longer relevant and bump the priority of the remaining sorts to fill in the ægapsÆ made by non-longer-relevant ones.

FAssetViewSortManager has been overhauled to take SortPriority into consideration when sortingà
Firstly, duplicate comparison structs were removed in favour of single structs which have æascendingÆ as a paramà any remaining duplicate code was removed in favour of an inherited usage. The base struct has an array of æfurther methodsÆ to employ in the result of a tie. Should a tie occur the ænextÆ comparison is done, and so on until itÆs not a tie or we run out of comparisons to perform.*
The manager defaults to having no secondary sort, so it relies on the interface thatÆs using it to provide it.
Whenever a column is assign the code makes sure that itÆs not already assigned to another column and corrects the order to take this into account.

Fixed a bug in FCompareFAssetItemByTagNumericalAscending comparing A with A (instead of B).
Various optimizations to the sort to descrease times

*The only places æSecondaryÆ is referred to in code is in GetSortingBrush (so it can display the correct icon for the sort), and OnTitleClicked (so it can set to correct sort based on input). The sorting code itself has no concept as to a secondary sort, it can support any number of sorting levels so if in future a tertiary (or more) is added, it is only these two function which should need updating as the sort will automatically accommodate it.

diff --git a/Engine/Source/Developer/CollisionAnalyzer/Private/SCollisionAnalyzer.cpp b/Engine/Source/Developer/CollisionAnalyzer/Private/SCollisionAnalyzer.cpp
--- a/Engine/Source/Developer/CollisionAnalyzer/Private/SCollisionAnalyzer.cpp
+++ b/Engine/Source/Developer/CollisionAnalyzer/Private/SCollisionAnalyzer.cpp
@@ -688,7 +688,7 @@ void SCollisionAnalyzer::FilterTextCommitted(const FText& CommentText, ETextComm
 	RebuildFilteredList();
 }
 
-void SCollisionAnalyzer::OnSortByChanged(const FName& ColumnName, EColumnSortMode::Type NewSortMode)
+void SCollisionAnalyzer::OnSortByChanged(const EColumnSortPriority::Type SortPriority, const FName& ColumnName, const EColumnSortMode::Type NewSortMode)
 {
 	SortBy = EQuerySortMode::ByID;
 
diff --git a/Engine/Source/Developer/CollisionAnalyzer/Private/SCollisionAnalyzer.h b/Engine/Source/Developer/CollisionAnalyzer/Private/SCollisionAnalyzer.h
--- a/Engine/Source/Developer/CollisionAnalyzer/Private/SCollisionAnalyzer.h
+++ b/Engine/Source/Developer/CollisionAnalyzer/Private/SCollisionAnalyzer.h
@@ -108,7 +108,7 @@ public:
 	void OnGroupByTagChanged(ESlateCheckBoxState::Type NewState);
 	void OnGroupByOwnerChanged(ESlateCheckBoxState::Type NewState);
 	void FilterTextCommitted(const FText& CommentText, ETextCommit::Type CommitInfo);
-	void OnSortByChanged(const FName& ColumnName, EColumnSortMode::Type NewSortMode);
+	void OnSortByChanged(const EColumnSortPriority::Type SortPriority, const FName& ColumnName, const EColumnSortMode::Type NewSortMode);
 	// Table delegates
 	TSharedRef<ITableRow> QueryTreeGenerateRow(TSharedPtr<FQueryTreeItem> InItem, const TSharedRef<STableViewBase>& OwnerTable);
 	void OnGetChildrenForQueryGroup( TSharedPtr<FQueryTreeItem> InItem, TArray<TSharedPtr<FQueryTreeItem> >& OutChildren );
diff --git a/Engine/Source/Developer/LogVisualizer/Private/SLogVisualizer.cpp b/Engine/Source/Developer/LogVisualizer/Private/SLogVisualizer.cpp
--- a/Engine/Source/Developer/LogVisualizer/Private/SLogVisualizer.cpp
+++ b/Engine/Source/Developer/LogVisualizer/Private/SLogVisualizer.cpp
@@ -1988,7 +1988,7 @@ FString SLogVisualizer::GetLogEntryStatusText() const
 	return TEXT("Pause game with Pause button\nand select log entry to start viewing\nlog's content");
 }
 
-void SLogVisualizer::OnSortByChanged(const FName& ColumnName, EColumnSortMode::Type NewSortMode)
+void SLogVisualizer::OnSortByChanged(const EColumnSortPriority::Type SortPriority, const FName& ColumnName, const EColumnSortMode::Type NewSortMode)
 {
 	SortBy = ELogsSortMode::ByName;
 
diff --git a/Engine/Source/Developer/LogVisualizer/Private/SLogVisualizer.h b/Engine/Source/Developer/LogVisualizer/Private/SLogVisualizer.h
--- a/Engine/Source/Developer/LogVisualizer/Private/SLogVisualizer.h
+++ b/Engine/Source/Developer/LogVisualizer/Private/SLogVisualizer.h
@@ -96,7 +96,7 @@ public:
 	FReply OnRemove();
 	void OnPauseChanged(ESlateCheckBoxState::Type NewState);
 	void FilterTextCommitted(const FText& CommentText, ETextCommit::Type CommitInfo);
-	void OnSortByChanged(const FName& ColumnName, EColumnSortMode::Type NewSortMode);
+	void OnSortByChanged(const EColumnSortPriority::Type SortPriority, const FName& ColumnName, const EColumnSortMode::Type NewSortMode);
 	// Table delegates
 	TSharedRef<ITableRow> LogsListGenerateRow(TSharedPtr<FLogsListItem> InItem, const TSharedRef<STableViewBase>& OwnerTable);
 	void LogsListSelectionChanged(TSharedPtr<FLogsListItem> SelectedItem, ESelectInfo::Type SelectInfo);
diff --git a/Engine/Source/Developer/Profiler/Private/Widgets/SEventGraph.cpp b/Engine/Source/Developer/Profiler/Private/Widgets/SEventGraph.cpp
--- a/Engine/Source/Developer/Profiler/Private/Widgets/SEventGraph.cpp
+++ b/Engine/Source/Developer/Profiler/Private/Widgets/SEventGraph.cpp
@@ -2086,7 +2086,7 @@ void SEventGraph::InitializeAndShowHeaderColumns()
 }
 
 
-void SEventGraph::TreeViewHeaderRow_OnSortModeChanged( const FName& ColumnID, EColumnSortMode::Type SortMode )
+void SEventGraph::TreeViewHeaderRow_OnSortModeChanged( const EColumnSortPriority::Type SortPriority, const FName& ColumnID, const EColumnSortMode::Type SortMode )
 {
 	SetSortModeForColumn( ColumnID, SortMode );
 	TreeView_Refresh();
diff --git a/Engine/Source/Developer/Profiler/Private/Widgets/SEventGraph.h b/Engine/Source/Developer/Profiler/Private/Widgets/SEventGraph.h
--- a/Engine/Source/Developer/Profiler/Private/Widgets/SEventGraph.h
+++ b/Engine/Source/Developer/Profiler/Private/Widgets/SEventGraph.h
@@ -265,7 +265,7 @@ protected:
 	void HighlightHotPath_OnCheckStateChanged( ESlateCheckBoxState::Type InState );
 
 	void InitializeAndShowHeaderColumns();
-	void TreeViewHeaderRow_OnSortModeChanged( const FName& ColumnID, EColumnSortMode::Type );
+	void TreeViewHeaderRow_OnSortModeChanged( const EColumnSortPriority::Type SortPriority, const FName& ColumnID, const EColumnSortMode::Type SortMode );
 
 	void SetSortModeForColumn( const FName& ColumnID, EColumnSortMode::Type SortMode );
 
diff --git a/Engine/Source/Developer/TaskGraph/Private/SEventsTree.cpp b/Engine/Source/Developer/TaskGraph/Private/SEventsTree.cpp
--- a/Engine/Source/Developer/TaskGraph/Private/SEventsTree.cpp
+++ b/Engine/Source/Developer/TaskGraph/Private/SEventsTree.cpp
@@ -278,7 +278,7 @@ void SEventsTree::OnGetChildrenForEventsList( TSharedPtr<FVisualizerEvent> InIte
 	OutChildren = InItem->Children;
 }
 
-void SEventsTree::OnColumnSortModeChanged( const FName& ColumnId, EColumnSortMode::Type InSortMode )
+void SEventsTree::OnColumnSortModeChanged( const EColumnSortPriority::Type SortPriority, const FName& ColumnId, const EColumnSortMode::Type InSortMode )
 {
 	SortByColumn = ColumnId;
 	SortMode = InSortMode;
diff --git a/Engine/Source/Developer/TaskGraph/Public/SEventsTree.h b/Engine/Source/Developer/TaskGraph/Public/SEventsTree.h
--- a/Engine/Source/Developer/TaskGraph/Public/SEventsTree.h
+++ b/Engine/Source/Developer/TaskGraph/Public/SEventsTree.h
@@ -139,7 +139,7 @@ protected:
 	 *
 	 * @param InSortMode New sorting mode
 	 */
-	void OnColumnSortModeChanged( const FName& ColumnId, EColumnSortMode::Type InSortMode );
+	void OnColumnSortModeChanged( const EColumnSortPriority::Type SortPriority, const FName& ColumnId, const EColumnSortMode::Type InSortMode );
 
 	/**
 	 * @return Returns the current sort mode of the specified column
diff --git a/Engine/Source/Editor/ContentBrowser/Private/AssetViewSortManager.cpp b/Engine/Source/Editor/ContentBrowser/Private/AssetViewSortManager.cpp
--- a/Engine/Source/Editor/ContentBrowser/Private/AssetViewSortManager.cpp
+++ b/Engine/Source/Editor/ContentBrowser/Private/AssetViewSortManager.cpp
@@ -3,352 +3,297 @@
 #include "ContentBrowserPCH.h"
 #include "AssetViewTypes.h"
 
-static bool CompareFolderFirst(const TSharedPtr<FAssetViewItem>& A, const TSharedPtr<FAssetViewItem>& B, bool bAscending, bool& bComparisonResult)
+struct FCompareFAssetItemBase
 {
-	if(A->GetType() == EAssetItemType::Folder && B->GetType() == EAssetItemType::Folder)
+public:
+	/** Default constructor */
+	FCompareFAssetItemBase(bool bInAscending, const FName& InTag) : bAscending(bInAscending), Tag(InTag) {}
+
+	/** Sort function */
+	FORCEINLINE bool operator()(const TSharedPtr<FAssetViewItem>& A, const TSharedPtr<FAssetViewItem>& B) const
 	{
-		if(bAscending)
-		{
-			bComparisonResult = StaticCastSharedPtr<FAssetViewFolder>(A)->FolderName.CompareTo(StaticCastSharedPtr<FAssetViewFolder>(B)->FolderName) < 0;
-		}
-		else
+		bool bFolderComparisonResult = false;
+		if (CompareFolderFirst(A, B, bFolderComparisonResult))
 		{
-			bComparisonResult = StaticCastSharedPtr<FAssetViewFolder>(A)->FolderName.CompareTo(StaticCastSharedPtr<FAssetViewFolder>(B)->FolderName) > 0;
+			return bFolderComparisonResult;
 		}
-		return true;
+		return Compare(A, B);
 	}
-	else if(A->GetType() == EAssetItemType::Folder)
+
+	/** Get the next comparisons array */
+	FORCEINLINE TArray<TUniquePtr<FCompareFAssetItemBase>>& GetNextComparisons() const
 	{
-		bComparisonResult = bAscending;
-		return true;
+		return NextComparisons;
 	}
-	else if(B->GetType() == EAssetItemType::Folder)
+
+	/** Set the next comparisons array */
+	FORCEINLINE bool SetNextComparisons(TArray<TUniquePtr<FCompareFAssetItemBase>>& InNextComparisons) const
 	{
-		bComparisonResult = !bAscending;
-		return true;
+		check(NextComparisons.Num() == 0);
+		NextComparisons = MoveTemp(InNextComparisons);
 	}
 
-	return false;
-}
+protected:
+	/** Comparison function, needs to be implemented by inheriting structs - by default just forwards to the next comparitor */
+	FORCEINLINE virtual bool Compare(const TSharedPtr<FAssetViewItem>& A, const TSharedPtr<FAssetViewItem>& B) const { return CompareNext(A, B); }
 
-struct FCompareFAssetItemByNameAscending
-{
-	FORCEINLINE bool operator()( const TSharedPtr<FAssetViewItem>& A, const TSharedPtr<FAssetViewItem>& B ) const
+private:
+	/** Compare the folders of both assets */
+	FORCEINLINE bool CompareFolderFirst(const TSharedPtr<FAssetViewItem>& A, const TSharedPtr<FAssetViewItem>& B, bool& bComparisonResult) const
 	{
-		bool bFolderComparisonResult = false;
-		if(CompareFolderFirst(A, B, true, bFolderComparisonResult))
+		if (A->GetType() == EAssetItemType::Folder)
 		{
-			return bFolderComparisonResult;
+			if (B->GetType() == EAssetItemType::Folder)
+			{
+				const FText& ValueA = StaticCastSharedPtr<FAssetViewFolder>(A)->FolderName;
+				const FText& ValueB = StaticCastSharedPtr<FAssetViewFolder>(B)->FolderName;
+				const int32 Result = ValueA.CompareTo(ValueB);
+				bComparisonResult = bAscending ? Result < 0 : Result > 0;
+			}
+			else
+			{
+				bComparisonResult = bAscending;
+			}
+			return true;
 		}
-
-		return StaticCastSharedPtr<FAssetViewAsset>(A)->Data.AssetName.Compare(StaticCastSharedPtr<FAssetViewAsset>(B)->Data.AssetName) < 0;
-	}
-};
-
-struct FCompareFAssetItemByNameDescending
-{
-	FORCEINLINE bool operator()( const TSharedPtr<FAssetViewItem>& A, const TSharedPtr<FAssetViewItem>& B ) const
-	{
-		bool bFolderComparisonResult = false;
-		if(CompareFolderFirst(A, B, false, bFolderComparisonResult))
+		else if (B->GetType() == EAssetItemType::Folder)
 		{
-			return bFolderComparisonResult;
+			bComparisonResult = !bAscending;
+			return true;
 		}
-
-		return StaticCastSharedPtr<FAssetViewAsset>(A)->Data.AssetName.Compare(StaticCastSharedPtr<FAssetViewAsset>(B)->Data.AssetName) > 0;
+		return false;
 	}
-};
 
-struct FCompareFAssetItemByClassAscending
-{
-	FORCEINLINE bool operator()( const TSharedPtr<FAssetViewItem>& A, const TSharedPtr<FAssetViewItem>& B ) const
+	/** Comparison function in the event of a tie*/
+	FORCEINLINE bool CompareNext(const TSharedPtr<FAssetViewItem>& A, const TSharedPtr<FAssetViewItem>& B) const
 	{
-		bool bFolderComparisonResult = false;
-		if(CompareFolderFirst(A, B, true, bFolderComparisonResult))
+		if (NextComparisons.Num() > 0)
 		{
-			return bFolderComparisonResult;
-		}
+			TUniquePtr<FCompareFAssetItemBase> CompareNext = MoveTemp(NextComparisons[0]);
+			check(CompareNext);
+			NextComparisons.RemoveAt(0);
 
-		return StaticCastSharedPtr<FAssetViewAsset>(A)->Data.AssetClass.Compare(StaticCastSharedPtr<FAssetViewAsset>(B)->Data.AssetClass) < 0;
-	}
-};
+			// Move all the remaining comparisons to the next comparitor
+			CompareNext->SetNextComparisons(GetNextComparisons());
+			const bool Result = CompareNext->Compare(A, B);
 
-struct FCompareFAssetItemByClassDescending
-{
-	FORCEINLINE bool operator()( const TSharedPtr<FAssetViewItem>& A, const TSharedPtr<FAssetViewItem>& B ) const
-	{
-		bool bFolderComparisonResult = false;
-		if(CompareFolderFirst(A, B, false, bFolderComparisonResult))
-		{
-			return bFolderComparisonResult;
-		}
+			// Move the comparisons back for further ties
+			SetNextComparisons(CompareNext->GetNextComparisons());
+			NextComparisons.Insert(MoveTemp(CompareNext), 0);
 
-		return StaticCastSharedPtr<FAssetViewAsset>(A)->Data.AssetClass.Compare(StaticCastSharedPtr<FAssetViewAsset>(B)->Data.AssetClass) > 0;
+			return Result;
+		}
+		return true;	// default to true if a tie
 	}
-};
 
-struct FCompareFAssetItemByPathAscending
-{
-	FORCEINLINE bool operator()( const TSharedPtr<FAssetViewItem>& A, const TSharedPtr<FAssetViewItem>& B ) const
-	{
-		bool bFolderComparisonResult = false;
-		if(CompareFolderFirst(A, B, true, bFolderComparisonResult))
-		{
-			return bFolderComparisonResult;
-		}
+protected:
+	/** Whether to sort ascending or descending */
+	bool bAscending;
 
-		return StaticCastSharedPtr<FAssetViewAsset>(A)->Data.PackagePath.Compare(StaticCastSharedPtr<FAssetViewAsset>(B)->Data.PackagePath) < 0;
-	}
+	/** The tag type we need to compare next */
+	FName Tag;
+
+private:
+	/** The follow up sorting methods in the event of a tie */
+	mutable TArray<TUniquePtr<FCompareFAssetItemBase>> NextComparisons;
 };
 
-struct FCompareFAssetItemByPathDescending
+struct FCompareFAssetItemByName : public FCompareFAssetItemBase
 {
-	FORCEINLINE bool operator()( const TSharedPtr<FAssetViewItem>& A, const TSharedPtr<FAssetViewItem>& B ) const
+public:
+	FCompareFAssetItemByName(bool bInAscending, const FName& InTag) : FCompareFAssetItemBase(bInAscending, InTag) {}
+
+protected:
+	FORCEINLINE virtual bool Compare(const TSharedPtr<FAssetViewItem>& A, const TSharedPtr<FAssetViewItem>& B) const override
 	{
-		bool bFolderComparisonResult = false;
-		if(CompareFolderFirst(A, B, false, bFolderComparisonResult))
+		const FName& ValueA = StaticCastSharedPtr<FAssetViewAsset>(A)->Data.AssetName;
+		const FName& ValueB = StaticCastSharedPtr<FAssetViewAsset>(B)->Data.AssetName;
+		const int32 Result = ValueA.Compare(ValueB);
+		if (Result < 0)
 		{
-			return bFolderComparisonResult;
+			return bAscending;
 		}
-
-		return StaticCastSharedPtr<FAssetViewAsset>(A)->Data.PackagePath.Compare(StaticCastSharedPtr<FAssetViewAsset>(B)->Data.PackagePath) > 0;
+		else if (Result > 0)
+		{
+			return !bAscending;
+		}
+		return FCompareFAssetItemBase::Compare(A, B);
 	}
 };
 
-struct FCompareFAssetItemByTagAscending
+struct FCompareFAssetItemByClass : public FCompareFAssetItemBase
 {
-	FCompareFAssetItemByTagAscending(FName InTag) : Tag(InTag) {}
+public:
+	FCompareFAssetItemByClass(bool bInAscending, const FName& InTag) : FCompareFAssetItemBase(bInAscending, InTag) {}
 
-	FORCEINLINE bool operator()( const TSharedPtr<FAssetViewItem>& A, const TSharedPtr<FAssetViewItem>& B ) const
+protected:
+	FORCEINLINE virtual bool Compare(const TSharedPtr<FAssetViewItem>& A, const TSharedPtr<FAssetViewItem>& B) const override
 	{
-		bool bFolderComparisonResult = false;
-		if(CompareFolderFirst(A, B, true, bFolderComparisonResult))
+		const FName& ValueA = StaticCastSharedPtr<FAssetViewAsset>(A)->Data.AssetClass;
+		const FName& ValueB = StaticCastSharedPtr<FAssetViewAsset>(B)->Data.AssetClass;
+		const int32 Result = ValueA.Compare(ValueB);
+		if (Result < 0)
 		{
-			return bFolderComparisonResult;
-		}
-
-		const FString* AValue = StaticCastSharedPtr<FAssetViewAsset>(A)->Data.TagsAndValues.Find(Tag);
-		if ( AValue )
-		{
-			const FString* BValue = StaticCastSharedPtr<FAssetViewAsset>(B)->Data.TagsAndValues.Find(Tag);
-			if ( BValue )
-			{
-				return *AValue < *BValue;
-			}
-			else
-			{
-				return false;
-			}
+			return bAscending;
 		}
-		else
+		else if (Result > 0)
 		{
-			return true;
+			return !bAscending;
 		}
+		return FCompareFAssetItemBase::Compare(A, B);
 	}
-
-	FName Tag;
 };
 
-struct FCompareFAssetItemByTagDescending
+struct FCompareFAssetItemByPath : public FCompareFAssetItemBase
 {
-	FCompareFAssetItemByTagDescending(FName InTag) : Tag(InTag) {}
+public:
+	FCompareFAssetItemByPath(bool bInAscending, const FName& InTag) : FCompareFAssetItemBase(bInAscending, InTag) {}
 
-	FORCEINLINE bool operator()( const TSharedPtr<FAssetViewItem>& A, const TSharedPtr<FAssetViewItem>& B ) const
+protected:
+	FORCEINLINE virtual bool Compare(const TSharedPtr<FAssetViewItem>& A, const TSharedPtr<FAssetViewItem>& B) const override
 	{
-		bool bFolderComparisonResult = false;
-		if(CompareFolderFirst(A, B, false, bFolderComparisonResult))
+		const FName& ValueA = StaticCastSharedPtr<FAssetViewAsset>(A)->Data.PackagePath;
+		const FName& ValueB = StaticCastSharedPtr<FAssetViewAsset>(B)->Data.PackagePath;
+		const int32 Result = ValueA.Compare(ValueB);
+		if (Result < 0)
 		{
-			return bFolderComparisonResult;
+			return bAscending;
 		}
-
-		const FString* AValue = StaticCastSharedPtr<FAssetViewAsset>(A)->Data.TagsAndValues.Find(Tag);
-		if ( AValue )
+		else if (Result > 0)
 		{
-			const FString* BValue = StaticCastSharedPtr<FAssetViewAsset>(B)->Data.TagsAndValues.Find(Tag);
-			if ( BValue )
-			{
-				return *AValue > *BValue;
-			}
-			else
-			{
-				return true;
-			}
-		}
-		else
-		{
-			return false;
+			return !bAscending;
 		}
+		return FCompareFAssetItemBase::Compare(A, B);
 	}
-
-	FName Tag;
 };
 
-struct FCompareFAssetItemByTagNumericalAscending
+struct FCompareFAssetItemByTag : public FCompareFAssetItemBase
 {
-	FCompareFAssetItemByTagNumericalAscending(FName InTag) : Tag(InTag) {}
+public:
+	FCompareFAssetItemByTag(bool bInAscending, const FName& InTag) : FCompareFAssetItemBase(bInAscending, InTag) {}
 
-	FORCEINLINE bool operator()( const TSharedPtr<FAssetViewItem>& A, const TSharedPtr<FAssetViewItem>& B ) const
+protected:
+	FORCEINLINE virtual bool Compare(const TSharedPtr<FAssetViewItem>& A, const TSharedPtr<FAssetViewItem>& B) const override
 	{
-		bool bFolderComparisonResult = false;
-		if(CompareFolderFirst(A, B, true, bFolderComparisonResult))
+		// Depending if we're sorting ascending or descending it's quicker to flip the compares incase tags are missing
+		const FString* Value1 = bAscending ? StaticCastSharedPtr<FAssetViewAsset>(A)->Data.TagsAndValues.Find(Tag) : StaticCastSharedPtr<FAssetViewAsset>(B)->Data.TagsAndValues.Find(Tag);
+		const FString* Value2 = NULL;
+		if (!Value1)
 		{
-			return bFolderComparisonResult;
+			return true;
 		}
-
-		const FString* AValue = StaticCastSharedPtr<FAssetViewAsset>(A)->Data.TagsAndValues.Find(Tag);
-		if ( AValue )
+		else
 		{
-			const FString* BValue = StaticCastSharedPtr<FAssetViewAsset>(A)->Data.TagsAndValues.Find(Tag);
-			if ( BValue )
-			{
-				return FCString::Atof(**AValue) < FCString::Atof(**BValue);
-			}
-			else
+			Value2 = bAscending ? StaticCastSharedPtr<FAssetViewAsset>(B)->Data.TagsAndValues.Find(Tag) : StaticCastSharedPtr<FAssetViewAsset>(A)->Data.TagsAndValues.Find(Tag);
+			if (!Value2)
 			{
 				return false;
 			}
 		}
-		else
+
+		const int32 Result = Value1->Compare(*Value2, ESearchCase::IgnoreCase);
+		if (Result < 0)
 		{
 			return true;
 		}
+		else if (Result > 0)
+		{
+			return false;
+		}
+		return FCompareFAssetItemBase::Compare(A, B);
 	}
-
-	FName Tag;
 };
 
-struct FCompareFAssetItemByTagNumericalDescending
+struct FCompareFAssetItemByTagNumerical : public FCompareFAssetItemBase
 {
-	FCompareFAssetItemByTagNumericalDescending(FName InTag) : Tag(InTag) {}
+public:
+	FCompareFAssetItemByTagNumerical(bool bInAscending, const FName& InTag) : FCompareFAssetItemBase(bInAscending, InTag) {}
 
-	FORCEINLINE bool operator()( const TSharedPtr<FAssetViewItem>& A, const TSharedPtr<FAssetViewItem>& B ) const
+protected:
+	FORCEINLINE virtual bool Compare(const TSharedPtr<FAssetViewItem>& A, const TSharedPtr<FAssetViewItem>& B) const override
 	{
-		bool bFolderComparisonResult = false;
-		if(CompareFolderFirst(A, B, false, bFolderComparisonResult))
+		// Depending if we're sorting ascending or descending it's quicker to flip the compares incase tags are missing
+		const FString* Value1 = bAscending ? StaticCastSharedPtr<FAssetViewAsset>(A)->Data.TagsAndValues.Find(Tag) : StaticCastSharedPtr<FAssetViewAsset>(B)->Data.TagsAndValues.Find(Tag);
+		const FString* Value2 = NULL;
+		if (!Value1)
 		{
-			return bFolderComparisonResult;
+			return true;
 		}
-
-		const FString* AValue = StaticCastSharedPtr<FAssetViewAsset>(A)->Data.TagsAndValues.Find(Tag);
-		if ( AValue )
+		else
 		{
-			const FString* BValue = StaticCastSharedPtr<FAssetViewAsset>(B)->Data.TagsAndValues.Find(Tag);
-			if ( BValue )
+			Value2 = bAscending ? StaticCastSharedPtr<FAssetViewAsset>(B)->Data.TagsAndValues.Find(Tag) : StaticCastSharedPtr<FAssetViewAsset>(A)->Data.TagsAndValues.Find(Tag);
+			if (!Value2)
 			{
-				return FCString::Atof(**AValue) > FCString::Atof(**BValue);
-			}
-			else
-			{
-				return true;
+				return false;
 			}
 		}
-		else
+
+		const float Num1 = FCString::Atof(**Value1);
+		const float Num2 = FCString::Atof(**Value2);
+		if (Num1 < Num2)
+		{
+			return true;
+		}
+		else if (Num1 > Num2)
 		{
 			return false;
 		}
+		return FCompareFAssetItemBase::Compare(A, B);
 	}
-
-	FName Tag;
 };
 
-struct FCompareFAssetItemByTagDimensionalAscending
+struct FCompareFAssetItemByTagDimensional : public FCompareFAssetItemBase
 {
-	FCompareFAssetItemByTagDimensionalAscending(FName InTag) : Tag(InTag) {}
+public:
+	FCompareFAssetItemByTagDimensional(bool bInAscending, const FName& InTag) : FCompareFAssetItemBase(bInAscending, InTag) {}
 
-	FORCEINLINE bool operator()( const TSharedPtr<FAssetViewItem>& A, const TSharedPtr<FAssetViewItem>& B ) const
+protected:
+	FORCEINLINE virtual bool Compare(const TSharedPtr<FAssetViewItem>& A, const TSharedPtr<FAssetViewItem>& B) const override
 	{
-		bool bFolderComparisonResult = false;
-		if(CompareFolderFirst(A, B, true, bFolderComparisonResult))
+		// Depending if we're sorting ascending or descending it's quicker to flip the compares incase tags are missing
+		const FString* Value1 = bAscending ? StaticCastSharedPtr<FAssetViewAsset>(A)->Data.TagsAndValues.Find(Tag) : StaticCastSharedPtr<FAssetViewAsset>(B)->Data.TagsAndValues.Find(Tag);
+		const FString* Value2 = NULL;
+		if (!Value1)
 		{
-			return bFolderComparisonResult;
+			return true;
 		}
-
-		const FString* AValue = StaticCastSharedPtr<FAssetViewAsset>(A)->Data.TagsAndValues.Find(Tag);
-		if ( AValue )
+		else
 		{
-			const FString* BValue = StaticCastSharedPtr<FAssetViewAsset>(B)->Data.TagsAndValues.Find(Tag);
-			if ( BValue )
-			{
-				TArray<FString> ATokens;
-				(*AValue).ParseIntoArray(&ATokens, TEXT("x"), true);
-				float ATotal = 1.f;
-				for ( auto ATokenIt = ATokens.CreateConstIterator(); ATokenIt; ++ATokenIt )
-				{
-					ATotal *= FCString::Atof(**ATokenIt);
-				}
-
-				TArray<FString> BTokens;
-				(*BValue).ParseIntoArray(&BTokens, TEXT("x"), true);
-				float BTotal = 1.f;
-				for ( auto BTokenIt = BTokens.CreateConstIterator(); BTokenIt; ++BTokenIt )
-				{
-					BTotal *= FCString::Atof(**BTokenIt);
-				}
-
-				return ATotal < BTotal;
-			}
-			else
+			Value2 = bAscending ? StaticCastSharedPtr<FAssetViewAsset>(B)->Data.TagsAndValues.Find(Tag) : StaticCastSharedPtr<FAssetViewAsset>(A)->Data.TagsAndValues.Find(Tag);
+			if (!Value2)
 			{
 				return false;
 			}
 		}
-		else
-		{
-			return true;
-		}
-	}
-
-	FName Tag;
-};
-
-struct FCompareFAssetItemByTagDimensionalDescending
-{
-	FCompareFAssetItemByTagDimensionalDescending(FName InTag) : Tag(InTag) {}
 
-	FORCEINLINE bool operator()( const TSharedPtr<FAssetViewItem>& A, const TSharedPtr<FAssetViewItem>& B ) const
-	{
-		bool bFolderComparisonResult = false;
-		if(CompareFolderFirst(A, B, false, bFolderComparisonResult))
+		float Num1 = 1.f;
 		{
-			return bFolderComparisonResult;
-		}
-
-		const FString* AValue = StaticCastSharedPtr<FAssetViewAsset>(A)->Data.TagsAndValues.Find(Tag);
-		if ( AValue )
-		{
-			const FString* BValue = StaticCastSharedPtr<FAssetViewAsset>(B)->Data.TagsAndValues.Find(Tag);
-			if ( BValue )
+			TArray<FString> Tokens1;
+			Value1->ParseIntoArray(&Tokens1, TEXT("x"), true);
+			for (auto TokenIt1 = Tokens1.CreateConstIterator(); TokenIt1; ++TokenIt1)
 			{
-				TArray<FString> ATokens;
-				(*AValue).ParseIntoArray(&ATokens, TEXT("x"), true);
-				float ATotal = 1.f;
-				for ( auto ATokenIt = ATokens.CreateConstIterator(); ATokenIt; ++ATokenIt )
-				{
-					ATotal *= FCString::Atof(**ATokenIt);
-				}
-
-				TArray<FString> BTokens;
-				(*BValue).ParseIntoArray(&BTokens, TEXT("x"), true);
-				float BTotal = 1.f;
-				for ( auto BTokenIt = BTokens.CreateConstIterator(); BTokenIt; ++BTokenIt )
-				{
-					BTotal *= FCString::Atof(**BTokenIt);
-				}
-
-				return ATotal > BTotal;
+				Num1 *= FCString::Atof(**TokenIt1);
 			}
-			else
+		}
+		float Num2 = 1.f;
+		{
+			TArray<FString> Tokens2;
+			Value2->ParseIntoArray(&Tokens2, TEXT("x"), true);
+			for (auto TokenIt2 = Tokens2.CreateConstIterator(); TokenIt2; ++TokenIt2)
 			{
-				return true;
+				Num2 *= FCString::Atof(**TokenIt2);
 			}
 		}
-		else
+		if (Num1 < Num2)
+		{
+			return true;
+		}
+		else if (Num1 > Num2)
 		{
 			return false;
 		}
+		return FCompareFAssetItemBase::Compare(A, B);
 	}
-
-	FName Tag;
 };
 
 const FName FAssetViewSortManager::NameColumnId = "Name";
@@ -357,140 +302,185 @@ const FName FAssetViewSortManager::PathColumnId = "Path";
 
 FAssetViewSortManager::FAssetViewSortManager()
 {
-	SortColumnId = NameColumnId;
-	SortMode = EColumnSortMode::Ascending;
+	ResetSort();
 }
 
-void FAssetViewSortManager::SortList(TArray<TSharedPtr<FAssetViewItem>>& AssetItems, FName MajorityAssetType) const
+void FAssetViewSortManager::ResetSort()
 {
-	if ( SortColumnId == NameColumnId )
+	SortColumnId[EColumnSortPriority::Primary] = NameColumnId;
+	SortMode[EColumnSortPriority::Primary] = EColumnSortMode::Ascending;
+	for (int32 PriorityIdx = 1; PriorityIdx < EColumnSortPriority::Max; PriorityIdx++)
 	{
-		if ( SortMode == EColumnSortMode::Ascending )
-		{
-			AssetItems.Sort( FCompareFAssetItemByNameAscending() );
-		}
-		else if ( SortMode == EColumnSortMode::Descending )
-		{
-			AssetItems.Sort( FCompareFAssetItemByNameDescending() );
-		}
+		SortColumnId[PriorityIdx] = NAME_None;
+		SortMode[PriorityIdx] = EColumnSortMode::None;
 	}
-	else if ( SortColumnId == ClassColumnId )
+}
+
+void FAssetViewSortManager::SortList(TArray<TSharedPtr<FAssetViewItem>>& AssetItems, const FName& MajorityAssetType) const
+{
+	//double SortListStartTime = FPlatformTime::Seconds();
+
+	TArray<TUniquePtr<FCompareFAssetItemBase>> SortMethod;
+	for (int32 PriorityIdx = 0; PriorityIdx < EColumnSortPriority::Max; PriorityIdx++)
 	{
-		if ( SortMode == EColumnSortMode::Ascending )
+		const bool bAscending(SortMode[PriorityIdx] == EColumnSortMode::Ascending);
+		const FName& Tag(SortColumnId[PriorityIdx]);
+
+		if (Tag == NAME_None)
 		{
-			AssetItems.Sort( FCompareFAssetItemByClassAscending() );
+			break;
 		}
-		else if ( SortMode == EColumnSortMode::Descending )
+
+		if (Tag == NameColumnId)
 		{
-			AssetItems.Sort( FCompareFAssetItemByClassDescending() );
+			SortMethod.Add(MakeUnique<FCompareFAssetItemByName>(bAscending, Tag));
 		}
-	}
-	else if ( SortColumnId == PathColumnId )
-	{
-		if ( SortMode == EColumnSortMode::Ascending )
+		else if (Tag == ClassColumnId)
 		{
-			AssetItems.Sort( FCompareFAssetItemByPathAscending() );
+			SortMethod.Add(MakeUnique<FCompareFAssetItemByClass>(bAscending, Tag));
 		}
-		else if ( SortMode == EColumnSortMode::Descending )
+		else if (Tag == PathColumnId)
 		{
-			AssetItems.Sort( FCompareFAssetItemByPathDescending() );
+			SortMethod.Add(MakeUnique<FCompareFAssetItemByPath>(bAscending, Tag));
 		}
-	}
-	else
-	{
-		// Since this ColumnId is not one of preset columns, sort by asset registry tag
-		UObject::FAssetRegistryTag::ETagType TagType = UObject::FAssetRegistryTag::TT_Alphabetical;
-		if ( MajorityAssetType != NAME_None )
+		else
 		{
-			UClass* Class = FindObject<UClass>(ANY_PACKAGE, *MajorityAssetType.ToString());
-			if ( Class )
+			// Since this SortData.Tag is not one of preset columns, sort by asset registry tag
+			UObject::FAssetRegistryTag::ETagType TagType = UObject::FAssetRegistryTag::TT_Alphabetical;
+			if (MajorityAssetType != NAME_None)
 			{
-				UObject* CDO = Class->GetDefaultObject();
-
-				if ( CDO )
+				UClass* Class = FindObject<UClass>(ANY_PACKAGE, *MajorityAssetType.ToString());
+				if (Class)
 				{
-					TArray<UObject::FAssetRegistryTag> TagList;
-					CDO->GetAssetRegistryTags(TagList);
-
-					for ( auto TagIt = TagList.CreateConstIterator(); TagIt; ++TagIt )
+					UObject* CDO = Class->GetDefaultObject();
+					if (CDO)
 					{
-						if ( TagIt->Name == SortColumnId )
+						TArray<UObject::FAssetRegistryTag> TagList;
+						CDO->GetAssetRegistryTags(TagList);
+
+						for (auto TagIt = TagList.CreateConstIterator(); TagIt; ++TagIt)
 						{
-							TagType = TagIt->Type;
-							break;
+							if (TagIt->Name == Tag)
+							{
+								TagType = TagIt->Type;
+								break;
+							}
 						}
 					}
 				}
 			}
-		}
-		
-		if ( TagType == UObject::FAssetRegistryTag::TT_Numerical )
-		{
-			// The property is a number, compare using atof
-			if ( SortMode == EColumnSortMode::Ascending )
+
+			if (TagType == UObject::FAssetRegistryTag::TT_Numerical)
+			{
+				// The property is a Num2er, compare using atof
+				SortMethod.Add(MakeUnique<FCompareFAssetItemByTagNumerical>(bAscending, Tag));
+			}
+			else if (TagType == UObject::FAssetRegistryTag::TT_Dimensional)
 			{
-				AssetItems.Sort( FCompareFAssetItemByTagNumericalAscending(SortColumnId) );
+				// The property is a series of Num2ers representing dimensions, compare by using atof for each Num2er, delimited by an "x"
+				SortMethod.Add(MakeUnique<FCompareFAssetItemByTagDimensional>(bAscending, Tag));
 			}
-			else if ( SortMode == EColumnSortMode::Descending )
+			else
 			{
-				AssetItems.Sort( FCompareFAssetItemByTagNumericalDescending(SortColumnId) );
+				// Unknown or alphabetical, sort alphabetically either way
+				SortMethod.Add(MakeUnique<FCompareFAssetItemByTag>(bAscending, Tag));
 			}
 		}
- 		else if ( TagType == UObject::FAssetRegistryTag::TT_Dimensional )
- 		{
- 			// The property is a series of numbers representing dimensions, compare by using atof for each number, delimited by an "x"
- 			if ( SortMode == EColumnSortMode::Ascending )
- 			{
- 				AssetItems.Sort( FCompareFAssetItemByTagDimensionalAscending(SortColumnId) );
- 			}
- 			else if ( SortMode == EColumnSortMode::Descending )
- 			{
- 				AssetItems.Sort( FCompareFAssetItemByTagDimensionalDescending(SortColumnId) );
- 			}
- 		}
-		else
+	}
+
+	// Sort the list...
+	if (SortMethod.Num() > 0)
+	{
+		TUniquePtr<FCompareFAssetItemBase> PrimarySortMethod = MoveTemp(SortMethod[EColumnSortPriority::Primary]);
+		check(PrimarySortMethod);
+		SortMethod.RemoveAt(0);
+
+		// Move all the comparisons to the primary sort method
+		PrimarySortMethod->SetNextComparisons(SortMethod);
+		AssetItems.Sort(*(PrimarySortMethod.Get()));
+
+		// Move the comparisons back for ease of cleanup
+		SortMethod = MoveTemp(PrimarySortMethod->GetNextComparisons());
+		SortMethod.Insert(MoveTemp(PrimarySortMethod), 0);
+	}
+
+	// Cleanup the methods we no longer need.
+	for (int32 PriorityIdx = 0; PriorityIdx < SortMethod.Num(); PriorityIdx++)
+	{
+		SortMethod[PriorityIdx].Reset();
+	}
+	SortMethod.Empty();
+
+	//UE_LOG(LogContentBrowser, Warning/*VeryVerbose*/, TEXT("FAssetViewSortManager Sort Time: %0.4f seconds."), FPlatformTime::Seconds() - SortListStartTime);
+}
+
+void FAssetViewSortManager::SetSortColumnId(const EColumnSortPriority::Type InSortPriority, const FName& InColumnId)
+{
+	check(InSortPriority < EColumnSortPriority::Max);
+	SortColumnId[InSortPriority] = InColumnId;
+
+	// Prevent the same ColumnId being assigned to multiple columns
+	bool bOrderChanged = false;
+	for (int32 PriorityIdxA = 0; PriorityIdxA < EColumnSortPriority::Max; PriorityIdxA++)
+	{
+		for (int32 PriorityIdxB = 0; PriorityIdxB < EColumnSortPriority::Max; PriorityIdxB++)
 		{
-			// Unknown or alphabetical, sort alphabetically either way
-			if ( SortMode == EColumnSortMode::Ascending )
+			if (PriorityIdxA != PriorityIdxB)
 			{
-				AssetItems.Sort( FCompareFAssetItemByTagAscending(SortColumnId) );
+				if (SortColumnId[PriorityIdxA] == SortColumnId[PriorityIdxB] && SortColumnId[PriorityIdxB] != NAME_None)
+				{
+					SortColumnId[PriorityIdxB] = NAME_None;
+					bOrderChanged = true;
+				}
 			}
-			else if ( SortMode == EColumnSortMode::Descending )
+		}
+	}
+	if (bOrderChanged)
+	{
+		// If the order has changed, we need to remove any unneeded sorts by bumping the priority of the remaining valid ones
+		for (int32 PriorityIdxA = 0, PriorityNum = 0; PriorityNum < EColumnSortPriority::Max - 1; PriorityNum++, PriorityIdxA++)
+		{
+			if (SortColumnId[PriorityIdxA] == NAME_None)
 			{
-				AssetItems.Sort( FCompareFAssetItemByTagDescending(SortColumnId) );
+				for (int32 PriorityIdxB = PriorityIdxA; PriorityIdxB < EColumnSortPriority::Max - 1; PriorityIdxB++)
+				{
+					SortColumnId[PriorityIdxB] = SortColumnId[PriorityIdxB + 1];
+					SortMode[PriorityIdxB] = SortMode[PriorityIdxB + 1];
+				}
+				SortColumnId[EColumnSortPriority::Max - 1] = NAME_None;
+				PriorityIdxA--;
 			}
 		}
 	}
 }
 
-void FAssetViewSortManager::SetSortColumnId( const FName& ColumnId )
-{
-	SortColumnId = ColumnId;
-}
-
-void FAssetViewSortManager::SetSortMode( const EColumnSortMode::Type InSortMode )
+void FAssetViewSortManager::SetSortMode(const EColumnSortPriority::Type InSortPriority, const EColumnSortMode::Type InSortMode)
 {
-	SortMode = InSortMode;
+	check(InSortPriority < EColumnSortPriority::Max);
+	SortMode[InSortPriority] = InSortMode;
 }
 
-void FAssetViewSortManager::SetOrToggleSortColumn(FName ColumnId)
+bool FAssetViewSortManager::SetOrToggleSortColumn(const EColumnSortPriority::Type InSortPriority, const FName& InColumnId)
 {
-	if ( SortColumnId != ColumnId )
+	check(InSortPriority < EColumnSortPriority::Max);
+	if (SortColumnId[InSortPriority] != InColumnId)
 	{
 		// Clicked a new column, default to ascending
-		SortColumnId = ColumnId;
-		SortMode = EColumnSortMode::Ascending;
+		SortColumnId[InSortPriority] = InColumnId;
+		SortMode[InSortPriority] = EColumnSortMode::Ascending;
+		return true;
 	}
 	else
 	{
 		// Clicked the current column, toggle sort mode
-		if ( SortMode == EColumnSortMode::Ascending )
+		if (SortMode[InSortPriority] == EColumnSortMode::Ascending)
 		{
-			SortMode = EColumnSortMode::Descending;
+			SortMode[InSortPriority] = EColumnSortMode::Descending;
 		}
 		else
 		{
-			SortMode = EColumnSortMode::Ascending;
+			SortMode[InSortPriority] = EColumnSortMode::Ascending;
 		}
+		return false;
 	}
 }
\ No newline at end of file
diff --git a/Engine/Source/Editor/ContentBrowser/Private/AssetViewSortManager.h b/Engine/Source/Editor/ContentBrowser/Private/AssetViewSortManager.h
--- a/Engine/Source/Editor/ContentBrowser/Private/AssetViewSortManager.h
+++ b/Engine/Source/Editor/ContentBrowser/Private/AssetViewSortManager.h
@@ -8,23 +8,26 @@ public:
 	/** Constructor */
 	FAssetViewSortManager();
 
+	/** Reset the sort mode back to default */
+	void ResetSort();
+
 	/** Sorts a list using the current ColumnId and Mode. Supply a MajorityAssetType to help discover sorting type (numerical vs alphabetical) */
-	void SortList(TArray<TSharedPtr<struct FAssetViewItem>>& AssetItems, FName MajorityAssetType) const;
+	void SortList(TArray<TSharedPtr<struct FAssetViewItem>>& AssetItems, const FName& MajorityAssetType) const;
 
-	/** Sets the sort mode based on the column that was clicked */
-	void SetOrToggleSortColumn(FName ColumnId);
+	/** Sets the sort mode based on the column that was clicked, returns true if newly assigned */
+	bool SetOrToggleSortColumn(const EColumnSortPriority::Type InSortPriority, const FName& InColumnId);
 
 	/**	Sets the column to sort */
-	void SetSortColumnId( const FName& ColumnId );
+	void SetSortColumnId(const EColumnSortPriority::Type InSortPriority, const FName& InColumnId);
 
 	/**	Sets the current sort mode */
-	void SetSortMode( const EColumnSortMode::Type InSortMode );
+	void SetSortMode(const EColumnSortPriority::Type InSortPriority, const EColumnSortMode::Type InSortMode);
 
 	/** Gets the current sort mode */
-	EColumnSortMode::Type GetSortMode() const { return SortMode; }
+	EColumnSortMode::Type GetSortMode(const EColumnSortPriority::Type InSortPriority) const { check(InSortPriority < EColumnSortPriority::Max); return SortMode[InSortPriority]; }
 
 	/** Gets the current sort column id */
-	FName GetSortColumnId() const { return SortColumnId; }
+	FName GetSortColumnId(const EColumnSortPriority::Type InSortPriority) const { check(InSortPriority < EColumnSortPriority::Max); return SortColumnId[InSortPriority]; }
 
 public:
 	/** The names of non-type specific columns in the columns view. */
@@ -34,8 +37,8 @@ public:
 
 private:
 	/** The name of the column that is currently used for sorting. */
-	FName SortColumnId;
+	FName SortColumnId[EColumnSortPriority::Max];
 
 	/** Whether the sort is ascending or descending. */
-	EColumnSortMode::Type SortMode;
+	EColumnSortMode::Type SortMode[EColumnSortPriority::Max];
 };
\ No newline at end of file
diff --git a/Engine/Source/Editor/ContentBrowser/Private/SAssetView.cpp b/Engine/Source/Editor/ContentBrowser/Private/SAssetView.cpp
--- a/Engine/Source/Editor/ContentBrowser/Private/SAssetView.cpp
+++ b/Engine/Source/Editor/ContentBrowser/Private/SAssetView.cpp
@@ -1227,12 +1227,14 @@ TSharedRef<SAssetColumnView> SAssetView::CreateColumnView()
 			+ SHeaderRow::Column(SortManager.NameColumnId)
 			.FillWidth(300)
 			.SortMode( TAttribute< EColumnSortMode::Type >::Create( TAttribute< EColumnSortMode::Type >::FGetter::CreateSP( this, &SAssetView::GetColumnSortMode, SortManager.NameColumnId ) ) )
+			.SortPriority(TAttribute< EColumnSortPriority::Type >::Create(TAttribute< EColumnSortPriority::Type >::FGetter::CreateSP(this, &SAssetView::GetColumnSortPriority, SortManager.NameColumnId)))
 			.OnSort( FOnSortModeChanged::CreateSP( this, &SAssetView::OnSortColumnHeader ) )
 			.DefaultLabel( LOCTEXT("Column_Name", "Name") )
 			//@TODO: Query the OnAssetTagWantsToBeDisplayed column filter here too, in case the user wants to bury the type column
 			+ SHeaderRow::Column(SortManager.ClassColumnId)
 			.FillWidth(160)
 			.SortMode( TAttribute< EColumnSortMode::Type >::Create( TAttribute< EColumnSortMode::Type >::FGetter::CreateSP( this, &SAssetView::GetColumnSortMode, SortManager.ClassColumnId ) ) )
+			.SortPriority(TAttribute< EColumnSortPriority::Type >::Create(TAttribute< EColumnSortPriority::Type >::FGetter::CreateSP(this, &SAssetView::GetColumnSortPriority, SortManager.ClassColumnId)))
 			.OnSort( FOnSortModeChanged::CreateSP( this, &SAssetView::OnSortColumnHeader ) )
 			.DefaultLabel( LOCTEXT("Column_Class", "Type") )
 		);
@@ -1626,10 +1628,24 @@ void SAssetView::SetMajorityAssetType(FName NewMajorityAssetType)
 
 		// Keep track of the current column name to see if we need to change it now that columns are being removed
 		// Name, Class, and Path are always relevant
-		const FName CurrentSortColumn = SortManager.GetSortColumnId();
-		bool bSortColumnStillRelevant = CurrentSortColumn == FAssetViewSortManager::NameColumnId
-									|| CurrentSortColumn == FAssetViewSortManager::ClassColumnId
-									|| CurrentSortColumn == FAssetViewSortManager::PathColumnId;
+		struct FSortOrder
+		{
+			bool bSortRelevant;
+			FName SortColumn;
+			FSortOrder(bool bInSortRelevant, const FName& InSortColumn) : bSortRelevant(bInSortRelevant), SortColumn(InSortColumn) {}
+		};
+		TArray<FSortOrder> CurrentSortOrder;
+		for (int32 PriorityIdx = 0; PriorityIdx < EColumnSortPriority::Max; PriorityIdx++)
+		{
+			const FName SortColumn = SortManager.GetSortColumnId(static_cast<EColumnSortPriority::Type>(PriorityIdx));
+			if (SortColumn != NAME_None)
+			{
+				const bool bSortRelevant = SortColumn == FAssetViewSortManager::NameColumnId
+					|| SortColumn == FAssetViewSortManager::ClassColumnId
+					|| SortColumn == FAssetViewSortManager::PathColumnId;
+				CurrentSortOrder.Add(FSortOrder(bSortRelevant, SortColumn));
+			}
+		}
 
 		// If we have a new majority type, add the new type's columns
 		if ( NewMajorityAssetType != NAME_None )
@@ -1667,6 +1683,7 @@ void SAssetView::SetMajorityAssetType(FName NewMajorityAssetType)
 								ColumnView->GetHeaderRow()->AddColumn(
 										SHeaderRow::Column(Tag)
 										.SortMode( TAttribute< EColumnSortMode::Type >::Create( TAttribute< EColumnSortMode::Type >::FGetter::CreateSP( this, &SAssetView::GetColumnSortMode, Tag ) ) )
+										.SortPriority(TAttribute< EColumnSortPriority::Type >::Create(TAttribute< EColumnSortPriority::Type >::FGetter::CreateSP(this, &SAssetView::GetColumnSortPriority, Tag)))
 										.OnSort( FOnSortModeChanged::CreateSP( this, &SAssetView::OnSortColumnHeader ) )
 										.DefaultLabel( DisplayName )
 										.HAlignCell( (TagIt->Type == UObject::FAssetRegistryTag::TT_Numerical) ? HAlign_Right : HAlign_Left )
@@ -1674,9 +1691,12 @@ void SAssetView::SetMajorityAssetType(FName NewMajorityAssetType)
 									);
 
 								// If we found a tag the matches the column we are currently sorting on, there will be no need to change the column
-								if ( Tag == CurrentSortColumn )
+								for (int32 SortIdx = 0; SortIdx < CurrentSortOrder.Num(); SortIdx++)
 								{
-									bSortColumnStillRelevant = true;
+									if (Tag == CurrentSortOrder[SortIdx].SortColumn)
+									{
+										CurrentSortOrder[SortIdx].bSortRelevant = true;
+									}
 								}
 							}
 						}
@@ -1685,10 +1705,36 @@ void SAssetView::SetMajorityAssetType(FName NewMajorityAssetType)
 			}	
 		}
 
-		if ( !bSortColumnStillRelevant )
+		// Are any of the sort columns irrelevant now, if so remove them from the list
+		bool CurrentSortChanged = false;
+		for (int32 SortIdx = CurrentSortOrder.Num() - 1; SortIdx >= 0; SortIdx--)
+		{
+			if (!CurrentSortOrder[SortIdx].bSortRelevant)
+			{
+				CurrentSortOrder.RemoveAt(SortIdx);
+				CurrentSortChanged = true;
+			}
+		}
+		if (CurrentSortOrder.Num() > 0 && CurrentSortChanged)
+		{
+			// Sort order has changed, update the columns keeping those that are relevant
+			int32 PriorityNum = EColumnSortPriority::Primary;
+			for (int32 SortIdx = 0; SortIdx < CurrentSortOrder.Num(); SortIdx++)
+			{
+				check(CurrentSortOrder[SortIdx].bSortRelevant);
+				if (!SortManager.SetOrToggleSortColumn(static_cast<EColumnSortPriority::Type>(PriorityNum), CurrentSortOrder[SortIdx].SortColumn))
+				{
+					// Toggle twice so mode is preserved if this isn't a new column assignation
+					SortManager.SetOrToggleSortColumn(static_cast<EColumnSortPriority::Type>(PriorityNum), CurrentSortOrder[SortIdx].SortColumn);
+				}				
+				bPendingSortFilteredItems = true;
+				PriorityNum++;
+			}
+		}
+		else if (CurrentSortOrder.Num() == 0)
 		{
 			// If the current sort column is no longer relevant, revert to "Name" and resort when convenient
-			SortManager.SetOrToggleSortColumn(FAssetViewSortManager::NameColumnId);
+			SortManager.ResetSort();
 			bPendingSortFilteredItems = true;
 		}
 	}
@@ -3492,26 +3538,40 @@ float SAssetView::GetTileViewItemBaseWidth() const
 	return ( TileViewThumbnailSize + TileViewThumbnailPadding * 2 ) * FMath::Lerp( MinThumbnailScale, MaxThumbnailScale, GetThumbnailScale() );
 }
 
-EColumnSortMode::Type SAssetView::GetColumnSortMode(FName ColumnId) const
+EColumnSortMode::Type SAssetView::GetColumnSortMode(const FName ColumnId) const
 {
-	if ( ColumnId == SortManager.GetSortColumnId() )
+	for (int32 PriorityIdx = 0; PriorityIdx < EColumnSortPriority::Max; PriorityIdx++)
 	{
-		return SortManager.GetSortMode();
+		const EColumnSortPriority::Type SortPriority = static_cast<EColumnSortPriority::Type>(PriorityIdx);
+		if (ColumnId == SortManager.GetSortColumnId(SortPriority))
+		{
+			return SortManager.GetSortMode(SortPriority);
+		}
 	}
-	else
+	return EColumnSortMode::None;
+}
+
+EColumnSortPriority::Type SAssetView::GetColumnSortPriority(const FName ColumnId) const
+{
+	for (int32 PriorityIdx = 0; PriorityIdx < EColumnSortPriority::Max; PriorityIdx++)
 	{
-		return EColumnSortMode::None;
+		const EColumnSortPriority::Type SortPriority = static_cast<EColumnSortPriority::Type>(PriorityIdx);
+		if (ColumnId == SortManager.GetSortColumnId(SortPriority))
+		{
+			return SortPriority;
+		}
 	}
+	return EColumnSortPriority::Primary;
 }
 
-void SAssetView::OnSortColumnHeader( const FName& ColumnId, EColumnSortMode::Type NewSortMode )
+void SAssetView::OnSortColumnHeader(const EColumnSortPriority::Type SortPriority, const FName& ColumnId, const EColumnSortMode::Type NewSortMode)
 {
-	SortManager.SetSortColumnId( ColumnId );
-	SortManager.SetSortMode( NewSortMode );
+	SortManager.SetSortColumnId(SortPriority, ColumnId);
+	SortManager.SetSortMode(SortPriority, NewSortMode);
 	SortList();
 }
 
-bool SAssetView::IsPathInAssetItemsList(FName ObjectPath) const
+bool SAssetView::IsPathInAssetItemsList(const FName& ObjectPath) const
 {
 	for ( auto AssetIt = AssetItems.CreateConstIterator(); AssetIt; ++AssetIt )
 	{
diff --git a/Engine/Source/Editor/ContentBrowser/Private/SAssetView.h b/Engine/Source/Editor/ContentBrowser/Private/SAssetView.h
--- a/Engine/Source/Editor/ContentBrowser/Private/SAssetView.h
+++ b/Engine/Source/Editor/ContentBrowser/Private/SAssetView.h
@@ -502,13 +502,16 @@ private:
 	float GetTileViewItemBaseWidth() const;
 
 	/** Gets the sort mode for the supplied ColumnId */
-	EColumnSortMode::Type GetColumnSortMode(FName ColumnId) const;
+	EColumnSortMode::Type GetColumnSortMode(const FName ColumnId) const;
+
+	/** Gets the sort order for the supplied ColumnId */
+	EColumnSortPriority::Type GetColumnSortPriority(const FName ColumnId) const;
 
 	/** Handler for when a column header is clicked */
-	void OnSortColumnHeader( const FName& ColumnId, EColumnSortMode::Type NewSortMode );
+	void OnSortColumnHeader(const EColumnSortPriority::Type SortPriority, const FName& ColumnId, const EColumnSortMode::Type NewSortMode);
 
 	/** Handler for when a column header is clicked */
-	bool IsPathInAssetItemsList(FName ObjectPath) const;
+	bool IsPathInAssetItemsList(const FName& ObjectPath) const;
 
 	/** Returns the state of the is working progress bar */
 	TOptional< float > GetIsWorkingProgressBarState() const;
diff --git a/Engine/Source/Editor/EditorStyle/Private/SlateEditorStyle.cpp b/Engine/Source/Editor/EditorStyle/Private/SlateEditorStyle.cpp
--- a/Engine/Source/Editor/EditorStyle/Private/SlateEditorStyle.cpp
+++ b/Engine/Source/Editor/EditorStyle/Private/SlateEditorStyle.cpp
@@ -36,6 +36,7 @@ FSlateEditorStyle::FStyle::FStyle( const TWeakObjectPtr< UEditorStyleSettings >&
 	// Slate Units do NOT have to map to pixels.
 	, Icon7x16(7.0f, 16.0f)
 	, Icon8x4(8.0f, 4.0f)
+	, Icon16x4(16.0f, 4.0f)
 	, Icon8x8(8.0f, 8.0f)
 	, Icon10x10(10.0f, 10.0f)
 	, Icon12x12(12.0f, 12.0f)
@@ -688,8 +689,10 @@ void FSlateEditorStyle::FStyle::SetupGeneralStyles()
 		Set( "TreeArrow_Expanded_Hovered", new IMAGE_BRUSH( "Common/TreeArrow_Expanded_Hovered", Icon10x10, DefaultForeground ) );
 
 		const FTableColumnHeaderStyle TableColumnHeaderStyle = FTableColumnHeaderStyle()
-			.SetSortAscendingImage( IMAGE_BRUSH( "Common/SortUpArrow", Icon8x4 ) )
-			.SetSortDescendingImage( IMAGE_BRUSH( "Common/SortDownArrow", Icon8x4 ) )
+			.SetSortPrimaryAscendingImage( IMAGE_BRUSH( "Common/SortUpArrow", Icon8x4 ) )
+			.SetSortPrimaryDescendingImage( IMAGE_BRUSH( "Common/SortDownArrow", Icon8x4 ) )
+			.SetSortSecondaryAscendingImage(IMAGE_BRUSH("Common/SortUpArrows", Icon16x4))
+			.SetSortSecondaryDescendingImage(IMAGE_BRUSH("Common/SortDownArrows", Icon16x4))
 			.SetNormalBrush( BOX_BRUSH( "Common/ColumnHeader", 4.f/32.f ) )
 			.SetHoveredBrush( BOX_BRUSH( "Common/ColumnHeader_Hovered", 4.f/32.f ) )
 			.SetMenuDropdownImage( IMAGE_BRUSH( "Common/ColumnHeader_Arrow", Icon8x8 ) )
@@ -698,8 +701,10 @@ void FSlateEditorStyle::FStyle::SetupGeneralStyles()
 		Set( "TableView.Header.Column", TableColumnHeaderStyle );
 
 		const FTableColumnHeaderStyle TableLastColumnHeaderStyle = FTableColumnHeaderStyle()
-			.SetSortAscendingImage( IMAGE_BRUSH( "Common/SortUpArrow", Icon8x4 ) )
-			.SetSortDescendingImage( IMAGE_BRUSH( "Common/SortDownArrow", Icon8x4 ) )
+			.SetSortPrimaryAscendingImage(IMAGE_BRUSH("Common/SortUpArrow", Icon8x4))
+			.SetSortPrimaryDescendingImage(IMAGE_BRUSH("Common/SortDownArrow", Icon8x4))
+			.SetSortSecondaryAscendingImage(IMAGE_BRUSH("Common/SortUpArrows", Icon16x4))
+			.SetSortSecondaryDescendingImage(IMAGE_BRUSH("Common/SortDownArrows", Icon16x4))
 			.SetNormalBrush( FSlateNoResource() )
 			.SetHoveredBrush( BOX_BRUSH( "Common/LastColumnHeader_Hovered", 4.f/32.f ) )
 			.SetMenuDropdownImage( IMAGE_BRUSH( "Common/ColumnHeader_Arrow", Icon8x8 ) )
@@ -2439,8 +2444,10 @@ void FSlateEditorStyle::FStyle::SetupPropertyEditorStyles()
 			);
 
 		const FTableColumnHeaderStyle PropertyTableColumnHeaderStyle = FTableColumnHeaderStyle()
-			.SetSortAscendingImage( IMAGE_BRUSH( "Common/SortUpArrow", Icon8x4 ) )
-			.SetSortDescendingImage( IMAGE_BRUSH( "Common/SortDownArrow", Icon8x4 ) )
+			.SetSortPrimaryAscendingImage(IMAGE_BRUSH("Common/SortUpArrow", Icon8x4))
+			.SetSortPrimaryDescendingImage(IMAGE_BRUSH("Common/SortDownArrow", Icon8x4))
+			.SetSortSecondaryAscendingImage(IMAGE_BRUSH("Common/SortUpArrows", Icon16x4))
+			.SetSortSecondaryDescendingImage(IMAGE_BRUSH("Common/SortDownArrows", Icon16x4))
 			.SetNormalBrush( BOX_BRUSH( "Common/ColumnHeader", 4.f/32.f ) )
 			.SetHoveredBrush( BOX_BRUSH( "Common/ColumnHeader_Hovered", 4.f/32.f ) )
 			.SetMenuDropdownImage( IMAGE_BRUSH( "Common/ColumnHeader_Arrow", Icon8x8 ) )
@@ -2448,8 +2455,10 @@ void FSlateEditorStyle::FStyle::SetupPropertyEditorStyles()
 			.SetMenuDropdownHoveredBorderBrush( BOX_BRUSH( "Common/ColumnHeaderMenuButton_Hovered", 4.f/32.f ) );
 
 		const FTableColumnHeaderStyle PropertyTableLastColumnHeaderStyle = FTableColumnHeaderStyle()
-			.SetSortAscendingImage( IMAGE_BRUSH( "Common/SortUpArrow", Icon8x4 ) )
-			.SetSortDescendingImage( IMAGE_BRUSH( "Common/SortDownArrow", Icon8x4 ) )
+			.SetSortPrimaryAscendingImage(IMAGE_BRUSH("Common/SortUpArrow", Icon8x4))
+			.SetSortPrimaryDescendingImage(IMAGE_BRUSH("Common/SortDownArrow", Icon8x4))
+			.SetSortSecondaryAscendingImage(IMAGE_BRUSH("Common/SortUpArrows", Icon16x4))
+			.SetSortSecondaryDescendingImage(IMAGE_BRUSH("Common/SortDownArrows", Icon16x4))
 			.SetNormalBrush( FSlateNoResource() )
 			.SetHoveredBrush( BOX_BRUSH( "Common/LastColumnHeader_Hovered", 4.f/32.f ) )
 			.SetMenuDropdownImage( IMAGE_BRUSH( "Common/ColumnHeader_Arrow", Icon8x8 ) )
diff --git a/Engine/Source/Editor/EditorStyle/Private/SlateEditorStyle.h b/Engine/Source/Editor/EditorStyle/Private/SlateEditorStyle.h
--- a/Engine/Source/Editor/EditorStyle/Private/SlateEditorStyle.h
+++ b/Engine/Source/Editor/EditorStyle/Private/SlateEditorStyle.h
@@ -82,6 +82,7 @@ public:
 
 		const FVector2D Icon7x16;
 		const FVector2D Icon8x4;
+		const FVector2D Icon16x4;
 		const FVector2D Icon8x8;
 		const FVector2D Icon10x10;
 		const FVector2D Icon12x12;
diff --git a/Engine/Source/Editor/PackagesDialog/Private/SPackagesDialog.cpp b/Engine/Source/Editor/PackagesDialog/Private/SPackagesDialog.cpp
--- a/Engine/Source/Editor/PackagesDialog/Private/SPackagesDialog.cpp
+++ b/Engine/Source/Editor/PackagesDialog/Private/SPackagesDialog.cpp
@@ -597,7 +597,7 @@ EColumnSortMode::Type SPackagesDialog::GetColumnSortMode( const FName ColumnId )
 	return SortMode;
 }
 
-void SPackagesDialog::OnColumnSortModeChanged( const FName& ColumnId, EColumnSortMode::Type InSortMode )
+void SPackagesDialog::OnColumnSortModeChanged( const EColumnSortPriority::Type SortPriority, const FName& ColumnId, const EColumnSortMode::Type InSortMode )
 {
 	SortByColumn = ColumnId;
 	SortMode = InSortMode;
diff --git a/Engine/Source/Editor/PackagesDialog/Private/SPackagesDialog.h b/Engine/Source/Editor/PackagesDialog/Private/SPackagesDialog.h
--- a/Engine/Source/Editor/PackagesDialog/Private/SPackagesDialog.h
+++ b/Engine/Source/Editor/PackagesDialog/Private/SPackagesDialog.h
@@ -424,7 +424,7 @@ private:
 	 * @param	ColumnId	The new column to sort by
 	 * @param	InSortMode	The sort mode (ascending or descending)
 	 */
-	void OnColumnSortModeChanged( const FName& ColumnId, EColumnSortMode::Type InSortMode );
+	void OnColumnSortModeChanged( const EColumnSortPriority::Type SortPriority, const FName& ColumnId, const EColumnSortMode::Type InSortMode );
 
 	/**
 	 * Requests that the source list data be sorted according to the current sort column and mode,
diff --git a/Engine/Source/Editor/PropertyEditor/Private/Presentation/PropertyTable/PropertyTable.cpp b/Engine/Source/Editor/PropertyEditor/Private/Presentation/PropertyTable/PropertyTable.cpp
--- a/Engine/Source/Editor/PropertyEditor/Private/Presentation/PropertyTable/PropertyTable.cpp
+++ b/Engine/Source/Editor/PropertyEditor/Private/Presentation/PropertyTable/PropertyTable.cpp
@@ -330,7 +330,7 @@ EColumnSortMode::Type FPropertyTable::GetColumnSortMode( const TSharedRef< class
 	return EColumnSortMode::None;
 }
 
-void FPropertyTable::SortByColumnWithId( const FName& ColumnId, EColumnSortMode::Type SortMode )
+void FPropertyTable::SortByColumnWithId( const EColumnSortPriority::Type SortPriority, const FName& ColumnId, const EColumnSortMode::Type SortMode )
 {
 	for( auto ColumnIter = Columns.CreateIterator(); ColumnIter; ++ColumnIter )
 	{
diff --git a/Engine/Source/Editor/PropertyEditor/Private/Presentation/PropertyTable/PropertyTable.h b/Engine/Source/Editor/PropertyEditor/Private/Presentation/PropertyTable/PropertyTable.h
--- a/Engine/Source/Editor/PropertyEditor/Private/Presentation/PropertyTable/PropertyTable.h
+++ b/Engine/Source/Editor/PropertyEditor/Private/Presentation/PropertyTable/PropertyTable.h
@@ -109,7 +109,7 @@ public:
 	virtual void SetSelectionMode( const ESelectionMode::Type Mode ) override;
 
 	virtual EColumnSortMode::Type GetColumnSortMode( const TSharedRef< class IPropertyTableColumn > Column ) const override;
-	virtual void SortByColumnWithId( const FName& ColumnId, EColumnSortMode::Type SortMode ) override;
+	virtual void SortByColumnWithId( const EColumnSortPriority::Type SortPriority, const FName& ColumnId, const EColumnSortMode::Type SortMode ) override;
 	virtual void SortByColumn( const TSharedRef< class IPropertyTableColumn >& Column, EColumnSortMode::Type SortMode ) override;
 
 	virtual void PasteTextAtCell( const FString& Text, const TSharedRef< class IPropertyTableCell >& Cell ) override;
diff --git a/Engine/Source/Editor/PropertyEditor/Public/IPropertyTable.h b/Engine/Source/Editor/PropertyEditor/Public/IPropertyTable.h
--- a/Engine/Source/Editor/PropertyEditor/Public/IPropertyTable.h
+++ b/Engine/Source/Editor/PropertyEditor/Public/IPropertyTable.h
@@ -129,7 +129,7 @@ public:
 	virtual void SetSelectionMode( const ESelectionMode::Type Mode ) = 0;
 
 	virtual EColumnSortMode::Type GetColumnSortMode( const TSharedRef< class IPropertyTableColumn > Column ) const = 0;
-	virtual void SortByColumnWithId( const FName& ColumnId, EColumnSortMode::Type SortMode ) = 0;
+	virtual void SortByColumnWithId( const EColumnSortPriority::Type SortPriority, const FName& ColumnId, const EColumnSortMode::Type SortMode ) = 0;
 	virtual void SortByColumn( const TSharedRef< class IPropertyTableColumn >& Column, EColumnSortMode::Type SortMode ) = 0;
 
 	virtual void PasteTextAtCell( const FString& Text, const TSharedRef< class IPropertyTableCell >& Cell ) = 0;
diff --git a/Engine/Source/Editor/SceneOutliner/Private/SSceneOutliner.cpp b/Engine/Source/Editor/SceneOutliner/Private/SSceneOutliner.cpp
--- a/Engine/Source/Editor/SceneOutliner/Private/SSceneOutliner.cpp
+++ b/Engine/Source/Editor/SceneOutliner/Private/SSceneOutliner.cpp
@@ -3105,7 +3105,7 @@ namespace SceneOutliner
 		return SortMode;
 	}
 
-	void SSceneOutliner::OnColumnSortModeChanged( const FName& ColumnId, EColumnSortMode::Type InSortMode )
+	void SSceneOutliner::OnColumnSortModeChanged( const EColumnSortPriority::Type SortPriority, const FName& ColumnId, const EColumnSortMode::Type InSortMode )
 	{
 		if (CustomColumn.IsValid() && ColumnId == CustomColumn->GetColumnID() && !CustomColumn->SupportsSorting())
 		{
diff --git a/Engine/Source/Editor/SceneOutliner/Private/SSceneOutliner.h b/Engine/Source/Editor/SceneOutliner/Private/SSceneOutliner.h
--- a/Engine/Source/Editor/SceneOutliner/Private/SSceneOutliner.h
+++ b/Engine/Source/Editor/SceneOutliner/Private/SSceneOutliner.h
@@ -534,7 +534,7 @@ namespace SceneOutliner
 		/********** Sort functions **********/
 
 		/** Handles column sorting mode change */
-		void OnColumnSortModeChanged( const FName& ColumnId, EColumnSortMode::Type InSortMode );
+		void OnColumnSortModeChanged( const EColumnSortPriority::Type SortPriority, const FName& ColumnId, const EColumnSortMode::Type InSortMode );
 
 		/** @return Returns the current sort mode of the specified column */
 		EColumnSortMode::Type GetColumnSortMode( const FName ColumnId ) const;
diff --git a/Engine/Source/Runtime/Core/Public/Templates/UniquePtr.h b/Engine/Source/Runtime/Core/Public/Templates/UniquePtr.h
--- a/Engine/Source/Runtime/Core/Public/Templates/UniquePtr.h
+++ b/Engine/Source/Runtime/Core/Public/Templates/UniquePtr.h
@@ -16,6 +16,9 @@
 template <typename T>
 class TUniquePtr
 {
+	template <typename OtherT>
+	friend class TUniquePtr;
+
 public:
 	/**
 	 * Default constructor - initializes the TUniquePtr to null.
diff --git a/Engine/Source/Runtime/Slate/Private/Widgets/Views/SHeaderRow.cpp b/Engine/Source/Runtime/Slate/Private/Widgets/Views/SHeaderRow.cpp
--- a/Engine/Source/Runtime/Slate/Private/Widgets/Views/SHeaderRow.cpp
+++ b/Engine/Source/Runtime/Slate/Private/Widgets/Views/SHeaderRow.cpp
@@ -16,6 +16,7 @@ public:
 
 	STableColumnHeader()
 		: SortMode( EColumnSortMode::None )
+		, SortPriority( EColumnSortPriority::Primary )
 		, OnSortModeChanged()
 		, ContextMenuContent( SNullWidget::NullWidget )
 		, ColumnId( NAME_None )
@@ -38,8 +39,8 @@ public:
 		Style = InArgs._Style;
 		ColumnId = Column.ColumnId;
 		SortMode = Column.SortMode;
+		SortPriority = Column.SortPriority;
 
-		SortMode = Column.SortMode;
 		OnSortModeChanged = Column.OnSortModeChanged;
 		ContextMenuContent = Column.HeaderMenuContent.Widget;
 
@@ -182,6 +183,18 @@ public:
 		SortMode = NewMode;
 	}
 
+	/** Gets sorting order */
+	EColumnSortPriority::Type GetSortPriority() const
+	{
+		return SortPriority.Get();
+	}
+
+	/** Sets sorting order */
+	void SetSortPriority(EColumnSortPriority::Type NewPriority)
+	{
+		SortPriority = NewPriority;
+	}
+
 	virtual FReply OnMouseButtonUp( const FGeometry& MyGeometry, const FPointerEvent& MouseEvent ) override
 	{
 		if ( MouseEvent.GetEffectingButton() == EKeys::RightMouseButton && ContextMenuContent != SNullWidget::NullWidget )
@@ -276,7 +289,11 @@ private:
 	/** Gets the icon associated with the current sorting mode */
 	const FSlateBrush* GetSortingBrush() const
 	{
-		return (SortMode == EColumnSortMode::Ascending ? &Style->SortAscendingImage : &Style->SortDescendingImage);
+		EColumnSortPriority::Type ColumnSortPriority = SortPriority.Get();
+
+		return ( SortMode.Get() == EColumnSortMode::Ascending ? 
+			(SortPriority.Get() == EColumnSortPriority::Secondary ? &Style->SortSecondaryAscendingImage : &Style->SortPrimaryAscendingImage) :
+			(SortPriority.Get() == EColumnSortPriority::Secondary ? &Style->SortSecondaryDescendingImage : &Style->SortPrimaryDescendingImage) );
 	}
 
 	/** Checks if sorting mode has been selected */
@@ -290,17 +307,40 @@ private:
 	{
 		if ( OnSortModeChanged.IsBound() )
 		{
+			const bool bIsShiftClicked = FSlateApplication::Get().GetModifierKeys().IsShiftDown();
+			EColumnSortPriority::Type ColumnSortPriority = SortPriority.Get();
 			EColumnSortMode::Type ColumnSortMode = SortMode.Get();
-			if ( ColumnSortMode == EColumnSortMode::None || ColumnSortMode == EColumnSortMode::Descending )
+			if (ColumnSortMode == EColumnSortMode::None)
 			{
+				if (bIsShiftClicked && SortPriority.IsBound())
+				{
+					ColumnSortPriority = EColumnSortPriority::Secondary;
+				}
+				else
+				{
+					ColumnSortPriority = EColumnSortPriority::Primary;
+				}
+
 				ColumnSortMode = EColumnSortMode::Ascending;
 			}
 			else
 			{
-				ColumnSortMode = EColumnSortMode::Descending;
+				if (!bIsShiftClicked && ColumnSortPriority == EColumnSortPriority::Secondary)
+				{
+					ColumnSortPriority = EColumnSortPriority::Primary;
+				}
+
+				if (ColumnSortMode == EColumnSortMode::Descending)
+				{
+					ColumnSortMode = EColumnSortMode::Ascending;
+				}
+				else
+				{
+					ColumnSortMode = EColumnSortMode::Descending;
+				}
 			}
 
-			OnSortModeChanged.Execute( ColumnId, ColumnSortMode );
+			OnSortModeChanged.Execute(ColumnSortPriority, ColumnId, ColumnSortMode);
 		}
 
 		return FReply::Handled();
@@ -320,6 +360,9 @@ private:
 	/** Current sorting mode */
 	TAttribute< EColumnSortMode::Type > SortMode;
 
+	/** Current sorting order */
+	TAttribute< EColumnSortPriority::Type > SortPriority;
+
 	/** Callback triggered when sorting mode changes */
 	FOnSortModeChanged OnSortModeChanged;
 
diff --git a/Engine/Source/Runtime/Slate/Public/Widgets/Views/SHeaderRow.h b/Engine/Source/Runtime/Slate/Public/Widgets/Views/SHeaderRow.h
--- a/Engine/Source/Runtime/Slate/Public/Widgets/Views/SHeaderRow.h
+++ b/Engine/Source/Runtime/Slate/Public/Widgets/Views/SHeaderRow.h
@@ -5,6 +5,16 @@
 class SSplitter;
 class IGridRow; 
 
+namespace EColumnSortPriority
+{
+	enum Type
+	{
+		Primary = 0,
+		Secondary = 1,
+		Max,
+	};
+};
+
 namespace EColumnSortMode
 {
 	enum Type
@@ -46,7 +56,7 @@ enum class EHeaderComboVisibility
 };
 
 /** Callback when sort mode changes */
-DECLARE_DELEGATE_TwoParams( FOnSortModeChanged, const FName&, EColumnSortMode::Type );
+DECLARE_DELEGATE_ThreeParams( FOnSortModeChanged, EColumnSortPriority::Type, const FName&, EColumnSortMode::Type );
 
 /**	Callback when the width of the column changes */
 DECLARE_DELEGATE_OneParam( FOnWidthChanged, float );
@@ -96,6 +106,7 @@ public:
 			SLATE_ARGUMENT( EVerticalAlignment, VAlignCell )
 
 			SLATE_ATTRIBUTE( EColumnSortMode::Type, SortMode )
+			SLATE_ATTRIBUTE( EColumnSortPriority::Type, SortPriority )
 			SLATE_EVENT( FOnSortModeChanged, OnSort )
 		SLATE_END_ARGS()
 
@@ -115,6 +126,7 @@ public:
 			, CellHAlignment( InArgs._HAlignCell )
 			, CellVAlignment( InArgs._VAlignCell )
 			, SortMode( InArgs._SortMode )
+			, SortPriority( InArgs._SortPriority )
 			, OnSortModeChanged( InArgs._OnSort )
 		{
 			if ( InArgs._FixedWidth.IsSet() )
@@ -179,6 +191,7 @@ public:
 		EVerticalAlignment CellVAlignment;
 
 		TAttribute< EColumnSortMode::Type > SortMode;
+		TAttribute< EColumnSortPriority::Type > SortPriority;
 		FOnSortModeChanged OnSortModeChanged;
 	};
 
diff --git a/Engine/Source/Runtime/SlateCore/Private/Styling/CoreStyle.cpp b/Engine/Source/Runtime/SlateCore/Private/Styling/CoreStyle.cpp
--- a/Engine/Source/Runtime/SlateCore/Private/Styling/CoreStyle.cpp
+++ b/Engine/Source/Runtime/SlateCore/Private/Styling/CoreStyle.cpp
@@ -121,6 +121,7 @@ TSharedRef<ISlateStyle> FCoreStyle::Create( const FName& InStyleSetName )
 	// Slate Units do NOT have to map to pixels.
 	const FVector2D Icon5x16(5.0f, 16.0f);
 	const FVector2D Icon8x4(8.0f, 4.0f);
+	const FVector2D Icon16x4(16.0f, 4.0f);
 	const FVector2D Icon8x8(8.0f, 8.0f);
 	const FVector2D Icon10x10(10.0f, 10.0f);
 	const FVector2D Icon12x12(12.0f, 12.0f);
@@ -624,8 +625,10 @@ TSharedRef<ISlateStyle> FCoreStyle::Create( const FName& InStyleSetName )
 		Style->Set( "TreeArrow_Expanded_Hovered", new IMAGE_BRUSH( "Common/TreeArrow_Expanded_Hovered", Icon10x10, DefaultForeground ) );
 
 		const FTableColumnHeaderStyle TableColumnHeaderStyle = FTableColumnHeaderStyle()
-			.SetSortAscendingImage( IMAGE_BRUSH( "Common/SortUpArrow", Icon8x4 ) )
-			.SetSortDescendingImage( IMAGE_BRUSH( "Common/SortDownArrow", Icon8x4 ) )
+			.SetSortPrimaryAscendingImage(IMAGE_BRUSH("Common/SortUpArrow", Icon8x4))
+			.SetSortPrimaryDescendingImage(IMAGE_BRUSH("Common/SortDownArrow", Icon8x4))
+			.SetSortSecondaryAscendingImage(IMAGE_BRUSH("Common/SortUpArrows", Icon16x4))
+			.SetSortSecondaryDescendingImage(IMAGE_BRUSH("Common/SortDownArrows", Icon16x4))
 			.SetNormalBrush( BOX_BRUSH( "Common/ColumnHeader", 4.f/32.f ) )
 			.SetHoveredBrush( BOX_BRUSH( "Common/ColumnHeader_Hovered", 4.f/32.f ) )
 			.SetMenuDropdownImage( IMAGE_BRUSH( "Common/ColumnHeader_Arrow", Icon8x8 ) )
@@ -634,8 +637,10 @@ TSharedRef<ISlateStyle> FCoreStyle::Create( const FName& InStyleSetName )
 		Style->Set( "TableView.Header.Column", TableColumnHeaderStyle );
 
 		const FTableColumnHeaderStyle TableLastColumnHeaderStyle = FTableColumnHeaderStyle()
-			.SetSortAscendingImage( IMAGE_BRUSH( "Common/SortUpArrow", Icon8x4 ) )
-			.SetSortDescendingImage( IMAGE_BRUSH( "Common/SortDownArrow", Icon8x4 ) )
+			.SetSortPrimaryAscendingImage(IMAGE_BRUSH("Common/SortUpArrow", Icon8x4))
+			.SetSortPrimaryDescendingImage(IMAGE_BRUSH("Common/SortDownArrow", Icon8x4))
+			.SetSortSecondaryAscendingImage(IMAGE_BRUSH("Common/SortUpArrows", Icon16x4))
+			.SetSortSecondaryDescendingImage(IMAGE_BRUSH("Common/SortDownArrows", Icon16x4))
 			.SetNormalBrush( FSlateNoResource() )
 			.SetHoveredBrush( BOX_BRUSH( "Common/LastColumnHeader_Hovered", 4.f/32.f ) )
 			.SetMenuDropdownImage( IMAGE_BRUSH( "Common/ColumnHeader_Arrow", Icon8x8 ) )
diff --git a/Engine/Source/Runtime/SlateCore/Private/Styling/SlateTypes.cpp b/Engine/Source/Runtime/SlateCore/Private/Styling/SlateTypes.cpp
--- a/Engine/Source/Runtime/SlateCore/Private/Styling/SlateTypes.cpp
+++ b/Engine/Source/Runtime/SlateCore/Private/Styling/SlateTypes.cpp
@@ -489,8 +489,10 @@ FTableColumnHeaderStyle::FTableColumnHeaderStyle()
 
 void FTableColumnHeaderStyle::GetResources( TArray< const FSlateBrush* >& OutBrushes ) const
 {
-	OutBrushes.Add( &SortAscendingImage );
-	OutBrushes.Add( &SortDescendingImage );
+	OutBrushes.Add( &SortPrimaryAscendingImage );
+	OutBrushes.Add( &SortPrimaryDescendingImage );
+	OutBrushes.Add( &SortSecondaryAscendingImage );
+	OutBrushes.Add( &SortSecondaryDescendingImage );
 	OutBrushes.Add( &NormalBrush );
 	OutBrushes.Add( &HoveredBrush );
 	OutBrushes.Add( &MenuDropdownImage );
diff --git a/Engine/Source/Runtime/SlateCore/Public/Styling/SlateTypes.h b/Engine/Source/Runtime/SlateCore/Public/Styling/SlateTypes.h
--- a/Engine/Source/Runtime/SlateCore/Public/Styling/SlateTypes.h
+++ b/Engine/Source/Runtime/SlateCore/Public/Styling/SlateTypes.h
@@ -1047,15 +1047,25 @@ struct SLATECORE_API FTableColumnHeaderStyle : public FSlateWidgetStyle
 
 	static const FTableColumnHeaderStyle& GetDefault();
 
-	/** Image used when a column is sorted in ascending order */
+	/** Image used when a column is primarily sorted in ascending order */
 	UPROPERTY(EditAnywhere, Category=Appearance)
-	FSlateBrush SortAscendingImage;
-	FTableColumnHeaderStyle& SetSortAscendingImage( const FSlateBrush& InSortAscendingImage ){ SortAscendingImage = InSortAscendingImage; return *this; }
+	FSlateBrush SortPrimaryAscendingImage;
+	FTableColumnHeaderStyle& SetSortPrimaryAscendingImage(const FSlateBrush& InSortPrimaryAscendingImage){ SortPrimaryAscendingImage = InSortPrimaryAscendingImage; return *this; }
 
-	/** Image used when a column is sorted in descending order */
+	/** Image used when a column is primarily sorted in descending order */
 	UPROPERTY(EditAnywhere, Category=Appearance)
-	FSlateBrush SortDescendingImage;
-	FTableColumnHeaderStyle& SetSortDescendingImage( const FSlateBrush& InSortDescendingImage ){ SortDescendingImage = InSortDescendingImage; return *this; }
+	FSlateBrush SortPrimaryDescendingImage;
+	FTableColumnHeaderStyle& SetSortPrimaryDescendingImage(const FSlateBrush& InSortPrimaryDescendingImage){ SortPrimaryDescendingImage = InSortPrimaryDescendingImage; return *this; }
+
+	/** Image used when a column is secondarily sorted in ascending order */
+	UPROPERTY(EditAnywhere, Category = Appearance)
+	FSlateBrush SortSecondaryAscendingImage;
+	FTableColumnHeaderStyle& SetSortSecondaryAscendingImage(const FSlateBrush& InSortSecondaryAscendingImage){ SortSecondaryAscendingImage = InSortSecondaryAscendingImage; return *this; }
+
+	/** Image used when a column is secondarily sorted in descending order */
+	UPROPERTY(EditAnywhere, Category = Appearance)
+	FSlateBrush SortSecondaryDescendingImage;
+	FTableColumnHeaderStyle& SetSortSecondaryDescendingImage(const FSlateBrush& InSortSecondaryDescendingImage){ SortSecondaryDescendingImage = InSortSecondaryDescendingImage; return *this; }
 
 	/** Brush used to draw the header in its normal state */
 	UPROPERTY(EditAnywhere, Category=Appearance)
