Fix for compiler error and minor optimization

#UDN: Feature suggestion; Content browser can sort by a second column

diff --git a/Engine/Source/Editor/ContentBrowser/Private/AssetViewSortManager.cpp b/Engine/Source/Editor/ContentBrowser/Private/AssetViewSortManager.cpp
--- a/Engine/Source/Editor/ContentBrowser/Private/AssetViewSortManager.cpp
+++ b/Engine/Source/Editor/ContentBrowser/Private/AssetViewSortManager.cpp
@@ -27,7 +27,7 @@ public:
 	}
 
 	/** Set the next comparisons array */
-	FORCEINLINE bool SetNextComparisons(TArray<TUniquePtr<FCompareFAssetItemBase>>& InNextComparisons) const
+	FORCEINLINE void SetNextComparisons(TArray<TUniquePtr<FCompareFAssetItemBase>>& InNextComparisons) const
 	{
 		check(NextComparisons.Num() == 0);
 		NextComparisons = MoveTemp(InNextComparisons);
@@ -177,18 +177,15 @@ protected:
 	{
 		// Depending if we're sorting ascending or descending it's quicker to flip the compares incase tags are missing
 		const FString* Value1 = bAscending ? StaticCastSharedPtr<FAssetViewAsset>(A)->Data.TagsAndValues.Find(Tag) : StaticCastSharedPtr<FAssetViewAsset>(B)->Data.TagsAndValues.Find(Tag);
-		const FString* Value2 = NULL;
 		if (!Value1)
 		{
 			return true;
 		}
-		else
+
+		const FString* Value2 = bAscending ? StaticCastSharedPtr<FAssetViewAsset>(B)->Data.TagsAndValues.Find(Tag) : StaticCastSharedPtr<FAssetViewAsset>(A)->Data.TagsAndValues.Find(Tag);
+		if (!Value2)
 		{
-			Value2 = bAscending ? StaticCastSharedPtr<FAssetViewAsset>(B)->Data.TagsAndValues.Find(Tag) : StaticCastSharedPtr<FAssetViewAsset>(A)->Data.TagsAndValues.Find(Tag);
-			if (!Value2)
-			{
-				return false;
-			}
+			return false;
 		}
 
 		const int32 Result = Value1->Compare(*Value2, ESearchCase::IgnoreCase);
@@ -214,18 +211,15 @@ protected:
 	{
 		// Depending if we're sorting ascending or descending it's quicker to flip the compares incase tags are missing
 		const FString* Value1 = bAscending ? StaticCastSharedPtr<FAssetViewAsset>(A)->Data.TagsAndValues.Find(Tag) : StaticCastSharedPtr<FAssetViewAsset>(B)->Data.TagsAndValues.Find(Tag);
-		const FString* Value2 = NULL;
 		if (!Value1)
 		{
 			return true;
 		}
-		else
+
+		const FString* Value2 = bAscending ? StaticCastSharedPtr<FAssetViewAsset>(B)->Data.TagsAndValues.Find(Tag) : StaticCastSharedPtr<FAssetViewAsset>(A)->Data.TagsAndValues.Find(Tag);
+		if (!Value2)
 		{
-			Value2 = bAscending ? StaticCastSharedPtr<FAssetViewAsset>(B)->Data.TagsAndValues.Find(Tag) : StaticCastSharedPtr<FAssetViewAsset>(A)->Data.TagsAndValues.Find(Tag);
-			if (!Value2)
-			{
-				return false;
-			}
+			return false;
 		}
 
 		const float Num1 = FCString::Atof(**Value1);
@@ -252,18 +246,15 @@ protected:
 	{
 		// Depending if we're sorting ascending or descending it's quicker to flip the compares incase tags are missing
 		const FString* Value1 = bAscending ? StaticCastSharedPtr<FAssetViewAsset>(A)->Data.TagsAndValues.Find(Tag) : StaticCastSharedPtr<FAssetViewAsset>(B)->Data.TagsAndValues.Find(Tag);
-		const FString* Value2 = NULL;
 		if (!Value1)
 		{
 			return true;
 		}
-		else
+
+		const FString* Value2 = bAscending ? StaticCastSharedPtr<FAssetViewAsset>(B)->Data.TagsAndValues.Find(Tag) : StaticCastSharedPtr<FAssetViewAsset>(A)->Data.TagsAndValues.Find(Tag);
+		if (!Value2)
 		{
-			Value2 = bAscending ? StaticCastSharedPtr<FAssetViewAsset>(B)->Data.TagsAndValues.Find(Tag) : StaticCastSharedPtr<FAssetViewAsset>(A)->Data.TagsAndValues.Find(Tag);
-			if (!Value2)
-			{
-				return false;
-			}
+			return false;
 		}
 
 		float Num1 = 1.f;
