Fixed renaming assets in the Content Browser.  The renamed asset is correctly sorted into position and scrolled into view.  Removed the problems with assets being duplicated in the asset view.
#ttp 339563 - ContentBrowser: Renaming assets sometimes displays 2 assets with same name
#ttp 339737 - CRASH: EDITOR: Editor crashing when changing names of several assets one after another
#branch UE4
#proj Editor.ContentBrowser
#change Changed SAssetView::ProcessRecentlyAddedAssets to note if any filtered asset items have been added while iterating through the FilteredRecentlyAddedAssets list, and if so to perform a quick refresh at the end (so that the view is refreshed, and a sort is forced).
#change In SAssetView::OnAssetRemoved, remove AssetData from the RecentlyAddedAssets list if it is there.
#change In SAssetView::ProcessRecentlyAddedAssets, when we add something to the AssetItems list, also add it to the ExistingObjectPaths set so that duplicate items aren't added to the FilteredAssetItems list.  When filtering newly added assets, always process them in the same tick.
#change Added RequestAddNewAssetsImmediately, to override the delay in adding newly added items (used by Rename so that the effects are instant).

diff --git a/Engine/Source/Editor/ContentBrowser/Private/SAssetView.cpp b/Engine/Source/Editor/ContentBrowser/Private/SAssetView.cpp
--- a/Engine/Source/Editor/ContentBrowser/Private/SAssetView.cpp
+++ b/Engine/Source/Editor/ContentBrowser/Private/SAssetView.cpp
@@ -15,13 +15,20 @@
 
 #define LOCTEXT_NAMESPACE "ContentBrowser"
 
+namespace
+{
+	/** Time delay between recently added items being added to the filtered asset items list */
+	const double TimeBetweenAddingNewAssets = 4.0;
+
+	/** Time delay between performing the last jump, and the jump term being reset */
+	const double JumpDelaySeconds = 0.6;
+}
+
 #define MAX_THUMBNAIL_SIZE 4096
 #define MAX_CLASS_NAME_LENGTH 32 // Enforce a reasonable class name length so the path is not too long for PLATFORM_MAX_FILEPATH_LENGTH
 
 #define MAX_PROJECTED_COOKING_PATH 165
 
-const double SAssetView::FQuickJumpData::JumpDelaySeconds = 0.6;
-
 SAssetView::~SAssetView()
 {
 	// Load the asset registry module to unregister delegates
@@ -631,6 +638,11 @@ void SAssetView::RequestQuickFrontendListRefresh()
 	bQuickFrontendListRefreshRequested = true;
 }
 
+void SAssetView::RequestAddNewAssetsNextFrame()
+{
+	LastProcessAddsTime = FPlatformTime::Seconds() - TimeBetweenAddingNewAssets;
+}
+
 void SAssetView::SaveSettings(const FString& IniFilename, const FString& IniSection, const FString& SettingsString) const
 {
 	GConfig->SetFloat(*IniSection, *(SettingsString + TEXT(".ThumbnailScale")), ThumbnailScaleSliderValue.Get(), IniFilename);
@@ -897,7 +909,7 @@ void SAssetView::Tick(const FGeometry& AllottedGeometry, const double InCurrentT
 		QuickJumpData.LastJumpTime = InCurrentTime;
 		QuickJumpData.bHasValidMatch = PerformQuickJump(bWasJumping);
 	}
-	else if(QuickJumpData.bIsJumping && InCurrentTime > QuickJumpData.LastJumpTime + FQuickJumpData::JumpDelaySeconds)
+	else if(QuickJumpData.bIsJumping && InCurrentTime > QuickJumpData.LastJumpTime + JumpDelaySeconds)
 	{
 		ResetQuickJump();
 	}
@@ -1761,10 +1773,22 @@ void SAssetView::OnAssetAdded(const FAssetData& AssetData)
 
 void SAssetView::ProcessRecentlyAddedAssets()
 {
-	if ( FilteredRecentlyAddedAssets.Num() > 0 )
+	if (
+		(RecentlyAddedAssets.Num() > 2048) ||
+		(RecentlyAddedAssets.Num() > 0 && FPlatformTime::Seconds() - LastProcessAddsTime >= TimeBetweenAddingNewAssets)
+		)
+	{
+		RunAssetsThroughBackendFilter(RecentlyAddedAssets);
+		FilteredRecentlyAddedAssets.Append(RecentlyAddedAssets);
+		RecentlyAddedAssets.Empty();
+		LastProcessAddsTime = FPlatformTime::Seconds();
+	}
+
+	if (FilteredRecentlyAddedAssets.Num() > 0)
 	{
 		const static float MaxSecondsPerFrame = 0.015;
 		double TickStartTime = FPlatformTime::Seconds();
+		bool bNeedsRefresh = false;
 
 		TSet<FName> ExistingObjectPaths;
 		for ( auto AssetIt = AssetItems.CreateConstIterator(); AssetIt; ++AssetIt )
@@ -1789,12 +1813,12 @@ void SAssetView::ProcessRecentlyAddedAssets()
 					{
 						// Add the asset to the list
 						int32 AddedAssetIdx = AssetItems.Add(AssetData);
-						if ( !IsFrontendFilterActive() || PassesCurrentFrontendFilter(AssetItems[AddedAssetIdx]) )
+						ExistingObjectPaths.Add(AssetData.ObjectPath);
+						if (!IsFrontendFilterActive() || PassesCurrentFrontendFilter(AssetData))
 						{
-							FilteredAssetItems.Add(MakeShareable(new FAssetViewAsset(AssetItems[AddedAssetIdx])));
+							FilteredAssetItems.Add(MakeShareable(new FAssetViewAsset(AssetData)));
+							bNeedsRefresh = true;
 							bPendingSortFilteredItems = true;
-
-							RefreshList();
 						}
 					}
 				}
@@ -1813,16 +1837,11 @@ void SAssetView::ProcessRecentlyAddedAssets()
 		{
 			FilteredRecentlyAddedAssets.RemoveAt(0, AssetIdx);
 		}
-	}
-	else if (
-		( RecentlyAddedAssets.Num() > 2048 ) ||
-		( RecentlyAddedAssets.Num() > 0 && FPlatformTime::Seconds() - LastProcessAddsTime > 4 )
-		)
-	{
-		RunAssetsThroughBackendFilter(RecentlyAddedAssets);
-		FilteredRecentlyAddedAssets.Append(RecentlyAddedAssets);
-		RecentlyAddedAssets.Empty();
-		LastProcessAddsTime = FPlatformTime::Seconds();
+
+		if (bNeedsRefresh)
+		{
+			RefreshList();
+		}
 	}
 }
 
@@ -1843,6 +1862,7 @@ void SAssetView::OnAssetsRemovedFromCollection( const FCollectionNameType& Colle
 void SAssetView::OnAssetRemoved(const FAssetData& AssetData)
 {
 	RemoveAssetByPath( AssetData.ObjectPath );
+	RecentlyAddedAssets.RemoveSingleSwap(AssetData);
 }
 
 void SAssetView::OnAssetRegistryPathAdded(const FString& Path)
@@ -1966,6 +1986,9 @@ void SAssetView::OnAssetRenamed(const FAssetData& AssetData, const FString& OldO
 
 	// Add the new asset, if it should be in the cached list
 	OnAssetAdded( AssetData );
+
+	// Force an update of the recently added asset next frame
+	RequestAddNewAssetsNextFrame();
 }
 
 void SAssetView::OnAssetLoaded(UObject* Asset)
diff --git a/Engine/Source/Editor/ContentBrowser/Private/SAssetView.h b/Engine/Source/Editor/ContentBrowser/Private/SAssetView.h
--- a/Engine/Source/Editor/ContentBrowser/Private/SAssetView.h
+++ b/Engine/Source/Editor/ContentBrowser/Private/SAssetView.h
@@ -184,6 +184,9 @@ public:
 	/** Requests that the asset view refreshes only items that are filtered through frontend sources. This should be used when possible. */
 	void RequestQuickFrontendListRefresh();
 
+	/** Requests that the asset view adds any recently added items in the next update to the filtered asset items */
+	void RequestAddNewAssetsNextFrame();
+
 	/** Saves any settings to config that should be persistent between editor sessions */
 	void SaveSettings(const FString& IniFilename, const FString& IniSection, const FString& SettingsString) const;
 
@@ -848,9 +851,6 @@ private:
 
 		/** The string we should be be looking for */
 		FString JumpTerm;
-
-		/** Time delay between performing the last jump, and the jump term being reset */
-		static const double JumpDelaySeconds;
 	};
 
 	/** Data for the asset quick-jump */
