UE4: Thumbnail edit mode will no longer attempt to edit thumbnails for assets that either have no renderer or can not be visualized

diff --git a/Engine/Source/Editor/ContentBrowser/Private/SThumbnailEditModeTools.cpp b/Engine/Source/Editor/ContentBrowser/Private/SThumbnailEditModeTools.cpp
--- a/Engine/Source/Editor/ContentBrowser/Private/SThumbnailEditModeTools.cpp
+++ b/Engine/Source/Editor/ContentBrowser/Private/SThumbnailEditModeTools.cpp
@@ -135,19 +135,21 @@ FReply SThumbnailEditModeTools::OnMouseButtonDown( const FGeometry& MyGeometry,
 	{
 		// Load the asset, unless it is in an unloaded map package or already loaded
 		const FAssetData& AssetData = AssetThumbnail.Pin()->GetAssetData();
-		if ( !AssetData.IsAssetLoaded() )
-		{
-			// Getting the asset loads it.
-			AssetData.GetAsset();
-		}
+		
+		// Getting the asset loads it, if it isn't already.
+		UObject* Asset = AssetData.GetAsset();
 
 		USceneThumbnailInfo* ThumbnailInfo = GetSceneThumbnailInfo();
 		if ( ThumbnailInfo )
 		{
-			bModifiedThumbnailWhileDragging = false;
-			DragStartLocation = FIntPoint(MouseEvent.GetScreenSpacePosition().X, MouseEvent.GetScreenSpacePosition().Y);
+			FThumbnailRenderingInfo* RenderInfo = GUnrealEd->GetThumbnailManager()->GetRenderingInfo(Asset);
+			if (RenderInfo != NULL && RenderInfo->Renderer != NULL)
+			{
+				bModifiedThumbnailWhileDragging = false;
+				DragStartLocation = FIntPoint(MouseEvent.GetScreenSpacePosition().X, MouseEvent.GetScreenSpacePosition().Y);
 
-			return FReply::Handled().CaptureMouse( AsShared() ).UseHighPrecisionMouseMovement( AsShared() );
+				return FReply::Handled().CaptureMouse( AsShared() ).UseHighPrecisionMouseMovement( AsShared() );
+			}
 		}
 
 		// This thumbnail does not have a scene thumbnail info but thumbnail editing is enabled. Just consume the input.
