Add an toggle for showing plugin content, and disable it by default.

diff --git a/Engine/Source/Editor/ContentBrowser/ContentBrowser.Build.cs b/Engine/Source/Editor/ContentBrowser/ContentBrowser.Build.cs
--- a/Engine/Source/Editor/ContentBrowser/ContentBrowser.Build.cs
+++ b/Engine/Source/Editor/ContentBrowser/ContentBrowser.Build.cs
@@ -32,7 +32,8 @@ public class ContentBrowser : ModuleRules
 				"SourceControlWindows",
 				"WorkspaceMenuStructure",
 				"UnrealEd",
-				"EditorWidgets"
+				"EditorWidgets",
+				"Projects"
 			}
 		);
 
diff --git a/Engine/Source/Editor/ContentBrowser/Private/ContentBrowserUtils.cpp b/Engine/Source/Editor/ContentBrowser/Private/ContentBrowserUtils.cpp
--- a/Engine/Source/Editor/ContentBrowser/Private/ContentBrowserUtils.cpp
+++ b/Engine/Source/Editor/ContentBrowser/Private/ContentBrowserUtils.cpp
@@ -1001,6 +1001,19 @@ bool ContentBrowserUtils::IsDevelopersFolder( const FString& InPath )
 	return InPath.StartsWith(DeveloperPathWithSlash) || InPath == DeveloperPathWithoutSlash;
 }
 
+bool ContentBrowserUtils::IsPluginFolder( const FString& InPath )
+{
+	FString PathWithSlash = InPath / TEXT("");
+	for(const FPluginContentFolder& ContentFolder: IPluginManager::Get().GetPluginContentFolders())
+	{
+		if(PathWithSlash.StartsWith(ContentFolder.RootPath) || InPath == ContentFolder.Name)
+		{
+			return true;
+		}
+	}
+	return false;
+}
+
 void ContentBrowserUtils::GetObjectsInAssetData(const TArray<FAssetData>& AssetList, TArray<UObject*>& OutDroppedObjects)
 {	
 	for (int32 AssetIdx = 0; AssetIdx < AssetList.Num(); ++AssetIdx)
diff --git a/Engine/Source/Editor/ContentBrowser/Private/ContentBrowserUtils.h b/Engine/Source/Editor/ContentBrowser/Private/ContentBrowserUtils.h
--- a/Engine/Source/Editor/ContentBrowser/Private/ContentBrowserUtils.h
+++ b/Engine/Source/Editor/ContentBrowser/Private/ContentBrowserUtils.h
@@ -120,6 +120,9 @@ namespace ContentBrowserUtils
 	/** Returns true if the passed-in path is a developers folder */
 	bool IsDevelopersFolder( const FString& InPath );
 
+	/** Returns true if the passed-in path is a plugin folder */
+	bool IsPluginFolder( const FString& InPath );
+
 	/** Get all the objects in a list of asset data */
 	void GetObjectsInAssetData(const TArray<FAssetData>& AssetList, TArray<UObject*>& OutDroppedObjects);
 
diff --git a/Engine/Source/Editor/ContentBrowser/Private/SAssetView.cpp b/Engine/Source/Editor/ContentBrowser/Private/SAssetView.cpp
--- a/Engine/Source/Editor/ContentBrowser/Private/SAssetView.cpp
+++ b/Engine/Source/Editor/ContentBrowser/Private/SAssetView.cpp
@@ -12,6 +12,7 @@
 #include "ContentBrowserModule.h"
 #include "ObjectTools.h"
 #include "KismetEditorUtilities.h"
+#include "IPluginManager.h"
 
 #define LOCTEXT_NAMESPACE "ContentBrowser"
 
@@ -1378,6 +1379,7 @@ void SAssetView::RefreshSourceItems()
 
 	// Remove any assets that should be filtered out any redirectors and non-assets
 	const bool bDisplayEngine = GetDefault<UContentBrowserSettings>()->GetDisplayEngineFolder();
+	const bool bDisplayPlugins = GetDefault<UContentBrowserSettings>()->GetDisplayPluginFolders();
 	for (int32 AssetIdx = Items.Num() - 1; AssetIdx >= 0; --AssetIdx)
 	{
 		const FAssetData& Item = Items[AssetIdx];
@@ -1391,6 +1393,11 @@ void SAssetView::RefreshSourceItems()
 			// If this is an engine folder, and we don't want to show them, remove
 			Items.RemoveAtSwap(AssetIdx);
 		}
+		else if ( !bDisplayPlugins && ContentBrowserUtils::IsPluginFolder(Item.PackagePath.ToString()) )
+		{
+			// If this is a plugin folder, and we don't want to show them, remove
+			Items.RemoveAtSwap(AssetIdx);
+		}
 	}
 }
 
@@ -2130,140 +2137,156 @@ TSharedRef<SWidget> SAssetView::GetViewButtonContent()
 
 	MenuBuilder.BeginSection("AssetViewType", LOCTEXT("ViewTypeHeading", "View Type"));
 	{
-	MenuBuilder.AddMenuEntry(
-		LOCTEXT("TileViewOption", "Tiles"),
-		LOCTEXT("TileViewOptionToolTip", "View assets as tiles in a grid."),
-		FSlateIcon(),
-		FUIAction(
-			FExecuteAction::CreateSP( this, &SAssetView::SetCurrentViewType, EAssetViewType::Tile ),
-			FCanExecuteAction(),
-			FIsActionChecked::CreateSP( this, &SAssetView::IsCurrentViewType, EAssetViewType::Tile )
-			),
-		NAME_None,
-		EUserInterfaceActionType::RadioButton
-		);
+		MenuBuilder.AddMenuEntry(
+			LOCTEXT("TileViewOption", "Tiles"),
+			LOCTEXT("TileViewOptionToolTip", "View assets as tiles in a grid."),
+			FSlateIcon(),
+			FUIAction(
+				FExecuteAction::CreateSP( this, &SAssetView::SetCurrentViewType, EAssetViewType::Tile ),
+				FCanExecuteAction(),
+				FIsActionChecked::CreateSP( this, &SAssetView::IsCurrentViewType, EAssetViewType::Tile )
+				),
+			NAME_None,
+			EUserInterfaceActionType::RadioButton
+			);
 
-	MenuBuilder.AddMenuEntry(
-		LOCTEXT("ListViewOption", "List"),
-		LOCTEXT("ListViewOptionToolTip", "View assets in a list with thumbnails."),
-		FSlateIcon(),
-		FUIAction(
-			FExecuteAction::CreateSP( this, &SAssetView::SetCurrentViewType, EAssetViewType::List ),
-			FCanExecuteAction(),
-			FIsActionChecked::CreateSP( this, &SAssetView::IsCurrentViewType, EAssetViewType::List )
-			),
-		NAME_None,
-		EUserInterfaceActionType::RadioButton
-		);
+		MenuBuilder.AddMenuEntry(
+			LOCTEXT("ListViewOption", "List"),
+			LOCTEXT("ListViewOptionToolTip", "View assets in a list with thumbnails."),
+			FSlateIcon(),
+			FUIAction(
+				FExecuteAction::CreateSP( this, &SAssetView::SetCurrentViewType, EAssetViewType::List ),
+				FCanExecuteAction(),
+				FIsActionChecked::CreateSP( this, &SAssetView::IsCurrentViewType, EAssetViewType::List )
+				),
+			NAME_None,
+			EUserInterfaceActionType::RadioButton
+			);
 
-	MenuBuilder.AddMenuEntry(
-		LOCTEXT("ColumnViewOption", "Columns"),
-		LOCTEXT("ColumnViewOptionToolTip", "View assets in a list with columns of details."),
-		FSlateIcon(),
-		FUIAction(
-			FExecuteAction::CreateSP( this, &SAssetView::SetCurrentViewType, EAssetViewType::Column ),
-			FCanExecuteAction(),
-			FIsActionChecked::CreateSP( this, &SAssetView::IsCurrentViewType, EAssetViewType::Column )
-			),
-		NAME_None,
-		EUserInterfaceActionType::RadioButton
-		);
+		MenuBuilder.AddMenuEntry(
+			LOCTEXT("ColumnViewOption", "Columns"),
+			LOCTEXT("ColumnViewOptionToolTip", "View assets in a list with columns of details."),
+			FSlateIcon(),
+			FUIAction(
+				FExecuteAction::CreateSP( this, &SAssetView::SetCurrentViewType, EAssetViewType::Column ),
+				FCanExecuteAction(),
+				FIsActionChecked::CreateSP( this, &SAssetView::IsCurrentViewType, EAssetViewType::Column )
+				),
+			NAME_None,
+			EUserInterfaceActionType::RadioButton
+			);
 	}
 	MenuBuilder.EndSection();
 
 	MenuBuilder.BeginSection("Folders", LOCTEXT("FoldersHeading", "Folders"));
 	{
-	MenuBuilder.AddMenuEntry(
-		LOCTEXT("ShowFoldersOption", "Show Folders"),
-		LOCTEXT("ShowFoldersOptionToolTip", "Show folders in the view as well as assets."),
-		FSlateIcon(),
-		FUIAction(
-			FExecuteAction::CreateSP( this, &SAssetView::ToggleShowFolders ),
-			FCanExecuteAction::CreateSP( this, &SAssetView::IsToggleShowFoldersAllowed ),
-			FIsActionChecked::CreateSP( this, &SAssetView::IsShowingFolders )
+		MenuBuilder.AddMenuEntry(
+			LOCTEXT("ShowFoldersOption", "Show Folders"),
+			LOCTEXT("ShowFoldersOptionToolTip", "Show folders in the view as well as assets."),
+			FSlateIcon(),
+			FUIAction(
+				FExecuteAction::CreateSP( this, &SAssetView::ToggleShowFolders ),
+				FCanExecuteAction::CreateSP( this, &SAssetView::IsToggleShowFoldersAllowed ),
+				FIsActionChecked::CreateSP( this, &SAssetView::IsShowingFolders )
+				),
+			NAME_None,
+			EUserInterfaceActionType::ToggleButton
+			);
+
+		MenuBuilder.AddMenuEntry(
+			LOCTEXT("ShowOnlyAssetsInSelectedFolders", "Show Only Assets in Selected Folders"),
+			LOCTEXT("ShowOnlyAssetsInSelectedFoldersToolTip", "Only displays the assets of the selected folders"),
+			FSlateIcon(),
+			FUIAction(
+				FExecuteAction::CreateSP( this, &SAssetView::ToggleShowOnlyAssetsInSelectedFolders ),
+				FCanExecuteAction::CreateSP( this, &SAssetView::CanShowOnlyAssetsInSelectedFolders ),
+				FIsActionChecked::CreateSP( this, &SAssetView::IsShowingOnlyAssetsInSelectedFolders )
 			),
-		NAME_None,
-		EUserInterfaceActionType::ToggleButton
-		);
+			NAME_None,
+			EUserInterfaceActionType::ToggleButton
+			);
 
-	MenuBuilder.AddMenuEntry(
-		LOCTEXT("ShowOnlyAssetsInSelectedFolders", "Show Only Assets in Selected Folders"),
-		LOCTEXT("ShowOnlyAssetsInSelectedFoldersToolTip", "Only displays the assets of the selected folders"),
-		FSlateIcon(),
-		FUIAction(
-			FExecuteAction::CreateSP( this, &SAssetView::ToggleShowOnlyAssetsInSelectedFolders ),
-			FCanExecuteAction::CreateSP( this, &SAssetView::CanShowOnlyAssetsInSelectedFolders ),
-			FIsActionChecked::CreateSP( this, &SAssetView::IsShowingOnlyAssetsInSelectedFolders )
-		),
-		NAME_None,
-		EUserInterfaceActionType::ToggleButton
-		);
+		MenuBuilder.AddMenuEntry(
+			LOCTEXT("ShowDevelopersFolderOption", "Show Developers Folder"),
+			LOCTEXT("ShowDevelopersFolderOptionToolTip", "Show the developers folder in the view."),
+			FSlateIcon(),
+			FUIAction(
+				FExecuteAction::CreateSP( this, &SAssetView::ToggleShowDevelopersFolder ),
+				FCanExecuteAction::CreateSP( this, &SAssetView::IsToggleShowDevelopersFolderAllowed ),
+				FIsActionChecked::CreateSP( this, &SAssetView::IsShowingDevelopersFolder )
+			),
+			NAME_None,
+			EUserInterfaceActionType::ToggleButton
+			);
 
-	MenuBuilder.AddMenuEntry(
-		LOCTEXT("ShowDevelopersFolderOption", "Show Developers Folder"),
-		LOCTEXT("ShowDevelopersFolderOptionToolTip", "Show the developers folder in the view."),
-		FSlateIcon(),
-		FUIAction(
-			FExecuteAction::CreateSP( this, &SAssetView::ToggleShowDevelopersFolder ),
-			FCanExecuteAction::CreateSP( this, &SAssetView::IsToggleShowDevelopersFolderAllowed ),
-			FIsActionChecked::CreateSP( this, &SAssetView::IsShowingDevelopersFolder )
-		),
-		NAME_None,
-		EUserInterfaceActionType::ToggleButton
-		);
+		if(IPluginManager::Get().GetPluginContentFolders().Num() > 0)
+		{
+			MenuBuilder.AddMenuEntry(
+				LOCTEXT("ShowPluginFolderOption", "Show Plugin Content"),
+				LOCTEXT("ShowPluginFolderOptionToolTip", "Shows plugin content in the view."),
+				FSlateIcon(),
+				FUIAction(
+					FExecuteAction::CreateSP( this, &SAssetView::ToggleShowPluginFolders ),
+					FCanExecuteAction(),
+					FIsActionChecked::CreateSP( this, &SAssetView::IsShowingPluginFolders )
+				),
+				NAME_None,
+				EUserInterfaceActionType::ToggleButton
+				);
+		}
 
-	MenuBuilder.AddMenuEntry(
-		LOCTEXT("ShowEngineFolderOption", "Show Engine Content"),
-		LOCTEXT("ShowEngineFolderOptionToolTip", "Show the engine content in the view."),
-		FSlateIcon(),
-		FUIAction(
-		FExecuteAction::CreateSP( this, &SAssetView::ToggleShowEngineFolder ),
-		FCanExecuteAction(),
-		FIsActionChecked::CreateSP( this, &SAssetView::IsShowingEngineFolder )
-		),
-		NAME_None,
-		EUserInterfaceActionType::ToggleButton
-		);
+		MenuBuilder.AddMenuEntry(
+			LOCTEXT("ShowEngineFolderOption", "Show Engine Content"),
+			LOCTEXT("ShowEngineFolderOptionToolTip", "Show the engine content in the view."),
+			FSlateIcon(),
+			FUIAction(
+			FExecuteAction::CreateSP( this, &SAssetView::ToggleShowEngineFolder ),
+			FCanExecuteAction(),
+			FIsActionChecked::CreateSP( this, &SAssetView::IsShowingEngineFolder )
+			),
+			NAME_None,
+			EUserInterfaceActionType::ToggleButton
+			);
 	}
 	MenuBuilder.EndSection();
 
 	MenuBuilder.BeginSection("AssetThumbnails", LOCTEXT("ThumbnailsHeading", "Thumbnails"));
 	{
-	MenuBuilder.AddWidget(
-		SNew(SSlider)
-			.ToolTipText( LOCTEXT("ThumbnailScaleToolTip", "Adjust the size of thumbnails.") )
-			.Value( this, &SAssetView::GetThumbnailScale )
-			.OnValueChanged( this, &SAssetView::SetThumbnailScale )
-			.Locked( this, &SAssetView::IsThumbnailScalingLocked ),
-		LOCTEXT("ThumbnailScaleLabel", "Scale"),
-		/*bNoIndent=*/true
-		);
+		MenuBuilder.AddWidget(
+			SNew(SSlider)
+				.ToolTipText( LOCTEXT("ThumbnailScaleToolTip", "Adjust the size of thumbnails.") )
+				.Value( this, &SAssetView::GetThumbnailScale )
+				.OnValueChanged( this, &SAssetView::SetThumbnailScale )
+				.Locked( this, &SAssetView::IsThumbnailScalingLocked ),
+			LOCTEXT("ThumbnailScaleLabel", "Scale"),
+			/*bNoIndent=*/true
+			);
 
-	MenuBuilder.AddMenuEntry(
-		LOCTEXT("ThumbnailEditModeOption", "Thumbnail Edit Mode"),
-		LOCTEXT("ThumbnailEditModeOptionToolTip", "Toggle thumbnail editing mode. When in this mode you can rotate the camera on 3D thumbnails by dragging them."),
-		FSlateIcon(),
-		FUIAction(
-			FExecuteAction::CreateSP( this, &SAssetView::ToggleThumbnailEditMode ),
-			FCanExecuteAction::CreateSP( this, &SAssetView::IsThumbnailEditModeAllowed ),
-			FIsActionChecked::CreateSP( this, &SAssetView::IsThumbnailEditMode )
-			),
-		NAME_None,
-		EUserInterfaceActionType::ToggleButton
-		);
+		MenuBuilder.AddMenuEntry(
+			LOCTEXT("ThumbnailEditModeOption", "Thumbnail Edit Mode"),
+			LOCTEXT("ThumbnailEditModeOptionToolTip", "Toggle thumbnail editing mode. When in this mode you can rotate the camera on 3D thumbnails by dragging them."),
+			FSlateIcon(),
+			FUIAction(
+				FExecuteAction::CreateSP( this, &SAssetView::ToggleThumbnailEditMode ),
+				FCanExecuteAction::CreateSP( this, &SAssetView::IsThumbnailEditModeAllowed ),
+				FIsActionChecked::CreateSP( this, &SAssetView::IsThumbnailEditMode )
+				),
+			NAME_None,
+			EUserInterfaceActionType::ToggleButton
+			);
 
-	MenuBuilder.AddMenuEntry(
-		LOCTEXT("RealTimeThumbnailsOption", "Real-Time Thumbnails"),
-		LOCTEXT("RealTimeThumbnailsOptionToolTip", "Renders the assets thumbnails in real-time"),
-		FSlateIcon(),
-		FUIAction(
-			FExecuteAction::CreateSP( this, &SAssetView::ToggleRealTimeThumbnails ),
-			FCanExecuteAction::CreateSP( this, &SAssetView::CanShowRealTimeThumbnails ),
-			FIsActionChecked::CreateSP( this, &SAssetView::IsShowingRealTimeThumbnails )
-		),
-		NAME_None,
-		EUserInterfaceActionType::ToggleButton
-		);
+		MenuBuilder.AddMenuEntry(
+			LOCTEXT("RealTimeThumbnailsOption", "Real-Time Thumbnails"),
+			LOCTEXT("RealTimeThumbnailsOptionToolTip", "Renders the assets thumbnails in real-time"),
+			FSlateIcon(),
+			FUIAction(
+				FExecuteAction::CreateSP( this, &SAssetView::ToggleRealTimeThumbnails ),
+				FCanExecuteAction::CreateSP( this, &SAssetView::CanShowRealTimeThumbnails ),
+				FIsActionChecked::CreateSP( this, &SAssetView::IsShowingRealTimeThumbnails )
+			),
+			NAME_None,
+			EUserInterfaceActionType::ToggleButton
+			);
 	}
 	MenuBuilder.EndSection();
 
@@ -2320,6 +2343,29 @@ bool SAssetView::IsShowingRealTimeThumbnails() const
 	return CanShowRealTimeThumbnails() ? GetDefault<UContentBrowserSettings>()->RealTimeThumbnails : false;
 }
 
+void SAssetView::ToggleShowPluginFolders()
+{
+	bool bDisplayPlugins = GetDefault<UContentBrowserSettings>()->GetDisplayPluginFolders();
+	bool bRawDisplayPlugins = GetDefault<UContentBrowserSettings>()->GetDisplayPluginFolders( true );
+
+	// Only if both these flags are false when toggling we want to enable the flag, otherwise we're toggling off
+	if ( !bDisplayPlugins && !bRawDisplayPlugins )
+	{
+		GetMutableDefault<UContentBrowserSettings>()->SetDisplayPluginFolders( true );
+	}
+	else
+	{
+		GetMutableDefault<UContentBrowserSettings>()->SetDisplayPluginFolders( false );
+		GetMutableDefault<UContentBrowserSettings>()->SetDisplayPluginFolders( false, true );
+	}	
+	GetMutableDefault<UContentBrowserSettings>()->PostEditChange();
+}
+
+bool SAssetView::IsShowingPluginFolders() const
+{
+	return GetDefault<UContentBrowserSettings>()->GetDisplayPluginFolders();
+}
+
 void SAssetView::ToggleShowEngineFolder()
 {
 	bool bDisplayEngine = GetDefault<UContentBrowserSettings>()->GetDisplayEngineFolder();
diff --git a/Engine/Source/Editor/ContentBrowser/Private/SAssetView.h b/Engine/Source/Editor/ContentBrowser/Private/SAssetView.h
--- a/Engine/Source/Editor/ContentBrowser/Private/SAssetView.h
+++ b/Engine/Source/Editor/ContentBrowser/Private/SAssetView.h
@@ -361,6 +361,12 @@ private:
 	bool IsShowingRealTimeThumbnails() const;
 
 
+	/** Toggle whether plugin content folders should be shown or not */
+	void ToggleShowPluginFolders();
+
+	/** @return true when we are showing plugin content folders */
+	bool IsShowingPluginFolders() const;
+
 	/** Toggle whether the engine folder should be shown or not */
 	void ToggleShowEngineFolder();
 
diff --git a/Engine/Source/Editor/ContentBrowser/Private/SContentBrowser.cpp b/Engine/Source/Editor/ContentBrowser/Private/SContentBrowser.cpp
--- a/Engine/Source/Editor/ContentBrowser/Private/SContentBrowser.cpp
+++ b/Engine/Source/Editor/ContentBrowser/Private/SContentBrowser.cpp
@@ -641,25 +641,35 @@ void SContentBrowser::SyncToAssets( const TArray<FAssetData>& AssetDataList, con
 	const UContentBrowserSettings* tmp = GetDefault<UContentBrowserSettings>();
 	bool bDisplayDev = GetDefault<UContentBrowserSettings>()->GetDisplayDevelopersFolder();
 	bool bDisplayEngine = GetDefault<UContentBrowserSettings>()->GetDisplayEngineFolder();
-	if ( !bDisplayDev || !bDisplayEngine )
+	bool bDisplayPlugins = GetDefault<UContentBrowserSettings>()->GetDisplayPluginFolders();
+	if ( !bDisplayDev || !bDisplayEngine || !bDisplayPlugins )
 	{
-		for (int32 AssetIdx = AssetDataList.Num() - 1; AssetIdx >= 0 && ( !bDisplayDev || !bDisplayEngine ); --AssetIdx)
+		bool bRepopulate = false;
+		for (int32 AssetIdx = AssetDataList.Num() - 1; AssetIdx >= 0 && ( !bDisplayDev || !bDisplayEngine || !bDisplayPlugins ); --AssetIdx)
 		{
 			const FAssetData& Item = AssetDataList[AssetIdx];
 			if ( !bDisplayDev && ContentBrowserUtils::IsDevelopersFolder( Item.PackagePath.ToString() ) )
 			{
 				bDisplayDev = true;
 				GetMutableDefault<UContentBrowserSettings>()->SetDisplayDevelopersFolder(true, true);
+				bRepopulate = true;
 			}
 			else if ( !bDisplayEngine && ContentBrowserUtils::IsEngineFolder( Item.PackagePath.ToString() ) )
 			{
 				bDisplayEngine = true;
 				GetMutableDefault<UContentBrowserSettings>()->SetDisplayEngineFolder(true, true);
+				bRepopulate = true;
+			}
+			else if ( !bDisplayPlugins && ContentBrowserUtils::IsPluginFolder( Item.PackagePath.ToString() ) )
+			{
+				bDisplayPlugins = true;
+				GetMutableDefault<UContentBrowserSettings>()->SetDisplayPluginFolders(true, true);
+				bRepopulate = true;
 			}
 		}
 
 		// If we have auto-enabled any flags, force a refresh
-		if ( bDisplayDev || bDisplayEngine )
+		if ( bRepopulate )
 		{
 			PathViewPtr->Populate();
 		}
@@ -1727,32 +1737,34 @@ void SContentBrowser::HandleSettingChanged(FName PropertyName)
 {
 	if ((PropertyName == "DisplayDevelopersFolder") ||
 		(PropertyName == "DisplayEngineFolder") ||
+		(PropertyName == "DisplayPluginFolders") ||
 		(PropertyName == NAME_None))	// @todo: Needed if PostEditChange was called manually, for now
 	{
 		// If the dev or engine folder is no longer visible but we're inside it...
 		const bool bDisplayDev = GetDefault<UContentBrowserSettings>()->GetDisplayDevelopersFolder();
 		const bool bDisplayEngine = GetDefault<UContentBrowserSettings>()->GetDisplayEngineFolder();
-		if ( !bDisplayDev || !bDisplayEngine )
+		const bool bDisplayPlugins = GetDefault<UContentBrowserSettings>()->GetDisplayPluginFolders();
+		if ( !bDisplayDev || !bDisplayEngine || !bDisplayPlugins )
 		{
-			const FString OldSelectedPath = PathViewPtr->GetSelectedPath();
-			if ( (!bDisplayDev && ContentBrowserUtils::IsDevelopersFolder( OldSelectedPath )) || (!bDisplayEngine && ContentBrowserUtils::IsEngineFolder( OldSelectedPath )) )
-			{
-				// Set the folder back to the root, and refresh the contents
-				TArray<FString> SelectedPaths;
-				SelectedPaths.Add(TEXT("/Game"));
-				PathViewPtr->SetSelectedPaths(SelectedPaths);
-				SourcesChanged(SelectedPaths, TArray<FCollectionNameType>());
-			}
+		    const FString OldSelectedPath = PathViewPtr->GetSelectedPath();
+		    if ( (!bDisplayDev && ContentBrowserUtils::IsDevelopersFolder( OldSelectedPath )) || (!bDisplayEngine && ContentBrowserUtils::IsEngineFolder( OldSelectedPath )) || (!bDisplayPlugins && ContentBrowserUtils::IsPluginFolder( OldSelectedPath ) ) )
+		    {
+			    // Set the folder back to the root, and refresh the contents
+			    TArray<FString> SelectedPaths;
+			    SelectedPaths.Add(TEXT("/Game"));
+			    PathViewPtr->SetSelectedPaths(SelectedPaths);
+			    SourcesChanged(SelectedPaths, TArray<FCollectionNameType>());
+		    }
 		}
 
 		// Update our path view so that it can include/exclude the dev folder
 		PathViewPtr->Populate();
 
 		// If the dev or engine folder has become visible and we're inside it...
-		if ( bDisplayDev || bDisplayEngine )
+		if ( bDisplayDev || bDisplayEngine || bDisplayPlugins )
 		{
 			const FString NewSelectedPath = PathViewPtr->GetSelectedPath();
-			if ( (bDisplayDev && ContentBrowserUtils::IsDevelopersFolder( NewSelectedPath )) || (bDisplayEngine && ContentBrowserUtils::IsEngineFolder( NewSelectedPath )) )
+			if ( (bDisplayDev && ContentBrowserUtils::IsDevelopersFolder( NewSelectedPath )) || (bDisplayEngine && ContentBrowserUtils::IsEngineFolder( NewSelectedPath )) || (bDisplayPlugins && ContentBrowserUtils::IsPluginFolder( NewSelectedPath ) ) )
 			{
 				// Refresh the contents
 				TArray<FString> SelectedPaths;
diff --git a/Engine/Source/Editor/ContentBrowser/Private/SPathView.cpp b/Engine/Source/Editor/ContentBrowser/Private/SPathView.cpp
--- a/Engine/Source/Editor/ContentBrowser/Private/SPathView.cpp
+++ b/Engine/Source/Editor/ContentBrowser/Private/SPathView.cpp
@@ -712,9 +712,13 @@ TSharedPtr<struct FTreeItem> SPathView::AddRootItem( const FString& InFolderName
 	const bool bDisplayEngine = GetDefault<UContentBrowserSettings>()->GetDisplayEngineFolder();
 	if ( bDisplayEngine || !ContentBrowserUtils::IsEngineFolder(InFolderName) )
 	{
-		NewItem = MakeShareable( new FTreeItem(InFolderName, FString(TEXT("/")) + InFolderName, TSharedPtr<FTreeItem>()));
-		TreeRootItems.Add( NewItem );
-		TreeViewPtr->RequestTreeRefresh();
+		const bool bDisplayPlugins = GetDefault<UContentBrowserSettings>()->GetDisplayPluginFolders();
+		if ( bDisplayPlugins || !ContentBrowserUtils::IsPluginFolder(InFolderName) )
+		{
+			NewItem = MakeShareable( new FTreeItem(InFolderName, FString(TEXT("/")) + InFolderName, TSharedPtr<FTreeItem>()));
+			TreeRootItems.Add( NewItem );
+			TreeViewPtr->RequestTreeRefresh();
+		}
 	}
 
 	return NewItem;
diff --git a/Engine/Source/Editor/UnrealEd/Classes/Settings/ContentBrowserSettings.h b/Engine/Source/Editor/UnrealEd/Classes/Settings/ContentBrowserSettings.h
--- a/Engine/Source/Editor/UnrealEd/Classes/Settings/ContentBrowserSettings.h
+++ b/Engine/Source/Editor/UnrealEd/Classes/Settings/ContentBrowserSettings.h
@@ -62,6 +62,17 @@ public:
 		return ( ( bExcludeOverride ? false : OverrideDisplayDevelopersFolder ) | DisplayDevelopersFolder );
 	}
 
+	/** Sets whether we are allowed to display the plugin folders or not, optional flag for setting override instead */
+	void SetDisplayPluginFolders( bool bInDisplayPluginFolders, bool bOverride = false )
+	{ 
+		bOverride ? OverrideDisplayPluginFolders = bInDisplayPluginFolders : DisplayPluginFolders = bInDisplayPluginFolders;
+	}
+	/** Gets whether we are allowed to display the plugin folders or not, optional flag ignoring the override */
+	bool GetDisplayPluginFolders( bool bExcludeOverride = false ) const
+	{ 
+		return ( ( bExcludeOverride ? false : OverrideDisplayPluginFolders ) | DisplayPluginFolders );
+	}
+
 	/**
 	 * Returns an event delegate that is executed when a setting has changed.
 	 *
@@ -94,6 +105,13 @@ private:
 	/** If true, overrides the DisplayDev setting */
 	bool OverrideDisplayDevelopersFolder;
 
+	/** List of plugin folders to display in the content browser. */
+	UPROPERTY(config)
+	bool DisplayPluginFolders;
+
+	/** Temporary override for the DisplayPluginFolders setting */
+	bool OverrideDisplayPluginFolders;
+
 	// Holds an event delegate that is executed when a setting has changed.
 	static FSettingChangedEvent SettingChangedEvent;
 };
diff --git a/Engine/Source/Runtime/Projects/Private/PluginManager.cpp b/Engine/Source/Runtime/Projects/Private/PluginManager.cpp
--- a/Engine/Source/Runtime/Projects/Private/PluginManager.cpp
+++ b/Engine/Source/Runtime/Projects/Private/PluginManager.cpp
@@ -27,23 +27,6 @@ FPluginInstance::FPluginInstance(const FString& InFileName, const FPluginDescrip
 {
 }
 
-void FPluginInstance::RegisterPluginMountPoints( const IPluginManager::FRegisterMountPointDelegate& RegisterMountPointDelegate )
-{
-	// If this plugin has content, mount it before loading the first module it contains
-	// @todo plugin: Should we do this even if there is no Content directory?  Probably, so that we can add content in Content Browser!
-	if( Descriptor.bCanContainContent )
-	{
-		if( ensure( RegisterMountPointDelegate.IsBound() ) )
-		{
-			// @todo plugin content: Consider prefixing path with "Plugins" or more specifically, either "EnginePlugins" or "GamePlugins", etc.
-			// @todo plugin content: We may need the full relative path to the plugin under the respective Plugins directory if we want to fully avoid collisions
-			const FString ContentRootPath( FString::Printf( TEXT( "/%s/" ), *Name ) );
-			const FString PluginContentDirectory( FPaths::GetPath(FileName) / TEXT( "Content" ) );
-
-			RegisterMountPointDelegate.Execute( ContentRootPath, PluginContentDirectory );
-		}
-	}
-}
 
 
 
@@ -291,6 +274,7 @@ void FPluginManager::EnablePluginsThatAreConfiguredToBeEnabled()
 	TSet< FString > AllEnabledPlugins;
 	AllEnabledPlugins.Append( MoveTemp(EnabledPluginNames) );
 
+	// Enable all the plugins by name
 	for( const TSharedRef< FPluginInstance > Plugin : AllPlugins )
 	{
 		if ( AllEnabledPlugins.Contains(Plugin->Name) )
@@ -298,15 +282,27 @@ void FPluginManager::EnablePluginsThatAreConfiguredToBeEnabled()
 			Plugin->bEnabled = true;
 		}
 	}
-}
 
-void FPluginManager::RegisterEnabledPluginMountPoints()
-{
-	for( const TSharedRef< FPluginInstance > Plugin : AllPlugins )
+	// Build the list of content folders
+	ContentFolders.Empty();
+	for(const TSharedRef<FPluginInstance>& Plugin: AllPlugins)
 	{
-		if ( Plugin->bEnabled )
+		if(Plugin->bEnabled && Plugin->Descriptor.bCanContainContent)
 		{
-			Plugin->RegisterPluginMountPoints( RegisterMountPointDelegate );
+			FPluginContentFolder ContentFolder;
+			ContentFolder.Name = Plugin->Name;
+			ContentFolder.RootPath = FString::Printf(TEXT("/%s/"), *Plugin->Name);
+			ContentFolder.ContentPath = FPaths::GetPath(Plugin->FileName) / TEXT("Content");
+			ContentFolders.Emplace(ContentFolder);
+		}
+	}
+
+	// Mount all the plugin content folders
+	if( ensure( RegisterMountPointDelegate.IsBound() ) )
+	{
+		for(const FPluginContentFolder& ContentFolder: ContentFolders)
+		{
+			RegisterMountPointDelegate.Execute(ContentFolder.RootPath, ContentFolder.ContentPath);
 		}
 	}
 }
@@ -316,7 +312,6 @@ void FPluginManager::LoadModulesForEnabledPlugins( const ELoadingPhase::Type Loa
 	if ( !bEarliestPhaseProcessed )
 	{
 		EnablePluginsThatAreConfiguredToBeEnabled();
-		RegisterEnabledPluginMountPoints();
 		bEarliestPhaseProcessed = true;
 	}
 
@@ -389,7 +384,6 @@ bool FPluginManager::AreEnabledPluginModulesUpToDate()
 	if (!bEarliestPhaseProcessed)
 	{
 		EnablePluginsThatAreConfiguredToBeEnabled();
-		RegisterEnabledPluginMountPoints();
 		bEarliestPhaseProcessed = true;
 	}
 
@@ -439,6 +433,7 @@ TArray< FPluginStatus > FPluginManager::QueryStatusForAllPlugins() const
 		PluginStatus.bIsBuiltIn = ( Plugin->LoadedFrom == EPluginLoadedFrom::Engine );
 		PluginStatus.bIsEnabledByDefault = PluginInfo.bEnabledByDefault;
 		PluginStatus.bIsBetaVersion = PluginInfo.bIsBetaVersion;
+		PluginStatus.bHasContentFolder = PluginInfo.bCanContainContent;
 
 		// @todo plugedit: Maybe we should do the FileExists check ONCE at plugin load time and not at query time
 		const FString Icon128FilePath = FPaths::GetPath(Plugin->FileName) / PluginSystemDefs::RelativeIcon128FilePath;
@@ -453,4 +448,9 @@ TArray< FPluginStatus > FPluginManager::QueryStatusForAllPlugins() const
 	return PluginStatuses;
 }
 
+const TArray<FPluginContentFolder>& FPluginManager::GetPluginContentFolders() const
+{
+	return ContentFolders;
+}
+
 #undef LOCTEXT_NAMESPACE
\ No newline at end of file
diff --git a/Engine/Source/Runtime/Projects/Private/PluginManager.h b/Engine/Source/Runtime/Projects/Private/PluginManager.h
--- a/Engine/Source/Runtime/Projects/Private/PluginManager.h
+++ b/Engine/Source/Runtime/Projects/Private/PluginManager.h
@@ -44,13 +44,6 @@ public:
 	 * FPlugin constructor
 	 */
 	FPluginInstance(const FString &FileName, const FPluginDescriptor& InDescriptor, EPluginLoadedFrom::Type InLoadedFrom);
-
-	/**
-	 * Registers the content mount points for this delegate. This allows content to be created in the root folder for this plugin.
-	 *
-	 * @param	RegisterMountPointDelegate	Delegate for mounting content paths.  Bound by FPackageName code in CoreUObject, so that we can access content path mounting functionality from Core.
-	 */
-	void RegisterPluginMountPoints( const IPluginManager::FRegisterMountPointDelegate& RegisterMountPointDelegate );
 };
 
 /**
@@ -71,6 +64,7 @@ public:
 	virtual bool IsPluginModule( const FName ModuleName ) const override;
 	virtual bool AreEnabledPluginModulesUpToDate() override;
 	virtual TArray< FPluginStatus > QueryStatusForAllPlugins() const override;
+	virtual const TArray< FPluginContentFolder >& GetPluginContentFolders() const override;
 
 private:
 
@@ -81,9 +75,6 @@ private:
 	/** Sets the bPluginEnabled flag on all plugins found from DiscoverAllPlugins that are enabled in config */
 	void EnablePluginsThatAreConfiguredToBeEnabled();
 
-	/** Registers content mount points for all enabled plugins */
-	void RegisterEnabledPluginMountPoints();
-
 private:
 	/** All of the plugins that we know about */
 	TArray< TSharedRef< FPluginInstance > > AllPlugins;
@@ -91,6 +82,9 @@ private:
 	/** Map from module name to plugin containing that module */
 	TMap< FName, TSharedRef< FPluginInstance > > ModulePluginMap;
 
+	/** All the plugin content folders */
+	TArray<FPluginContentFolder> ContentFolders;
+
 	/** Delegate for mounting content paths.  Bound by FPackageName code in CoreUObject, so that we can access
 	    content path mounting functionality from Core. */
 	FRegisterMountPointDelegate RegisterMountPointDelegate;
diff --git a/Engine/Source/Runtime/Projects/Public/Interfaces/IPluginManager.h b/Engine/Source/Runtime/Projects/Public/Interfaces/IPluginManager.h
--- a/Engine/Source/Runtime/Projects/Public/Interfaces/IPluginManager.h
+++ b/Engine/Source/Runtime/Projects/Public/Interfaces/IPluginManager.h
@@ -56,9 +56,25 @@ public:
 	/** Marks the plugin as beta in the UI */
 	bool bIsBetaVersion;
 
+	/** Whether the plugin has a content folder */
+	bool bHasContentFolder;
+
 };
 
+/**
+ * Structure holding information about a plugin content folder
+ */
+struct FPluginContentFolder
+{
+	/** Name of the plugin */
+	FString Name;
 
+	/** Virtual root path for asset paths */
+	FString RootPath;
+
+	/** Content path on disk */
+	FString ContentPath;
+};
 
 /**
  * PluginManager manages available code and content extensions (both loaded and not loaded.)
@@ -116,4 +132,11 @@ public:
 	 * @return	 Array of plugin status objects
 	 */
 	virtual TArray< FPluginStatus > QueryStatusForAllPlugins() const = 0;
+
+	/**
+	 * Gets a list of plugin content folders
+	 *
+	 * @return	 Array of plugin content folders
+	 */
+	virtual const TArray< FPluginContentFolder >& GetPluginContentFolders() const = 0;
 };
