Content Browswer - Fix a crash when dragging the root folder, and also disallow moving a parent folder to a child folder.

diff --git a/Engine/Source/Editor/ContentBrowser/Private/ContentBrowserUtils.cpp b/Engine/Source/Editor/ContentBrowser/Private/ContentBrowserUtils.cpp
--- a/Engine/Source/Editor/ContentBrowser/Private/ContentBrowserUtils.cpp
+++ b/Engine/Source/Editor/ContentBrowser/Private/ContentBrowserUtils.cpp
@@ -701,9 +701,20 @@ void ContentBrowserUtils::MoveFolders(const TArray<FString>& InSourcePathNames,
 {
 	TMap<FString, TArray<UObject*> > SourcePathToLoadedAssets;
 
-	// Make sure the destination path is not in the source path list
+	// Do not allow parent directories to be moved to themselves or children.
 	TArray<FString> SourcePathNames = InSourcePathNames;
-	SourcePathNames.Remove(DestPath);
+	TArray<FString> SourcePathNamesToRemove;
+	for (auto SourcePathIt = SourcePathNames.CreateConstIterator(); SourcePathIt; ++SourcePathIt)
+	{
+		if (DestPath.StartsWith(*SourcePathIt))
+		{
+			SourcePathNamesToRemove.Add(*SourcePathIt);
+		}
+	}
+	for (auto SourcePathToRemoveIt = SourcePathNamesToRemove.CreateConstIterator(); SourcePathToRemoveIt; ++SourcePathToRemoveIt)
+	{
+		SourcePathNames.Remove(*SourcePathToRemoveIt);
+	}
 
 	// Load all assets in the source paths
 	PrepareFoldersForDragDrop(SourcePathNames, SourcePathToLoadedAssets);
@@ -782,14 +793,14 @@ void ContentBrowserUtils::PrepareFoldersForDragDrop(const TArray<FString>& Sourc
 		TArray<UObject*> AllLoadedAssets;
 		LoadAssetsIfNeeded(ObjectPaths, AllLoadedAssets, false);
 
+		// Add a slash to the end of the path so StartsWith doesn't get a false positive on similarly named folders
+		const FString SourcePathWithSlash = *PathIt + TEXT("/");
+
 		// Find all files in this path and subpaths
 		TArray<FString> Filenames;
-		FString RootFolder = FPackageName::LongPackageNameToFilename(*PathIt);
+		FString RootFolder = FPackageName::LongPackageNameToFilename(SourcePathWithSlash);
 		FPackageName::FindPackagesInDirectory(Filenames, RootFolder);
 
-		// Add a slash to the end of the path so StartsWith doesn't get a false positive on similarly named folders
-		const FString SourcePathWithSlash = *PathIt + TEXT("/");
-
 		// Now find all assets in memory that were loaded from this path that are valid for drag-droppping
 		TArray<UObject*> ValidLoadedAssets;
 		for ( auto AssetIt = AllLoadedAssets.CreateConstIterator(); AssetIt; ++AssetIt )
