CB - When front end filters change, if we were not previously performing a recursive filter, we now perform a full slow refresh in order to populate the data the front end filters operate on.  This fixes the issue with front end filters like "In Use By Level" not working until folders are changed.

TTP 341204

diff --git a/Engine/Source/Editor/ContentBrowser/Private/SAssetView.cpp b/Engine/Source/Editor/ContentBrowser/Private/SAssetView.cpp
--- a/Engine/Source/Editor/ContentBrowser/Private/SAssetView.cpp
+++ b/Engine/Source/Editor/ContentBrowser/Private/SAssetView.cpp
@@ -187,6 +187,7 @@ void SAssetView::Construct( const FArguments& InArgs )
 	bThumbnailEditMode = false;
 	bUserSearching = false;
 	bPendingFocusOnSync = false;
+	bWereItemsRecursivelyFiltered = false;
 
 	TagColumnRenames.Add("ResourceSize", TEXT("Size (kb)"));
 
@@ -1294,6 +1295,7 @@ void SAssetView::RefreshSourceItems()
 	{
 		AssetRegistryModule.Get().GetAllAssets(Items);
 		bWantToShowShowClasses = true;
+		bWereItemsRecursivelyFiltered = true;
 	}
 	else
 	{
@@ -1303,6 +1305,8 @@ void SAssetView::RefreshSourceItems()
 		const bool bUsingFolders = GetDefault<UContentBrowserSettings>()->ShowOnlyAssetsInSelectedFolders || IsShowingFolders();
 		FARFilter Filter = SourcesData.MakeFilter(bRecurse, bUsingFolders);
 
+		bWereItemsRecursivelyFiltered = bRecurse;
+
 		// Remove the classes path if it is in the list. We will add classes to the results later
 		bWantToShowShowClasses = Filter.PackagePaths.Remove(TEXT("/Classes")) > 0;
 
@@ -2017,6 +2021,13 @@ void SAssetView::OnObjectPropertyChanged(UObject* Asset, FPropertyChangedEvent&
 void SAssetView::OnFrontendFiltersChanged()
 {
 	RequestQuickFrontendListRefresh();
+
+	// If we're not operating on recursively filtered data, we need to ensure a full slow
+	// refresh is performed.
+	if ( ShouldFilterRecursively() && !bWereItemsRecursivelyFiltered )
+	{
+		RequestSlowFullListRefresh();
+	}
 }
 
 bool SAssetView::IsFrontendFilterActive() const
diff --git a/Engine/Source/Editor/ContentBrowser/Private/SAssetView.h b/Engine/Source/Editor/ContentBrowser/Private/SAssetView.h
--- a/Engine/Source/Editor/ContentBrowser/Private/SAssetView.h
+++ b/Engine/Source/Editor/ContentBrowser/Private/SAssetView.h
@@ -643,6 +643,9 @@ private:
 	/** Should we take focus when the PendingSyncAssets are processed? */
 	bool bPendingFocusOnSync;
 
+	/** Was recursive filtering enabled the last time a full slow refresh performed? */
+	bool bWereItemsRecursivelyFiltered;
+
 	/** Called to check if an asset should be filtered out by external code */
 	FOnShouldFilterAsset OnShouldFilterAsset;
 
