Content Browser now reset search filter when syncing asset to content browser

#ttp 340916 - Editor: Not able to find asset in contentbrowser when clicking the magnify button

#branch UE4

#change If any text was in the search field, sync to asset would navigate to the correct folder, but would be unable to show the asset. Asset type filters are already reset when this happens, but the search text was not

diff --git a/Engine/Source/Editor/ContentBrowser/Private/SContentBrowser.cpp b/Engine/Source/Editor/ContentBrowser/Private/SContentBrowser.cpp
--- a/Engine/Source/Editor/ContentBrowser/Private/SContentBrowser.cpp
+++ b/Engine/Source/Editor/ContentBrowser/Private/SContentBrowser.cpp
@@ -675,8 +675,14 @@ void SContentBrowser::SyncToAssets( const TArray<FAssetData>& AssetDataList, con
 		}
 	}
 
+	// Disable the filter categories
 	FilterListPtr->DisableFiltersThatHideAssets(AssetDataList);
 
+	// Disable the filter search (reset the filter, then clear the search text)
+	// Note: we have to remove the filter immediately, we can't wait for OnSearchBoxChanged to hit
+	SetSearchBoxText(FText::GetEmpty());
+	SearchBoxPtr->SetText(FText::GetEmpty());	
+
 	// Tell the sources view first so the asset view will be up to date by the time we request the sync
 	PathViewPtr->SyncToAssets(AssetDataList, bAllowImplicitSync);
 	AssetViewPtr->SyncToAssets(AssetDataList);
@@ -1060,7 +1066,7 @@ void SContentBrowser::NewFolderRequested(const FString& SelectedPath)
 	}
 }
 
-void SContentBrowser::OnSearchBoxChanged(const FText& InSearchText)
+void SContentBrowser::SetSearchBoxText(const FText& InSearchText)
 {
 	TextFilter->SetRawFilterText( InSearchText );
 	if(InSearchText.IsEmpty())
@@ -1073,6 +1079,11 @@ void SContentBrowser::OnSearchBoxChanged(const FText& InSearchText)
 		FrontendFilters->Add(TextFilter);
 		AssetViewPtr->SetUserSearching(true);
 	}
+}
+
+void SContentBrowser::OnSearchBoxChanged(const FText& InSearchText)
+{
+	SetSearchBoxText(InSearchText);
 
 	// Broadcast 'search box changed' delegate
 	FContentBrowserModule& ContentBrowserModule = FModuleManager::GetModuleChecked<FContentBrowserModule>( TEXT("ContentBrowser") );
@@ -1081,17 +1092,7 @@ void SContentBrowser::OnSearchBoxChanged(const FText& InSearchText)
 
 void SContentBrowser::OnSearchBoxCommitted(const FText& InSearchText, ETextCommit::Type CommitInfo)
 {
-	TextFilter->SetRawFilterText( InSearchText );
-	if(InSearchText.IsEmpty())
-	{
-		FrontendFilters->Remove(TextFilter);
-		AssetViewPtr->SetUserSearching(false);
-	}
-	else
-	{
-		FrontendFilters->Add(TextFilter);
-		AssetViewPtr->SetUserSearching(true);
-	}
+	SetSearchBoxText(InSearchText);
 }
 
 void SContentBrowser::OnPathClicked( const FString& CrumbData )
diff --git a/Engine/Source/Editor/ContentBrowser/Private/SContentBrowser.h b/Engine/Source/Editor/ContentBrowser/Private/SContentBrowser.h
--- a/Engine/Source/Editor/ContentBrowser/Private/SContentBrowser.h
+++ b/Engine/Source/Editor/ContentBrowser/Private/SContentBrowser.h
@@ -121,6 +121,9 @@ private:
 	/** Handler for when the path context menu requests a folder creation */
 	void NewFolderRequested(const FString& SelectedPath);
 
+	/** Called when the editable text needs to be set or cleared */
+	void SetSearchBoxText(const FText& InSearchText);
+
 	/** Called by the editable text control when the search text is changed by the user */
 	void OnSearchBoxChanged(const FText& InSearchText);
 
