Fixed Forward To and Back To arrows in content browser, so they skip deleted collections folders and paths.
#ttp 267036 - CONTENTBROWSER: Forward To and Back To arrows will go to deleted collections folders
#branch UE4
#proj Editor.ContentBrowser

diff --git a/Engine/Source/Editor/ContentBrowser/Private/HistoryManager.h b/Engine/Source/Editor/ContentBrowser/Private/HistoryManager.h
--- a/Engine/Source/Editor/ContentBrowser/Private/HistoryManager.h
+++ b/Engine/Source/Editor/ContentBrowser/Private/HistoryManager.h
@@ -67,6 +67,25 @@ public:
 	 */
 	void GetAvailableHistoryMenuItems(bool bGetPrior, FMenuBuilder& MenuBuilder);
 
+	/** Removes all history data as determined by the passed in predicate */
+	template <class PREDICATE_CLASS>
+	void RemoveHistoryData(const PREDICATE_CLASS& Predicate)
+	{
+		for (int32 HistoryIndex = 0; HistoryIndex < HistoryData.Num(); HistoryIndex++)
+		{
+			if (Predicate(HistoryData[HistoryIndex]))
+			{
+				HistoryData.RemoveAt(HistoryIndex);
+				if (CurrentHistoryIndex >= HistoryIndex)
+				{
+					// Ensure if possible that current history index continues to point to the same item if something is removed before it.
+					CurrentHistoryIndex = FMath::Max(0, CurrentHistoryIndex - 1);
+				}
+				HistoryIndex--;
+			}
+		}
+	}
+
 private:
 	/** Notifies the owner to update to the state described by the current history data */
 	void ApplyCurrentHistoryData();
diff --git a/Engine/Source/Editor/ContentBrowser/Private/SCollectionView.cpp b/Engine/Source/Editor/ContentBrowser/Private/SCollectionView.cpp
--- a/Engine/Source/Editor/ContentBrowser/Private/SCollectionView.cpp
+++ b/Engine/Source/Editor/ContentBrowser/Private/SCollectionView.cpp
@@ -413,6 +413,9 @@ void SCollectionView::RemoveCollectionItems( const TArray<TSharedPtr<FCollection
 		}
 	}
 
+	FScopedPreventSelectionChangedDelegate DelegatePrevention(SharedThis(this));
+	CollectionListPtr->ClearSelection();
+
 	// If we removed all the selected items and there is at least one other item, select it
 	if ( LastSelectedItemIdx > INDEX_NONE && NumSelectedItemsRemoved > 0 && NumSelectedItemsRemoved >= SelectedItems.Num() && CollectionItems.Num() > 1 )
 	{
diff --git a/Engine/Source/Editor/ContentBrowser/Private/SContentBrowser.cpp b/Engine/Source/Editor/ContentBrowser/Private/SContentBrowser.cpp
--- a/Engine/Source/Editor/ContentBrowser/Private/SContentBrowser.cpp
+++ b/Engine/Source/Editor/ContentBrowser/Private/SContentBrowser.cpp
@@ -8,6 +8,8 @@
 #include "ContentBrowserModule.h"
 #include "STutorialWrapper.h"
 #include "ContentBrowserCommands.h"
+#include "CollectionManagerModule.h"
+#include "AssetRegistryModule.h"
 
 #define LOCTEXT_NAMESPACE "ContentBrowser"
 
@@ -17,6 +19,20 @@ SContentBrowser::~SContentBrowser()
 {
 	// Remove the listener for when view settings are changed
 	UContentBrowserSettings::OnSettingChanged().RemoveAll( this );
+
+	// Remove listeners for when collections/paths are renamed/deleted
+	FCollectionManagerModule* CollectionManagerModule = FModuleManager::GetModulePtr<FCollectionManagerModule>(TEXT("CollectionManager"));
+	if (CollectionManagerModule != nullptr)
+	{
+		CollectionManagerModule->Get().OnCollectionRenamed().RemoveAll(this);
+		CollectionManagerModule->Get().OnCollectionDestroyed().RemoveAll(this);
+	}
+
+	FAssetRegistryModule* AssetRegistryModule = FModuleManager::GetModulePtr<FAssetRegistryModule>(TEXT("AssetRegistry"));
+	if (AssetRegistryModule != nullptr)
+	{
+		AssetRegistryModule->Get().OnPathRemoved().RemoveAll(this);
+	}
 }
 
 BEGIN_SLATE_FUNCTION_BUILD_OPTIMIZATION
@@ -573,6 +589,14 @@ void SContentBrowser::Construct( const FArguments& InArgs, const FName& InInstan
 	this->InstanceName = InInstanceName;
 	LoadSettings(InInstanceName);
 
+	// Bindings to manage history when items are deleted
+	FCollectionManagerModule& CollectionManagerModule = FModuleManager::LoadModuleChecked<FCollectionManagerModule>(TEXT("CollectionManager"));
+	CollectionManagerModule.Get().OnCollectionRenamed().AddSP(this, &SContentBrowser::HandleCollectionRenamed);
+	CollectionManagerModule.Get().OnCollectionDestroyed().AddSP(this, &SContentBrowser::HandleCollectionRemoved);
+
+	FAssetRegistryModule& AssetRegistryModule = FModuleManager::LoadModuleChecked<FAssetRegistryModule>(TEXT("AssetRegistry"));
+	AssetRegistryModule.Get().OnPathRemoved().AddSP(this, &SContentBrowser::HandlePathRemoved);
+
 	// Update the breadcrumb trail path
 	UpdatePath();
 }
@@ -920,8 +944,11 @@ void SContentBrowser::SourcesChanged(const TArray<FString>& SelectedPaths, const
 	
 	SourcesData.Collections = SelectedCollections;
 
-	// Update the current history data to preserve selection
-	HistoryManager.UpdateHistoryData();
+	if (!AssetViewPtr->GetSourcesData().IsEmpty())
+	{
+		// Update the current history data to preserve selection if there is a valid SourcesData
+		HistoryManager.UpdateHistoryData();
+	}
 
 	// Change the filter for the asset view
 	AssetViewPtr->SetSourcesData(SourcesData);
@@ -1776,6 +1803,41 @@ void SContentBrowser::HandleSettingChanged(FName PropertyName)
 	}
 }
 
+void SContentBrowser::HandleCollectionRemoved(const FCollectionNameType& Collection)
+{
+	AssetViewPtr->SetSourcesData(FSourcesData());
+
+	auto RemoveHistoryDelegate = [&](const FHistoryData& HistoryData)
+	{
+		return (HistoryData.SourcesData.Collections.Num() == 1 &&
+				HistoryData.SourcesData.PackagePaths.Num() == 0 &&
+				HistoryData.SourcesData.Collections.Contains(Collection));
+	};
+
+	HistoryManager.RemoveHistoryData(RemoveHistoryDelegate);
+}
+
+void SContentBrowser::HandleCollectionRenamed(const FCollectionNameType& OriginalCollection, const FCollectionNameType& NewCollection)
+{
+	return HandleCollectionRemoved(OriginalCollection);
+}
+
+void SContentBrowser::HandlePathRemoved(const FString& Path)
+{
+	const FName PathName(*Path);
+
+	AssetViewPtr->SetSourcesData(FSourcesData());
+
+	auto RemoveHistoryDelegate = [&](const FHistoryData& HistoryData)
+	{
+		return (HistoryData.SourcesData.PackagePaths.Num() == 1 &&
+				HistoryData.SourcesData.Collections.Num() == 0 &&
+				HistoryData.SourcesData.PackagePaths.Contains(PathName));
+	};
+
+	HistoryManager.RemoveHistoryData(RemoveHistoryDelegate);
+}
+
 FText SContentBrowser::GetSearchAssetsHintText() const
 {
 	if (PathViewPtr.IsValid())
diff --git a/Engine/Source/Editor/ContentBrowser/Private/SContentBrowser.h b/Engine/Source/Editor/ContentBrowser/Private/SContentBrowser.h
--- a/Engine/Source/Editor/ContentBrowser/Private/SContentBrowser.h
+++ b/Engine/Source/Editor/ContentBrowser/Private/SContentBrowser.h
@@ -2,6 +2,8 @@
 
 #pragma once
 
+struct FCollectionNameType;
+
 /**
  * A widget to display and work with all game and engine content
  */
@@ -270,6 +272,15 @@ private:
 	/** Delegate called when an editor setting is changed */
 	void HandleSettingChanged(FName PropertyName);
 
+	/** Handles an on collection destroyed event */
+	void HandleCollectionRemoved(const FCollectionNameType& Collection);
+
+	/** Handles an on collection renamed event */
+	void HandleCollectionRenamed(const FCollectionNameType& OriginalCollection, const FCollectionNameType& NewCollection);
+
+	/** Handles a path removed event */
+	void HandlePathRemoved(const FString& Path);
+
 	/** Gets all suggestions for the asset search box */
 	TArray<FString> GetAssetSearchSuggestions() const;
 
