More Clang compile warning fixes

diff --git a/Engine/Source/Developer/DesktopPlatform/Private/Mac/DesktopPlatformMac.cpp b/Engine/Source/Developer/DesktopPlatform/Private/Mac/DesktopPlatformMac.cpp
--- a/Engine/Source/Developer/DesktopPlatform/Private/Mac/DesktopPlatformMac.cpp
+++ b/Engine/Source/Developer/DesktopPlatform/Private/Mac/DesktopPlatformMac.cpp
@@ -559,7 +559,7 @@ void FDesktopPlatformMac::EnumerateEngineInstallations(TMap<FString, FString> &O
 	}
 	else
 	{
-		NSRunAlertPanel(@"Error", UProjectPath.GetNSString(), @"OK", NULL, NULL);
+		FPlatformMisc::MessageBoxExt(EAppMsgType::Ok, *UProjectPath, TEXT("Error"));
 	}
 
 	FConfigFile ConfigFile;
diff --git a/Engine/Source/Developer/MessageLog/Private/MessageFilter.h b/Engine/Source/Developer/MessageLog/Private/MessageFilter.h
--- a/Engine/Source/Developer/MessageLog/Private/MessageFilter.h
+++ b/Engine/Source/Developer/MessageLog/Private/MessageFilter.h
@@ -14,6 +14,8 @@ public:
 		, Display(true)
 	{}
 
+	virtual ~FMessageFilter() {}
+
 	/** Fires a callback when the mouse is released on an option (not linked to states) */
 	FReply OnMouseButtonUp(const FGeometry& MyGeometry, const FPointerEvent& MouseEvent);
 
diff --git a/Engine/Source/Developer/MessageLog/Private/Model/MessageLogListingModel.h b/Engine/Source/Developer/MessageLog/Private/Model/MessageLogListingModel.h
--- a/Engine/Source/Developer/MessageLog/Private/Model/MessageLogListingModel.h
+++ b/Engine/Source/Developer/MessageLog/Private/Model/MessageLogListingModel.h
@@ -27,6 +27,8 @@ protected:
 
 public:
 
+	virtual ~FMessageLogListingModel() {}
+
 	/**  
 	 *	Factory method which creates a new FMessageLogListingModel object
 	 *
diff --git a/Engine/Source/Developer/MessageLog/Private/Model/MessageLogModel.h b/Engine/Source/Developer/MessageLog/Private/Model/MessageLogModel.h
--- a/Engine/Source/Developer/MessageLog/Private/Model/MessageLogModel.h
+++ b/Engine/Source/Developer/MessageLog/Private/Model/MessageLogModel.h
@@ -10,7 +10,7 @@ class FMessageLogModel
 public:
 
 	/** Destructor */ 
-	~FMessageLogModel();
+	virtual ~FMessageLogModel();
 
 	/** Unregisters a log listing with the message log system, true if successful */
 	bool UnregisterLogListingModel( const FName& LogName );
diff --git a/Engine/Source/Developer/MessageLog/Private/Presentation/MessageLogViewModel.h b/Engine/Source/Developer/MessageLog/Private/Presentation/MessageLogViewModel.h
--- a/Engine/Source/Developer/MessageLog/Private/Presentation/MessageLogViewModel.h
+++ b/Engine/Source/Developer/MessageLog/Private/Presentation/MessageLogViewModel.h
@@ -22,7 +22,7 @@ public:
 	FMessageLogViewModel( const TSharedPtr< class FMessageLogModel >& InMessageLogModel );
 
 	/** Destructor */ 
-	~FMessageLogViewModel();
+	virtual ~FMessageLogViewModel();
 
 	/** Initializes the FMessageLogViewModel for use */
 	virtual void Initialize();
diff --git a/Engine/Source/Developer/MovieSceneTools/Private/SubMovieSceneTrackEditor.cpp b/Engine/Source/Developer/MovieSceneTools/Private/SubMovieSceneTrackEditor.cpp
--- a/Engine/Source/Developer/MovieSceneTools/Private/SubMovieSceneTrackEditor.cpp
+++ b/Engine/Source/Developer/MovieSceneTools/Private/SubMovieSceneTrackEditor.cpp
@@ -12,9 +12,9 @@ class FSubMovieSceneSection : public ISequencerSection
 {
 public:
 	FSubMovieSceneSection( TSharedPtr<ISequencer> InSequencer, UMovieSceneSection& InSectionObject, FName SectionName )
-		: Sequencer( InSequencer )
+		: DisplayName( NSLOCTEXT("SubMovieSceneSection", "DisplayName", "Scenes") )
 		, SectionObject( *CastChecked<USubMovieSceneSection>( &InSectionObject ) )
-		, DisplayName( NSLOCTEXT("SubMovieSceneSection", "DisplayName", "Scenes") )
+		, Sequencer( InSequencer )
 	{
 		MovieSceneInstance = InSequencer->GetInstanceForSubMovieSceneSection( InSectionObject );
 	}
diff --git a/Engine/Source/Editor/ContentBrowser/Private/AssetContextMenu.cpp b/Engine/Source/Editor/ContentBrowser/Private/AssetContextMenu.cpp
--- a/Engine/Source/Editor/ContentBrowser/Private/AssetContextMenu.cpp
+++ b/Engine/Source/Editor/ContentBrowser/Private/AssetContextMenu.cpp
@@ -26,14 +26,14 @@
 #define LOCTEXT_NAMESPACE "ContentBrowser"
 
 FAssetContextMenu::FAssetContextMenu(const TWeakPtr<SAssetView>& InAssetView)
-	: bAtLeastOneNonRedirectorSelected(false)
+	: AssetView(InAssetView)
+	, bAtLeastOneNonRedirectorSelected(false)
 	, bCanExecuteSCCCheckOut(false)
 	, bCanExecuteSCCOpenForAdd(false)
 	, bCanExecuteSCCCheckIn(false)
 	, bCanExecuteSCCHistory(false)
 	, bCanExecuteSCCRevert(false)
 	, bCanExecuteSCCSync(false)
-	, AssetView(InAssetView)
 {
 
 }
diff --git a/Engine/Source/Editor/ContentBrowser/Private/AssetViewSortManager.cpp b/Engine/Source/Editor/ContentBrowser/Private/AssetViewSortManager.cpp
--- a/Engine/Source/Editor/ContentBrowser/Private/AssetViewSortManager.cpp
+++ b/Engine/Source/Editor/ContentBrowser/Private/AssetViewSortManager.cpp
@@ -9,6 +9,8 @@ public:
 	/** Default constructor */
 	FCompareFAssetItemBase(bool bInAscending, const FName& InTag) : bAscending(bInAscending), Tag(InTag) {}
 
+	virtual ~FCompareFAssetItemBase() {}
+
 	/** Sort function */
 	FORCEINLINE bool operator()(const TSharedPtr<FAssetViewItem>& A, const TSharedPtr<FAssetViewItem>& B) const
 	{
@@ -102,6 +104,7 @@ struct FCompareFAssetItemByName : public FCompareFAssetItemBase
 {
 public:
 	FCompareFAssetItemByName(bool bInAscending, const FName& InTag) : FCompareFAssetItemBase(bInAscending, InTag) {}
+	virtual ~FCompareFAssetItemByName() {}
 
 protected:
 	FORCEINLINE virtual bool Compare(const TSharedPtr<FAssetViewItem>& A, const TSharedPtr<FAssetViewItem>& B) const override
@@ -125,6 +128,7 @@ struct FCompareFAssetItemByClass : public FCompareFAssetItemBase
 {
 public:
 	FCompareFAssetItemByClass(bool bInAscending, const FName& InTag) : FCompareFAssetItemBase(bInAscending, InTag) {}
+	virtual ~FCompareFAssetItemByClass() {}
 
 protected:
 	FORCEINLINE virtual bool Compare(const TSharedPtr<FAssetViewItem>& A, const TSharedPtr<FAssetViewItem>& B) const override
@@ -148,6 +152,7 @@ struct FCompareFAssetItemByPath : public FCompareFAssetItemBase
 {
 public:
 	FCompareFAssetItemByPath(bool bInAscending, const FName& InTag) : FCompareFAssetItemBase(bInAscending, InTag) {}
+	virtual ~FCompareFAssetItemByPath() {}
 
 protected:
 	FORCEINLINE virtual bool Compare(const TSharedPtr<FAssetViewItem>& A, const TSharedPtr<FAssetViewItem>& B) const override
@@ -171,6 +176,7 @@ struct FCompareFAssetItemByTag : public FCompareFAssetItemBase
 {
 public:
 	FCompareFAssetItemByTag(bool bInAscending, const FName& InTag) : FCompareFAssetItemBase(bInAscending, InTag) {}
+	virtual ~FCompareFAssetItemByTag() {}
 
 protected:
 	FORCEINLINE virtual bool Compare(const TSharedPtr<FAssetViewItem>& A, const TSharedPtr<FAssetViewItem>& B) const override
@@ -205,6 +211,7 @@ struct FCompareFAssetItemByTagNumerical : public FCompareFAssetItemBase
 {
 public:
 	FCompareFAssetItemByTagNumerical(bool bInAscending, const FName& InTag) : FCompareFAssetItemBase(bInAscending, InTag) {}
+	virtual ~FCompareFAssetItemByTagNumerical() {}
 
 protected:
 	FORCEINLINE virtual bool Compare(const TSharedPtr<FAssetViewItem>& A, const TSharedPtr<FAssetViewItem>& B) const override
@@ -240,6 +247,7 @@ struct FCompareFAssetItemByTagDimensional : public FCompareFAssetItemBase
 {
 public:
 	FCompareFAssetItemByTagDimensional(bool bInAscending, const FName& InTag) : FCompareFAssetItemBase(bInAscending, InTag) {}
+	virtual ~FCompareFAssetItemByTagDimensional() {}
 
 protected:
 	FORCEINLINE virtual bool Compare(const TSharedPtr<FAssetViewItem>& A, const TSharedPtr<FAssetViewItem>& B) const override
diff --git a/Engine/Source/Editor/ContentBrowser/Private/SourcesViewWidgets.h b/Engine/Source/Editor/ContentBrowser/Private/SourcesViewWidgets.h
--- a/Engine/Source/Editor/ContentBrowser/Private/SourcesViewWidgets.h
+++ b/Engine/Source/Editor/ContentBrowser/Private/SourcesViewWidgets.h
@@ -141,8 +141,8 @@ public:
 	DECLARE_DELEGATE_ThreeParams( FOnAssetsDragDropped, const TArray<FAssetData>& /*AssetList*/, const TSharedPtr<FCollectionItem>& /*CollectionItem*/, FText& /*OutMessage*/);
 
 	SLATE_BEGIN_ARGS( SCollectionListItem )
-		: _ParentWidget()
-		, _CollectionItem( TSharedPtr<FCollectionItem>() )
+		: _CollectionItem( TSharedPtr<FCollectionItem>() )
+		, _ParentWidget()
 	{}
 
 	/** Data for the collection this item represents */
diff --git a/Engine/Source/Editor/CurveAssetEditor/Private/CurveAssetEditor.h b/Engine/Source/Editor/CurveAssetEditor/Private/CurveAssetEditor.h
--- a/Engine/Source/Editor/CurveAssetEditor/Private/CurveAssetEditor.h
+++ b/Engine/Source/Editor/CurveAssetEditor/Private/CurveAssetEditor.h
@@ -9,9 +9,9 @@ class FCurveAssetEditor :  public ICurveAssetEditor
 public:
 
 	FCurveAssetEditor()
-		: ViewMinInput( 0.f )
+		: TrackWidget( NULL )
+		, ViewMinInput( 0.f )
 		, ViewMaxInput( 0.f )
-		, TrackWidget( NULL )
 	{}
 
 	virtual ~FCurveAssetEditor() {}
diff --git a/Engine/Source/Editor/DetailCustomizations/Private/CurveStructCustomization.cpp b/Engine/Source/Editor/DetailCustomizations/Private/CurveStructCustomization.cpp
--- a/Engine/Source/Editor/DetailCustomizations/Private/CurveStructCustomization.cpp
+++ b/Engine/Source/Editor/DetailCustomizations/Private/CurveStructCustomization.cpp
@@ -22,10 +22,10 @@ FCurveStructCustomization::~FCurveStructCustomization()
 }
 
 FCurveStructCustomization::FCurveStructCustomization()
-	: ViewMinInput(0.0f)
-	, ViewMaxInput(5.0f)
-	, RuntimeCurve(NULL)
+	: RuntimeCurve(NULL)
 	, Owner(NULL)
+	, ViewMinInput(0.0f)
+	, ViewMaxInput(5.0f)
 {
 }
 
diff --git a/Engine/Source/Editor/Documentation/Private/Documentation.h b/Engine/Source/Editor/Documentation/Private/Documentation.h
--- a/Engine/Source/Editor/Documentation/Private/Documentation.h
+++ b/Engine/Source/Editor/Documentation/Private/Documentation.h
@@ -11,7 +11,7 @@ public:
 
 public:
 
-	~FDocumentation();
+	virtual ~FDocumentation();
 
 	virtual bool OpenHome() const override;
 
diff --git a/Engine/Source/Editor/Documentation/Private/DocumentationPage.h b/Engine/Source/Editor/Documentation/Private/DocumentationPage.h
--- a/Engine/Source/Editor/Documentation/Private/DocumentationPage.h
+++ b/Engine/Source/Editor/Documentation/Private/DocumentationPage.h
@@ -13,7 +13,7 @@ public:
 
 public:
 
-	~FDocumentationPage();
+	virtual ~FDocumentationPage();
 	virtual bool HasExcerpt( const FString& ExcerptName ) override;
 	virtual int32 GetNumExcerpts() const override;
 	virtual bool GetExcerpt( const FString& ExcerptName, FExcerpt& Excerpt) override;
diff --git a/Engine/Source/Editor/GraphEditor/Private/MaterialNodes/SGraphNodeMaterialBase.cpp b/Engine/Source/Editor/GraphEditor/Private/MaterialNodes/SGraphNodeMaterialBase.cpp
--- a/Engine/Source/Editor/GraphEditor/Private/MaterialNodes/SGraphNodeMaterialBase.cpp
+++ b/Engine/Source/Editor/GraphEditor/Private/MaterialNodes/SGraphNodeMaterialBase.cpp
@@ -118,8 +118,8 @@ FIntPoint FPreviewViewport::GetSize() const
 // FPreviewElement
 
 FPreviewElement::FPreviewElement()
-	: ExpressionPreview(NULL)
-	, RenderTarget(new FSlateMaterialPreviewRenderTarget)
+	: RenderTarget(new FSlateMaterialPreviewRenderTarget)
+	, ExpressionPreview(NULL)
 	, bIsRealtime(false)
 {
 }
diff --git a/Engine/Source/Editor/Kismet/Private/BlueprintActionMenuBuilder.cpp b/Engine/Source/Editor/Kismet/Private/BlueprintActionMenuBuilder.cpp
--- a/Engine/Source/Editor/Kismet/Private/BlueprintActionMenuBuilder.cpp
+++ b/Engine/Source/Editor/Kismet/Private/BlueprintActionMenuBuilder.cpp
@@ -165,9 +165,9 @@ protected:
 
 //------------------------------------------------------------------------------
 FBlueprintActionMenuItemFactory::FBlueprintActionMenuItemFactory(FBlueprintActionContext const& ContextIn)
-	: Context(ContextIn)
-	, RootCategory(FText::GetEmpty())
+	: RootCategory(FText::GetEmpty())
 	, MenuGrouping(0)
+	, Context(ContextIn)
 {
 }
 
diff --git a/Engine/Source/Editor/Kismet/Private/BlueprintDetailsCustomization.h b/Engine/Source/Editor/Kismet/Private/BlueprintDetailsCustomization.h
--- a/Engine/Source/Editor/Kismet/Private/BlueprintDetailsCustomization.h
+++ b/Engine/Source/Editor/Kismet/Private/BlueprintDetailsCustomization.h
@@ -275,8 +275,8 @@ class FBlueprintGraphArgumentGroupLayout : public IDetailCustomNodeBuilder, publ
 {
 public:
 	FBlueprintGraphArgumentGroupLayout(TWeakPtr<class FBaseBlueprintGraphActionDetails> InGraphActionDetails, UK2Node_EditablePinBase* InTargetNode)
-		: TargetNode(InTargetNode)
-		, GraphActionDetailsPtr(InGraphActionDetails) {}
+		: GraphActionDetailsPtr(InGraphActionDetails)
+		, TargetNode(InTargetNode) {}
 
 private:
 	/** IDetailCustomNodeBuilder Interface*/
@@ -301,12 +301,12 @@ class FBlueprintGraphArgumentLayout : public IDetailCustomNodeBuilder, public TS
 {
 public:
 	FBlueprintGraphArgumentLayout(TWeakPtr<FUserPinInfo> PinInfo, UK2Node_EditablePinBase* InTargetNode, TWeakPtr<class FBaseBlueprintGraphActionDetails> InGraphActionDetails, FName InArgName, bool bInHasDefaultValue)
-		: ParamItemPtr(PinInfo)
+		: GraphActionDetailsPtr(InGraphActionDetails)
+		, ParamItemPtr(PinInfo)
 		, TargetNode(InTargetNode)
-		, GraphActionDetailsPtr(InGraphActionDetails)
-		, ArgumentName(InArgName)
-		, bHasDefaultValue(bInHasDefaultValue) {}
-	
+		, bHasDefaultValue(bInHasDefaultValue)
+		, ArgumentName(InArgName) {}
+
 private:
 	/** IDetailCustomNodeBuilder Interface*/
 	virtual void SetOnRebuildChildren( FSimpleDelegate InOnRegenerateChildren ) override {}
@@ -434,7 +434,7 @@ private:
 		uint32 SpecifierFlag;
 
 		FReplicationSpecifierLabel( FText InLocalizedName, uint32 InSpecifierFlag, FText InLocalizedToolTip ) :
-			LocalizedName( InLocalizedName ), SpecifierFlag( InSpecifierFlag ), LocalizedToolTip( InLocalizedToolTip )
+			LocalizedName( InLocalizedName ), LocalizedToolTip( InLocalizedToolTip ), SpecifierFlag( InSpecifierFlag )
 		{}
 	};
 
@@ -669,8 +669,8 @@ class FBlueprintGraphNodeDetails : public IDetailCustomization
 public:
 	/** Constructor */
 	FBlueprintGraphNodeDetails(TWeakPtr<FBlueprintEditor> InBlueprintEditorPtr)
-		: BlueprintEditorPtr(InBlueprintEditorPtr)
-		, GraphNodePtr(NULL)
+		: GraphNodePtr(NULL)
+		, BlueprintEditorPtr(InBlueprintEditorPtr)
 	{
 
 	}
diff --git a/Engine/Source/Editor/Kismet/Private/FindInBlueprints.cpp b/Engine/Source/Editor/Kismet/Private/FindInBlueprints.cpp
--- a/Engine/Source/Editor/Kismet/Private/FindInBlueprints.cpp
+++ b/Engine/Source/Editor/Kismet/Private/FindInBlueprints.cpp
@@ -193,7 +193,7 @@ FFindInBlueprintsResult::FFindInBlueprintsResult(const FText& InDisplayText )
 }
 
 FFindInBlueprintsResult::FFindInBlueprintsResult( const FText& InDisplayText, TSharedPtr<FFindInBlueprintsResult> InParent) 
-	: DisplayText(InDisplayText), Parent(InParent)
+	: Parent(InParent), DisplayText(InDisplayText)
 {
 }
 
diff --git a/Engine/Source/Editor/Kismet/Private/SBlueprintDiff.cpp b/Engine/Source/Editor/Kismet/Private/SBlueprintDiff.cpp
--- a/Engine/Source/Editor/Kismet/Private/SBlueprintDiff.cpp
+++ b/Engine/Source/Editor/Kismet/Private/SBlueprintDiff.cpp
@@ -103,7 +103,7 @@ struct FDiffResultItem: public TSharedFromThis<FDiffResultItem>
 };
 
 FListItemGraphToDiff::FListItemGraphToDiff( class SBlueprintDiff* InDiff, class UEdGraph* InGraphOld, class UEdGraph* InGraphNew, const FRevisionInfo& InRevisionOld, const FRevisionInfo& InRevisionNew )
-	: GraphOld(InGraphOld), GraphNew(InGraphNew), RevisionOld(InRevisionOld), RevisionNew(InRevisionNew), Diff(InDiff)
+	: Diff(InDiff), GraphOld(InGraphOld), GraphNew(InGraphNew), RevisionOld(InRevisionOld), RevisionNew(InRevisionNew)
 {
 	check(InGraphOld || InGraphNew); //one of them needs to exist
 
diff --git a/Engine/Source/Editor/Kismet/Private/SSCSEditor.cpp b/Engine/Source/Editor/Kismet/Private/SSCSEditor.cpp
--- a/Engine/Source/Editor/Kismet/Private/SSCSEditor.cpp
+++ b/Engine/Source/Editor/Kismet/Private/SSCSEditor.cpp
@@ -751,6 +751,8 @@ struct FSCSEditorComponentObjectTextFactory : public FCustomizableTextObjectFact
 		return FactoryPtr.ToSharedRef();
 	}
 
+	virtual ~FSCSEditorComponentObjectTextFactory() {}
+
 protected:	
 	// Constructor; protected to only allow this class to instance itself
 	FSCSEditorComponentObjectTextFactory()
diff --git a/Engine/Source/Editor/Kismet/Private/UserDefinedStructureEditor.cpp b/Engine/Source/Editor/Kismet/Private/UserDefinedStructureEditor.cpp
--- a/Engine/Source/Editor/Kismet/Private/UserDefinedStructureEditor.cpp
+++ b/Engine/Source/Editor/Kismet/Private/UserDefinedStructureEditor.cpp
@@ -37,7 +37,7 @@ public:
 		StructureDetailsView->GetOnFinishedChangingPropertiesDelegate().AddSP(this, &FStructureDefaultValueView::OnFinishedChangingProperties);
 	}
 
-	~FStructureDefaultValueView()
+	virtual ~FStructureDefaultValueView()
 	{
 		FStructureEditorUtils::FStructEditorManager::Get().RemoveListener(this);
 	}
@@ -430,9 +430,9 @@ class FUserDefinedStructureFieldLayout : public IDetailCustomNodeBuilder, public
 {
 public:
 	FUserDefinedStructureFieldLayout(TWeakPtr<class FUserDefinedStructureDetails> InStructureDetails, TWeakPtr<class FUserDefinedStructureLayout> InStructureLayout, FGuid InFieldGuid)
-		: FieldGuid(InFieldGuid)
+		: StructureDetails(InStructureDetails)
 		, StructureLayout(InStructureLayout)
-		, StructureDetails(InStructureDetails) {}
+		, FieldGuid(InFieldGuid) {}
 
 	void OnChanged()
 	{
diff --git a/Engine/Source/Editor/Kismet/Public/FindInBlueprints.h b/Engine/Source/Editor/Kismet/Public/FindInBlueprints.h
--- a/Engine/Source/Editor/Kismet/Public/FindInBlueprints.h
+++ b/Engine/Source/Editor/Kismet/Public/FindInBlueprints.h
@@ -14,6 +14,8 @@ public:
 	
 	/* Create a listing for a search result*/
 	FFindInBlueprintsResult(const FText& InDisplayText, TSharedPtr<FFindInBlueprintsResult> InParent);
+
+	virtual ~FFindInBlueprintsResult() {}
 	
 	/* Called when user clicks on the search item */
 	virtual FReply OnClick();
@@ -112,6 +114,7 @@ class FFindInBlueprintsGraphNode : public FFindInBlueprintsResult
 {
 public:
 	FFindInBlueprintsGraphNode(const FText& InValue, TSharedPtr<FFindInBlueprintsResult> InParent);
+	virtual ~FFindInBlueprintsGraphNode() {}
 
 	/** FFindInBlueprintsResult Interface */
 	virtual FReply OnClick() override;
@@ -144,6 +147,7 @@ class FFindInBlueprintsPin : public FFindInBlueprintsResult
 {
 public:
 	FFindInBlueprintsPin(const FText& InValue, TSharedPtr<FFindInBlueprintsResult> InParent, FString InSchemaName);
+	virtual ~FFindInBlueprintsPin() {}
 
 	/** FFindInBlueprintsResult Interface */
 	virtual TSharedRef<SWidget>	CreateIcon() const override;
@@ -170,6 +174,7 @@ class FFindInBlueprintsProperty : public FFindInBlueprintsResult
 {
 public:
 	FFindInBlueprintsProperty(const FText& InValue, TSharedPtr<FFindInBlueprintsResult> InParent);
+	virtual ~FFindInBlueprintsProperty() {}
 
 	/** FFindInBlueprintsResult Interface */
 	virtual TSharedRef<SWidget>	CreateIcon() const override;
@@ -193,6 +198,7 @@ class FFindInBlueprintsGraph : public FFindInBlueprintsResult
 {
 public:
 	FFindInBlueprintsGraph(const FText& InValue, TSharedPtr<FFindInBlueprintsResult> InParent, EGraphType InGraphType);
+	virtual ~FFindInBlueprintsGraph() {}
 
 	/** FFindInBlueprintsResult Interface */
 	virtual FReply OnClick() override;
diff --git a/Engine/Source/Editor/LandscapeEditor/Private/LandscapeEdMode.cpp b/Engine/Source/Editor/LandscapeEditor/Private/LandscapeEdMode.cpp
--- a/Engine/Source/Editor/LandscapeEditor/Private/LandscapeEdMode.cpp
+++ b/Engine/Source/Editor/LandscapeEditor/Private/LandscapeEdMode.cpp
@@ -148,16 +148,16 @@ void FLandscapeTool::SetEditRenderType()
 /** Constructor */
 FEdModeLandscape::FEdModeLandscape()
 	: FEdMode()
-	, CurrentGizmoActor(nullptr)
-	, LandscapeRenderAddCollision(nullptr)
-	, bToolActive(false)
-	, GizmoMaterial(nullptr)
-	, CopyPasteTool(nullptr)
-	, SplinesTool(nullptr)
 	, NewLandscapePreviewMode(ENewLandscapePreviewMode::None)
 	, DraggingEdge(ELandscapeEdge::None)
 	, DraggingEdge_Remainder(0)
+	, CurrentGizmoActor(nullptr)
+	, CopyPasteTool(nullptr)
+	, SplinesTool(nullptr)
+	, LandscapeRenderAddCollision(nullptr)
 	, CachedLandscapeMaterial(nullptr)
+	, bToolActive(false)
+	, GizmoMaterial(nullptr)
 {
 	GizmoMaterial = LoadObject<UMaterial>(NULL, TEXT("/Engine/EditorLandscapeResources/GizmoMaterial.GizmoMaterial"), NULL, LOAD_None, NULL);
 
diff --git a/Engine/Source/Editor/LandscapeEditor/Private/LandscapeEdModeBrushes.cpp b/Engine/Source/Editor/LandscapeEditor/Private/LandscapeEdModeBrushes.cpp
--- a/Engine/Source/Editor/LandscapeEditor/Private/LandscapeEdModeBrushes.cpp
+++ b/Engine/Source/Editor/LandscapeEditor/Private/LandscapeEdModeBrushes.cpp
@@ -499,8 +499,8 @@ public:
 	FEdModeLandscape* EdMode;
 
 	FLandscapeBrushGizmo(FEdModeLandscape* InEdMode)
-		: EdMode(InEdMode),
-		BrushMaterial(NULL)
+		: BrushMaterial(NULL)
+		, EdMode(InEdMode)
 	{
 		UMaterialInstanceConstant* GizmoMaterial = LoadObject<UMaterialInstanceConstant>(NULL, TEXT("/Engine/EditorLandscapeResources/MaskBrushMaterial_Gizmo.MaskBrushMaterial_Gizmo"), NULL, LOAD_None, NULL);
 		BrushMaterial = UMaterialInstanceDynamic::Create(GizmoMaterial, NULL);
diff --git a/Engine/Source/Editor/LandscapeEditor/Private/LandscapeEdModeComponentTools.cpp b/Engine/Source/Editor/LandscapeEditor/Private/LandscapeEdModeComponentTools.cpp
--- a/Engine/Source/Editor/LandscapeEditor/Private/LandscapeEdModeComponentTools.cpp
+++ b/Engine/Source/Editor/LandscapeEditor/Private/LandscapeEdModeComponentTools.cpp
@@ -25,12 +25,14 @@ class FLandscapeToolStrokeSelect : FLandscapeStrokeBase<const FLandscapeToolTarg
 public:
 	FLandscapeToolStrokeSelect(FEdModeLandscape* InEdMode, const FLandscapeToolTarget& InTarget)
 		: FLandscapeStrokeBase<const FLandscapeToolTarget>(InTarget)
+		, bInitializedComponentInvert(false)
 		, LandscapeInfo(InTarget.LandscapeInfo.Get())
 		, Cache(InTarget)
-		, bInitializedComponentInvert(false)
 	{
 	}
 
+	virtual ~FLandscapeToolStrokeSelect() {}
+
 	virtual void Apply(FEditorViewportClient* ViewportClient, FLandscapeBrush* Brush, const ULandscapeEditorObject* UISettings, const TArray<FLandscapeToolMousePosition>& MousePositions)
 	{
 		if (LandscapeInfo)
@@ -198,6 +200,8 @@ public:
 		, Cache(InTarget)
 	{}
 
+	virtual ~FLandscapeToolStrokeVisibility() {}
+
 	virtual void Apply(FEditorViewportClient* ViewportClient, FLandscapeBrush* Brush, const ULandscapeEditorObject* UISettings, const TArray<FLandscapeToolMousePosition>& MousePositions)
 	{
 		if (LandscapeInfo)
@@ -1124,6 +1128,8 @@ public:
 		, WeightCache(InTarget)
 	{}
 
+	virtual ~FLandscapeToolStrokeCopy() {}
+
 	struct FGizmoPreData
 	{
 		float Ratio;
@@ -1487,6 +1493,8 @@ public:
 		, WeightCache(InTarget)
 	{}
 
+	virtual ~FLandscapeToolStrokePaste() {}
+
 	virtual void Apply(FEditorViewportClient* ViewportClient, FLandscapeBrush* Brush, const ULandscapeEditorObject* UISettings, const TArray<FLandscapeToolMousePosition>& MousePositions)
 	{
 		//ULandscapeInfo* LandscapeInfo = EdMode->CurrentToolTarget.LandscapeInfo;
diff --git a/Engine/Source/Editor/LandscapeEditor/Private/LandscapeEdModePaintTools.cpp b/Engine/Source/Editor/LandscapeEditor/Private/LandscapeEdModePaintTools.cpp
--- a/Engine/Source/Editor/LandscapeEditor/Private/LandscapeEdModePaintTools.cpp
+++ b/Engine/Source/Editor/LandscapeEditor/Private/LandscapeEdModePaintTools.cpp
@@ -57,8 +57,8 @@ class FLandscapeToolStrokePaintBase
 {
 public:
 	FLandscapeToolStrokePaintBase(FEdModeLandscape* InEdMode, const FLandscapeToolTarget& InTarget)
-		: LandscapeInfo(InTarget.LandscapeInfo.Get())
-		, Cache(InTarget)
+		: Cache(InTarget)
+		, LandscapeInfo(InTarget.LandscapeInfo.Get())
 	{
 	}
 
@@ -330,6 +330,8 @@ public:
 		: FLandscapeToolStrokePaintBase<ToolTarget>(InEdMode, InTarget)
 	{}
 
+	virtual ~FLandscapeToolStrokeSmooth() {}
+
 	virtual void Apply(FEditorViewportClient* ViewportClient, FLandscapeBrush* Brush, const ULandscapeEditorObject* UISettings, const TArray<FLandscapeToolMousePosition>& MousePositions) override
 	{
 		if (!this->LandscapeInfo) return;
diff --git a/Engine/Source/Editor/Levels/Private/LevelViewModel.cpp b/Engine/Source/Editor/Levels/Private/LevelViewModel.cpp
--- a/Engine/Source/Editor/Levels/Private/LevelViewModel.cpp
+++ b/Engine/Source/Editor/Levels/Private/LevelViewModel.cpp
@@ -14,9 +14,9 @@
 FLevelViewModel::FLevelViewModel( const TWeakObjectPtr< class ULevel >& InLevel, 
 								  const TWeakObjectPtr< class ULevelStreaming >& InLevelStreaming, 
 								  const TWeakObjectPtr< UEditorEngine >& InEditor )
-	: Level( InLevel )
+	: Editor( InEditor )
+	, Level( InLevel )
 	, LevelStreaming( InLevelStreaming )
-	, Editor( InEditor )
 	, LevelActorsCount(0)
 {
 }
diff --git a/Engine/Source/Editor/Persona/Private/AnimationEditorViewportClient.cpp b/Engine/Source/Editor/Persona/Private/AnimationEditorViewportClient.cpp
--- a/Engine/Source/Editor/Persona/Private/AnimationEditorViewportClient.cpp
+++ b/Engine/Source/Editor/Persona/Private/AnimationEditorViewportClient.cpp
@@ -72,9 +72,9 @@ FAnimationViewportClient::FAnimationViewportClient( FPreviewScene& InPreviewScen
 	, bManipulating(false)
 	, bInTransaction(false)
 	, AnimationPlaybackScale(1.0f)
-	, SelectedWindActor(NULL)
-	, PrevWindStrength(0.0f)
 	, GravityScaleSliderValue(0.25f)
+	, PrevWindStrength(0.0f)
+	, SelectedWindActor(NULL)
 	, BodyTraceDistance(100000.0f)
 {
 	// load config
diff --git a/Engine/Source/Editor/Persona/Private/SAnimMontageScrubPanel.h b/Engine/Source/Editor/Persona/Private/SAnimMontageScrubPanel.h
--- a/Engine/Source/Editor/Persona/Private/SAnimMontageScrubPanel.h
+++ b/Engine/Source/Editor/Persona/Private/SAnimMontageScrubPanel.h
@@ -9,9 +9,9 @@ class SAnimMontageScrubPanel : public SAnimationScrubPanel
 public:
 
 	SLATE_BEGIN_ARGS(SAnimMontageScrubPanel)
-		: _Persona()
+		: _MontageEditor()
+		, _Persona()
 		, _LockedSequence()
-		, _MontageEditor()
 	{}
 		SLATE_ARGUMENT( TWeakPtr<class SMontageEditor>, MontageEditor)
 		SLATE_ARGUMENT(TWeakPtr<FPersona>, Persona )
diff --git a/Engine/Source/Editor/Persona/Private/SAnimNotifyPanel.cpp b/Engine/Source/Editor/Persona/Private/SAnimNotifyPanel.cpp
--- a/Engine/Source/Editor/Persona/Private/SAnimNotifyPanel.cpp
+++ b/Engine/Source/Editor/Persona/Private/SAnimNotifyPanel.cpp
@@ -190,10 +190,10 @@ public:
 		, _Sequence(NULL)
 		, _ViewInputMin()
 		, _ViewInputMax()
+		, _TrackIndex()
+		, _TrackColor(FLinearColor::White)
 		, _OnSelectionChanged()
 		, _OnUpdatePanel()
-		, _TrackColor(FLinearColor::White)
-		, _TrackIndex()
 		, _OnRequestTrackPan()
 		, _OnRequestOffsetRefresh()
 		, _OnDeleteNotify()
@@ -438,12 +438,12 @@ private:
 public:
 	SLATE_BEGIN_ARGS( SNotifyEdTrack )
 		: _TrackIndex(INDEX_NONE)
+		, _AnimNotifyPanel()
 		, _Sequence()
+		, _WidgetWidth()
 		, _ViewInputMin()
 		, _ViewInputMax()
 		, _OnSelectionChanged()
-		, _AnimNotifyPanel()
-		, _WidgetWidth()
 		, _OnUpdatePanel()
 		, _OnDeleteNotify()
 		, _OnDeselectAllNotifies()
diff --git a/Engine/Source/Editor/Persona/Private/SAnimationBlendSpace.h b/Engine/Source/Editor/Persona/Private/SAnimationBlendSpace.h
--- a/Engine/Source/Editor/Persona/Private/SAnimationBlendSpace.h
+++ b/Engine/Source/Editor/Persona/Private/SAnimationBlendSpace.h
@@ -399,8 +399,8 @@ public:
 	 * default value 
 	 */
 	FBlendSpaceGrid()
-		:	GridNum(5, 5), 
-			GridDim(FVector(0, 0, 0), FVector(100, 100, 0))
+		:	GridDim(FVector(0, 0, 0), FVector(100, 100, 0))
+		,	GridNum(5, 5)
 	{
 	}
 
diff --git a/Engine/Source/Editor/Persona/Private/SSkeletonTree.h b/Engine/Source/Editor/Persona/Private/SSkeletonTree.h
--- a/Engine/Source/Editor/Persona/Private/SSkeletonTree.h
+++ b/Engine/Source/Editor/Persona/Private/SSkeletonTree.h
@@ -153,6 +153,8 @@ public:
 	/** Set Translation Retargeting Mode for this bone. */
 	void SetBoneTranslationRetargetingMode(EBoneTranslationRetargetingMode::Type NewRetargetingMode);
 
+	virtual ~FDisplayedMeshBoneInfo() {}
+
 protected:
 	/** Hidden constructor, always use Make above */
 	FDisplayedMeshBoneInfo(const FName & InSource)
@@ -245,6 +247,8 @@ public:
 	/** Return socket name as FText for display in skeleton tree */
 	FText GetSocketNameAsText() const { return FText::FromName(SocketData->SocketName); }
 
+	virtual ~FDisplayedSocketInfo() {}
+
 protected:
 	/** Hidden constructor, always use Make above */
 	FDisplayedSocketInfo( USkeletalMeshSocket* InSource, ESocketParentType::Type InParentType )
@@ -340,6 +344,8 @@ public:
 	/** Handler for when the user double clicks on this item in the tree */
 	virtual void OnItemDoubleClicked();
 
+	virtual ~FDisplayedAttachedAssetInfo() {}
+
 protected:
 	/** Hidden constructor, always use Make above */
 	FDisplayedAttachedAssetInfo( const FName& InAttachedTo, UObject* InAsset ) :
diff --git a/Engine/Source/Editor/Persona/Private/STrack.h b/Engine/Source/Editor/Persona/Private/STrack.h
--- a/Engine/Source/Editor/Persona/Private/STrack.h
+++ b/Engine/Source/Editor/Persona/Private/STrack.h
@@ -69,19 +69,19 @@ public:
 	SLATE_BEGIN_ARGS( STrackNode )
 		: _ViewInputMin()
 		, _ViewInputMax()
-		, _OnSelectionChanged()
 		, _DataLength()
 		, _DataStartPos()
 		, _NodeName()
 		, _NodeColor()
 		, _SelectedNodeColor()
+		, _NodeSelectionSet(NULL)
+		, _AllowDrag(true)
 		, _OnTrackNodeDragged()
 		, _OnTrackNodeDropped()
-		, _OnTrackNodeClicked()
+		, _OnSelectionChanged()
 		, _OnNodeRightClickContextMenu()
+		, _OnTrackNodeClicked()
 		, _CenterOnPosition(false)
-		, _AllowDrag(true)
-		, _NodeSelectionSet(NULL)
 	{
 	}
 
@@ -197,19 +197,19 @@ public:
 	SLATE_BEGIN_ARGS( STrack )
 		: _ViewInputMin()
 		, _ViewInputMax()
-		, _OnSelectionChanged()
+		, _TrackMaxValue()
+		, _TrackMinValue()
 		, _TrackNumDiscreteValues()
+		, _ScrubPosition()
 		, _TrackColor(FLinearColor::White)
-		, _OnBarDrag()
+		, _OnSelectionChanged()
 		, _DraggableBars()
-		, _DraggableBarLabels()
 		, _DraggableBarSnapPositions()
-		, _TrackMinValue()
-		, _TrackMaxValue()
+		, _DraggableBarLabels()
+		, _OnBarDrag()
 		, _OnTrackDragDrop()
 		, _OnSummonContextMenu()
 		, _OnTrackRightClickContextMenu()
-		, _ScrubPosition()
 		, _StyleInfo(FEditorStyle::GetBrush( TEXT( "Persona.NotifyEditor.NotifyTrackBackground" )))
 	{}
 
diff --git a/Engine/Source/Editor/PropertyEditor/Private/DetailCategoryBuilderImpl.cpp b/Engine/Source/Editor/PropertyEditor/Private/DetailCategoryBuilderImpl.cpp
--- a/Engine/Source/Editor/PropertyEditor/Private/DetailCategoryBuilderImpl.cpp
+++ b/Engine/Source/Editor/PropertyEditor/Private/DetailCategoryBuilderImpl.cpp
@@ -75,9 +75,9 @@ FDetailWidgetRow FDetailLayoutCustomization::GetWidgetRow() const
 
 
 FDetailCategoryImpl::FDetailCategoryImpl( FName InCategoryName, TSharedRef<FDetailLayoutBuilderImpl> InDetailLayout )
-	: DetailLayoutBuilder( InDetailLayout )
+	: HeaderContentWidget( NULL )
+	, DetailLayoutBuilder( InDetailLayout )
 	, CategoryName( InCategoryName )
-	, HeaderContentWidget( NULL )
 	, SortOrder( 0 )
 	, bRestoreExpansionState( true )
 	, bShouldBeInitiallyCollapsed( false )
diff --git a/Engine/Source/Editor/PropertyEditor/Private/DetailGroup.cpp b/Engine/Source/Editor/PropertyEditor/Private/DetailGroup.cpp
--- a/Engine/Source/Editor/PropertyEditor/Private/DetailGroup.cpp
+++ b/Engine/Source/Editor/PropertyEditor/Private/DetailGroup.cpp
@@ -7,9 +7,9 @@
 #include "DetailItemNode.h"
 
 FDetailGroup::FDetailGroup( const FName InGroupName, TSharedRef<FDetailCategoryImpl> InParentCategory, const FString& InLocalizedDisplayName )
-	: LocalizedDisplayName( InLocalizedDisplayName ) 
+	: ParentCategory( InParentCategory )
+	, LocalizedDisplayName( InLocalizedDisplayName )
 	, GroupName( InGroupName )
-	, ParentCategory( InParentCategory )
 {
 
 }
diff --git a/Engine/Source/Editor/PropertyEditor/Private/DetailItemNode.cpp b/Engine/Source/Editor/PropertyEditor/Private/DetailItemNode.cpp
--- a/Engine/Source/Editor/PropertyEditor/Private/DetailItemNode.cpp
+++ b/Engine/Source/Editor/PropertyEditor/Private/DetailItemNode.cpp
@@ -15,10 +15,10 @@
 FDetailItemNode::FDetailItemNode(const FDetailLayoutCustomization& InCustomization, TSharedRef<FDetailCategoryImpl> InParentCategory, TAttribute<bool> InIsParentEnabled )
 	: Customization( InCustomization )
 	, ParentCategory( InParentCategory )
+	, IsParentEnabled( InIsParentEnabled )
+	, CachedItemVisibility( EVisibility::Visible )
 	, bShouldBeVisibleDueToFiltering( false )
 	, bShouldBeVisibleDueToChildFiltering( false )
-	, CachedItemVisibility( EVisibility::Visible )
-	, IsParentEnabled( InIsParentEnabled )
 	, bTickable( false )
 	, bIsExpanded( InCustomization.HasCustomBuilder() ? !InCustomization.CustomBuilderRow->IsInitiallyCollapsed() : false )
 {
diff --git a/Engine/Source/Editor/PropertyEditor/Private/DetailPropertyRow.cpp b/Engine/Source/Editor/PropertyEditor/Private/DetailPropertyRow.cpp
--- a/Engine/Source/Editor/PropertyEditor/Private/DetailPropertyRow.cpp
+++ b/Engine/Source/Editor/PropertyEditor/Private/DetailPropertyRow.cpp
@@ -10,8 +10,8 @@
 #include "CustomChildBuilder.h"
 
 FDetailPropertyRow::FDetailPropertyRow(TSharedPtr<FPropertyNode> InPropertyNode, TSharedRef<FDetailCategoryImpl> InParentCategory, TSharedPtr<FObjectPropertyNode> InExternalRootNode )
-	: PropertyNode( InPropertyNode )
-	, CustomIsEnabledAttrib( true )
+	: CustomIsEnabledAttrib( true )
+	, PropertyNode( InPropertyNode )
 	, ParentCategory( InParentCategory )
 	, ExternalRootNode( InExternalRootNode )
 	, bShowPropertyButtons( true )
diff --git a/Engine/Source/Editor/PropertyEditor/Private/Presentation/PropertyTable/DataSource.h b/Engine/Source/Editor/PropertyEditor/Private/Presentation/PropertyTable/DataSource.h
--- a/Engine/Source/Editor/PropertyEditor/Private/Presentation/PropertyTable/DataSource.h
+++ b/Engine/Source/Editor/PropertyEditor/Private/Presentation/PropertyTable/DataSource.h
@@ -12,6 +12,8 @@ public:
 	{
 	}
 
+	virtual ~UObjectDataSource() {}
+
 	virtual TWeakObjectPtr< UObject > AsUObject() const override { return Object; }
 	virtual TSharedPtr< FPropertyPath > AsPropertyPath() const override { return NULL; }
 
@@ -31,6 +33,8 @@ public:
 	{
 	}
 
+	virtual ~PropertyPathDataSource() {}
+
 	virtual TWeakObjectPtr< UObject > AsUObject() const override { return NULL; }
 	virtual TSharedPtr< FPropertyPath > AsPropertyPath() const override { return Path; }
 
@@ -45,6 +49,8 @@ class NoDataSource : public IDataSource
 {
 public:
 
+	virtual ~NoDataSource() {}
+
 	virtual TWeakObjectPtr< UObject > AsUObject() const override { return NULL; }
 	virtual TSharedPtr< FPropertyPath > AsPropertyPath() const override { return NULL; }
 
diff --git a/Engine/Source/Editor/PropertyEditor/Private/Presentation/PropertyTable/PropertyTable.cpp b/Engine/Source/Editor/PropertyEditor/Private/Presentation/PropertyTable/PropertyTable.cpp
--- a/Engine/Source/Editor/PropertyEditor/Private/Presentation/PropertyTable/PropertyTable.cpp
+++ b/Engine/Source/Editor/PropertyEditor/Private/Presentation/PropertyTable/PropertyTable.cpp
@@ -25,12 +25,12 @@ FPropertyTable::FPropertyTable()
 	, CurrentRow( NULL )
 	, CurrentCell( NULL )
 	, CurrentColumn( NULL )
-	, ItemHeight( 20 )
+	, RootPath( FPropertyPath::CreateEmpty() )
 	, SelectionUnit( EPropertyTableSelectionUnit::Cell )
 	, SelectionMode( ESelectionMode::Multi )
 	, ShowRowHeader( true )
 	, ShowObjectName( true )
-	, RootPath( FPropertyPath::CreateEmpty() )
+	, ItemHeight( 20 )
 	, AllowUserToChangeRoot( true )
 	, bRefreshRequested( false )
 	, Orientation( EPropertyTableOrientation::AlignPropertiesInColumns )
diff --git a/Engine/Source/Editor/PropertyEditor/Private/Presentation/PropertyTable/PropertyTableCell.h b/Engine/Source/Editor/PropertyEditor/Private/Presentation/PropertyTable/PropertyTableCell.h
--- a/Engine/Source/Editor/PropertyEditor/Private/Presentation/PropertyTable/PropertyTableCell.h
+++ b/Engine/Source/Editor/PropertyEditor/Private/Presentation/PropertyTable/PropertyTableCell.h
@@ -11,6 +11,8 @@ public:
 
 	FPropertyTableCell( const TSharedRef< class IPropertyTableColumn >& InColumn, const TSharedRef< class IPropertyTableRow >& InRow );
 
+	virtual ~FPropertyTableCell() {}
+
 	virtual void Refresh() override;
 
 	virtual bool IsReadOnly() const override;
diff --git a/Engine/Source/Editor/PropertyEditor/Private/Presentation/PropertyTable/PropertyTableColumn.cpp b/Engine/Source/Editor/PropertyEditor/Private/Presentation/PropertyTable/PropertyTableColumn.cpp
--- a/Engine/Source/Editor/PropertyEditor/Private/Presentation/PropertyTable/PropertyTableColumn.cpp
+++ b/Engine/Source/Editor/PropertyEditor/Private/Presentation/PropertyTable/PropertyTableColumn.cpp
@@ -165,8 +165,8 @@ struct FCompareRowByColumnAscending
 {
 public:
 	FCompareRowByColumnAscending( const TSharedRef< IPropertyTableColumn >& InColumn, const TWeakObjectPtr< UPropertyType >& InUProperty )
-		: Column( InColumn )
-		, Property( InUProperty )
+		: Property( InUProperty )
+		, Column( InColumn )
 	{
 
 	}
@@ -323,14 +323,14 @@ FORCEINLINE bool FCompareRowByColumnAscending<UObjectProperty>::ComparePropertyV
 
 
 FPropertyTableColumn::FPropertyTableColumn( const TSharedRef< IPropertyTable >& InTable, const TWeakObjectPtr< UObject >& InObject )
-	: DataSource( MakeShareable( new UObjectDataSource( InObject.Get() ) ) )
+	: Cells()
+	, DataSource( MakeShareable( new UObjectDataSource( InObject.Get() ) ) )
 	, Table( InTable )
-	, Cells()
-	, bIsHidden( false )
-	, bIsFrozen( false )
-	, Width( 1.0f )
 	, Id( FName( EName::NAME_None ) )
 	, DisplayName()
+	, Width( 1.0f )
+	, bIsHidden( false )
+	, bIsFrozen( false )
 	, PartialPath( FPropertyPath::CreateEmpty() )
 	, SizeMode(EPropertyTableColumnSizeMode::Fill)
 {
@@ -339,14 +339,14 @@ FPropertyTableColumn::FPropertyTableColumn( const TSharedRef< IPropertyTable >&
 }
 
 FPropertyTableColumn::FPropertyTableColumn( const TSharedRef< IPropertyTable >& InTable, const TSharedRef< FPropertyPath >& InPropertyPath )
-	: DataSource( MakeShareable( new PropertyPathDataSource( InPropertyPath ) ) )
+	: Cells()
+	, DataSource( MakeShareable( new PropertyPathDataSource( InPropertyPath ) ) )
 	, Table( InTable )
-	, Cells()
-	, bIsHidden( false )
-	, bIsFrozen( false )
-	, Width( 1.0f )
 	, Id( FName( EName::NAME_None ) )
 	, DisplayName()
+	, Width( 1.0f )
+	, bIsHidden( false )
+	, bIsFrozen( false )
 	, PartialPath( FPropertyPath::CreateEmpty() )
 	, SizeMode(EPropertyTableColumnSizeMode::Fill)
 {
@@ -355,14 +355,14 @@ FPropertyTableColumn::FPropertyTableColumn( const TSharedRef< IPropertyTable >&
 }
 
 FPropertyTableColumn::FPropertyTableColumn( const TSharedRef< class IPropertyTable >& InTable, const TWeakObjectPtr< UObject >& InObject, const TSharedRef< FPropertyPath >& InPartialPropertyPath )
-	: DataSource( MakeShareable( new UObjectDataSource( InObject.Get() ) ) )
+	: Cells()
+	, DataSource( MakeShareable( new UObjectDataSource( InObject.Get() ) ) )
 	, Table( InTable )
-	, Cells()
-	, bIsHidden( false )
-	, bIsFrozen( false )
-	, Width( 1.0f )
 	, Id( FName( EName::NAME_None ) )
 	, DisplayName()
+	, Width( 1.0f )
+	, bIsHidden( false )
+	, bIsFrozen( false )
 	, PartialPath( InPartialPropertyPath )
 	, SizeMode(EPropertyTableColumnSizeMode::Fill)
 {
diff --git a/Engine/Source/Editor/PropertyEditor/Private/Presentation/PropertyTable/PropertyTableColumn.h b/Engine/Source/Editor/PropertyEditor/Private/Presentation/PropertyTable/PropertyTableColumn.h
--- a/Engine/Source/Editor/PropertyEditor/Private/Presentation/PropertyTable/PropertyTableColumn.h
+++ b/Engine/Source/Editor/PropertyEditor/Private/Presentation/PropertyTable/PropertyTableColumn.h
@@ -13,6 +13,8 @@ public:
 
 	FPropertyTableColumn( const TSharedRef< IPropertyTable >& InTable, const TWeakObjectPtr< UObject >& InObject, const TSharedRef< FPropertyPath >& InPartialPropertyPath );
 
+	virtual ~FPropertyTableColumn() {}
+
 	// Begin IPropertyTable Interface
 
 	virtual FName GetId() const override;
diff --git a/Engine/Source/Editor/PropertyEditor/Private/Presentation/PropertyTable/PropertyTableObjectNameCell.h b/Engine/Source/Editor/PropertyEditor/Private/Presentation/PropertyTable/PropertyTableObjectNameCell.h
--- a/Engine/Source/Editor/PropertyEditor/Private/Presentation/PropertyTable/PropertyTableObjectNameCell.h
+++ b/Engine/Source/Editor/PropertyEditor/Private/Presentation/PropertyTable/PropertyTableObjectNameCell.h
@@ -10,6 +10,7 @@ class FPropertyTableObjectNameCell : public TSharedFromThis< FPropertyTableObjec
 public:
 
 	FPropertyTableObjectNameCell( const TSharedRef< class FPropertyTableObjectNameColumn >& InColumn, const TSharedRef< class IPropertyTableRow >& InRow );
+	virtual ~FPropertyTableObjectNameCell() {}
 
 	virtual void Refresh() override;
 
diff --git a/Engine/Source/Editor/PropertyEditor/Private/Presentation/PropertyTable/PropertyTableObjectNameColumn.h b/Engine/Source/Editor/PropertyEditor/Private/Presentation/PropertyTable/PropertyTableObjectNameColumn.h
--- a/Engine/Source/Editor/PropertyEditor/Private/Presentation/PropertyTable/PropertyTableObjectNameColumn.h
+++ b/Engine/Source/Editor/PropertyEditor/Private/Presentation/PropertyTable/PropertyTableObjectNameColumn.h
@@ -19,12 +19,14 @@ public:
 	FPropertyTableObjectNameColumn( const TSharedRef< IPropertyTable >& InTable )
 		: Table( InTable )
 		, Cells()
-		, bIsHidden( false )
 		, Width( 2.0f )
+		, bIsHidden( false )
 	{
 
 	}
 
+	virtual ~FPropertyTableObjectNameColumn() {}
+
 	//  Begin IPropertyTable Interface
 
 	virtual FName GetId() const override { return FName( TEXT("ObjectName") ); }
diff --git a/Engine/Source/Editor/PropertyEditor/Private/Presentation/PropertyTable/PropertyTablePropertyNameCell.cpp b/Engine/Source/Editor/PropertyEditor/Private/Presentation/PropertyTable/PropertyTablePropertyNameCell.cpp
--- a/Engine/Source/Editor/PropertyEditor/Private/Presentation/PropertyTable/PropertyTablePropertyNameCell.cpp
+++ b/Engine/Source/Editor/PropertyEditor/Private/Presentation/PropertyTable/PropertyTablePropertyNameCell.cpp
@@ -14,12 +14,12 @@
 
 
 FPropertyTablePropertyNameCell::FPropertyTablePropertyNameCell( const TSharedRef< class FPropertyTablePropertyNameColumn >& InColumn, const TSharedRef< class IPropertyTableRow >& InRow )
-	: EnteredEditModeEvent()
-	, ExitedEditModeEvent()
+	: bInEditMode( false )
+	, bIsBound( true )
 	, Column( InColumn )
+	, EnteredEditModeEvent()
+	, ExitedEditModeEvent()
 	, Row( InRow )
-	, bIsBound( true )
-	, bInEditMode( false )
 {
 	Refresh();
 }
diff --git a/Engine/Source/Editor/PropertyEditor/Private/Presentation/PropertyTable/PropertyTablePropertyNameColumn.cpp b/Engine/Source/Editor/PropertyEditor/Private/Presentation/PropertyTable/PropertyTablePropertyNameColumn.cpp
--- a/Engine/Source/Editor/PropertyEditor/Private/Presentation/PropertyTable/PropertyTablePropertyNameColumn.cpp
+++ b/Engine/Source/Editor/PropertyEditor/Private/Presentation/PropertyTable/PropertyTablePropertyNameColumn.cpp
@@ -9,11 +9,11 @@
 
 
 FPropertyTablePropertyNameColumn::FPropertyTablePropertyNameColumn( const TSharedRef< IPropertyTable >& InTable )
-	: Table( InTable )
+	: bIsHidden( false )
 	, Cells()
-	, bIsHidden( false )
-	, Width( 2.0f )
 	, DataSource( MakeShareable( new NoDataSource() ) )
+	, Table( InTable )
+	, Width( 2.0f )
 {
 
 }
diff --git a/Engine/Source/Editor/PropertyEditor/Private/Presentation/PropertyTable/PropertyTablePropertyNameColumn.h b/Engine/Source/Editor/PropertyEditor/Private/Presentation/PropertyTable/PropertyTablePropertyNameColumn.h
--- a/Engine/Source/Editor/PropertyEditor/Private/Presentation/PropertyTable/PropertyTablePropertyNameColumn.h
+++ b/Engine/Source/Editor/PropertyEditor/Private/Presentation/PropertyTable/PropertyTablePropertyNameColumn.h
@@ -18,6 +18,8 @@ public:
 
 	FPropertyTablePropertyNameColumn( const TSharedRef< IPropertyTable >& InTable );
 
+	virtual ~FPropertyTablePropertyNameColumn() {}
+
 	// Begin IPropertyTableColumn Interface
 
 	virtual bool CanSelectCells() const override { return true; }
diff --git a/Engine/Source/Editor/PropertyEditor/Private/Presentation/PropertyTable/PropertyTableRow.h b/Engine/Source/Editor/PropertyEditor/Private/Presentation/PropertyTable/PropertyTableRow.h
--- a/Engine/Source/Editor/PropertyEditor/Private/Presentation/PropertyTable/PropertyTableRow.h
+++ b/Engine/Source/Editor/PropertyEditor/Private/Presentation/PropertyTable/PropertyTableRow.h
@@ -15,6 +15,8 @@ public:
 
 	FPropertyTableRow( const TSharedRef< class IPropertyTable >& InTable, const TSharedRef< FPropertyPath >& InPropertyPath, const TSharedRef< FPropertyPath >& InPartialPropertyPath );
 
+	virtual ~FPropertyTableRow() {}
+
 	virtual TSharedRef< class IDataSource > GetDataSource() const override { return DataSource; }
 
 	virtual bool HasChildren() const override;
diff --git a/Engine/Source/Editor/PropertyEditor/Private/Presentation/PropertyTable/PropertyTableRowHeaderColumn.h b/Engine/Source/Editor/PropertyEditor/Private/Presentation/PropertyTable/PropertyTableRowHeaderColumn.h
--- a/Engine/Source/Editor/PropertyEditor/Private/Presentation/PropertyTable/PropertyTableRowHeaderColumn.h
+++ b/Engine/Source/Editor/PropertyEditor/Private/Presentation/PropertyTable/PropertyTableRowHeaderColumn.h
@@ -23,6 +23,8 @@ public:
 
 	}
 
+	virtual ~FPropertyTableRowHeaderColumn() {}
+
 	virtual FName GetId() const override { return FName( TEXT("RowHeader") ); }
 	virtual FText GetDisplayName() const override { return FText::GetEmpty(); }
 	virtual TSharedRef< IDataSource > GetDataSource() const override { return MakeShareable( new PropertyPathDataSource( FPropertyPath::CreateEmpty() ) ); }
diff --git a/Engine/Source/Editor/PropertyEditor/Private/PropertyChangeListener.h b/Engine/Source/Editor/PropertyEditor/Private/PropertyChangeListener.h
--- a/Engine/Source/Editor/PropertyEditor/Private/PropertyChangeListener.h
+++ b/Engine/Source/Editor/PropertyEditor/Private/PropertyChangeListener.h
@@ -8,6 +8,7 @@
 class FPropertyChangeListener : public IPropertyChangeListener
 {
 public:
+	virtual ~FPropertyChangeListener() {}
 	/** IPropertyChangeListener interface */
 	virtual void SetObject( UObject& Object, const FPropertyListenerSettings& InPropertyListenerSettings ) override;
 	virtual bool ScanForChanges( bool bRecacheNewValues = true ) override;
diff --git a/Engine/Source/Editor/PropertyEditor/Private/UserInterface/PropertyTable/BooleanPropertyTableCellPresenter.h b/Engine/Source/Editor/PropertyEditor/Private/UserInterface/PropertyTable/BooleanPropertyTableCellPresenter.h
--- a/Engine/Source/Editor/PropertyEditor/Private/UserInterface/PropertyTable/BooleanPropertyTableCellPresenter.h
+++ b/Engine/Source/Editor/PropertyEditor/Private/UserInterface/PropertyTable/BooleanPropertyTableCellPresenter.h
@@ -9,6 +9,8 @@ public:
 
 	FBooleanPropertyTableCellPresenter( const TSharedRef< class FPropertyEditor >& InPropertyEditor );
 
+	virtual ~FBooleanPropertyTableCellPresenter() {}
+
 	virtual TSharedRef< class SWidget > ConstructDisplayWidget() override;
 
 	virtual bool RequiresDropDown() override;
diff --git a/Engine/Source/Editor/PropertyEditor/Private/UserInterface/PropertyTable/ColorPropertyTableCellPresenter.h b/Engine/Source/Editor/PropertyEditor/Private/UserInterface/PropertyTable/ColorPropertyTableCellPresenter.h
--- a/Engine/Source/Editor/PropertyEditor/Private/UserInterface/PropertyTable/ColorPropertyTableCellPresenter.h
+++ b/Engine/Source/Editor/PropertyEditor/Private/UserInterface/PropertyTable/ColorPropertyTableCellPresenter.h
@@ -9,6 +9,8 @@ public:
 
 	FColorPropertyTableCellPresenter( const TSharedRef< class FPropertyEditor >& InPropertyEditor, const TSharedRef< class IPropertyTableUtilities >& InPropertyUtilities );
 
+	virtual ~FColorPropertyTableCellPresenter() {}
+
 	virtual TSharedRef< class SWidget > ConstructDisplayWidget() override;
 
 	virtual bool RequiresDropDown() override;
diff --git a/Engine/Source/Editor/PropertyEditor/Private/UserInterface/PropertyTable/ObjectNameTableCellPresenter.h b/Engine/Source/Editor/PropertyEditor/Private/UserInterface/PropertyTable/ObjectNameTableCellPresenter.h
--- a/Engine/Source/Editor/PropertyEditor/Private/UserInterface/PropertyTable/ObjectNameTableCellPresenter.h
+++ b/Engine/Source/Editor/PropertyEditor/Private/UserInterface/PropertyTable/ObjectNameTableCellPresenter.h
@@ -10,12 +10,14 @@ class FObjectNameTableCellPresenter : public TSharedFromThis< FObjectNameTableCe
 public:
 
 	FObjectNameTableCellPresenter( const TSharedRef< IPropertyTableCell >& InCell )
-		: Cell( InCell )
-		, FocusWidget( SNullWidget::NullWidget )
+		: FocusWidget( SNullWidget::NullWidget )
+		, Cell( InCell )
 	{
 
 	}
 
+	virtual ~FObjectNameTableCellPresenter() {}
+
 	virtual TSharedRef< class SWidget > ConstructDisplayWidget() override
 	{
 		return SNew( SBox )
diff --git a/Engine/Source/Editor/PropertyEditor/Private/UserInterface/PropertyTable/PropertyTableWidgetHandle.h b/Engine/Source/Editor/PropertyEditor/Private/UserInterface/PropertyTable/PropertyTableWidgetHandle.h
--- a/Engine/Source/Editor/PropertyEditor/Private/UserInterface/PropertyTable/PropertyTableWidgetHandle.h
+++ b/Engine/Source/Editor/PropertyEditor/Private/UserInterface/PropertyTable/PropertyTableWidgetHandle.h
@@ -9,6 +9,8 @@ public:
 	FPropertyTableWidgetHandle(const TSharedRef<SPropertyTable>& InPropertyTableWidget)
 		: PropertyTableWidget(InPropertyTableWidget) {}
 
+	virtual ~FPropertyTableWidgetHandle() {}
+
 	virtual void RequestRefresh() override
 	{
 		PropertyTableWidget->RequestRefresh();
diff --git a/Engine/Source/Editor/PropertyEditor/Private/UserInterface/PropertyTable/TextPropertyTableCellPresenter.cpp b/Engine/Source/Editor/PropertyEditor/Private/UserInterface/PropertyTable/TextPropertyTableCellPresenter.cpp
--- a/Engine/Source/Editor/PropertyEditor/Private/UserInterface/PropertyTable/TextPropertyTableCellPresenter.cpp
+++ b/Engine/Source/Editor/PropertyEditor/Private/UserInterface/PropertyTable/TextPropertyTableCellPresenter.cpp
@@ -23,9 +23,9 @@
 
 
 FTextPropertyTableCellPresenter::FTextPropertyTableCellPresenter( const TSharedRef< class FPropertyEditor >& InPropertyEditor, const TSharedRef< class IPropertyTableUtilities >& InPropertyUtilities, FSlateFontInfo InFont /*= FEditorStyle::GetFontStyle( PropertyTableConstants::NormalFontStyle ) */ )
-	: PropertyEditor( InPropertyEditor )
+	: PropertyWidget( NULL )
+	, PropertyEditor( InPropertyEditor )
 	, PropertyUtilities( InPropertyUtilities )
-	, PropertyWidget( NULL )
 	, HasReadOnlyEditingWidget( false )
 	, Font(InFont)
 {
diff --git a/Engine/Source/Editor/PropertyEditor/Private/UserInterface/PropertyTable/TextPropertyTableCellPresenter.h b/Engine/Source/Editor/PropertyEditor/Private/UserInterface/PropertyTable/TextPropertyTableCellPresenter.h
--- a/Engine/Source/Editor/PropertyEditor/Private/UserInterface/PropertyTable/TextPropertyTableCellPresenter.h
+++ b/Engine/Source/Editor/PropertyEditor/Private/UserInterface/PropertyTable/TextPropertyTableCellPresenter.h
@@ -10,6 +10,8 @@ public:
 
 	FTextPropertyTableCellPresenter( const TSharedRef< class FPropertyEditor >& InPropertyEditor, const TSharedRef< class IPropertyTableUtilities >& InPropertyUtilities, FSlateFontInfo InFont = FEditorStyle::GetFontStyle( PropertyTableConstants::NormalFontStyle ) );
 
+	virtual ~FTextPropertyTableCellPresenter() {}
+
 	virtual TSharedRef< class SWidget > ConstructDisplayWidget() override;
 
 	virtual bool RequiresDropDown() override;
diff --git a/Engine/Source/Editor/SceneOutliner/Public/SceneOutlinerFilters.h b/Engine/Source/Editor/SceneOutliner/Public/SceneOutlinerFilters.h
--- a/Engine/Source/Editor/SceneOutliner/Public/SceneOutlinerFilters.h
+++ b/Engine/Source/Editor/SceneOutliner/Public/SceneOutlinerFilters.h
@@ -29,7 +29,7 @@ namespace SceneOutliner
 
 		/** Constructor for filtering the specified folder path */
 		FOutlinerFilterProxy(FName InFolder)
-			: FolderPath(InFolder), Actor(nullptr)
+			: Actor(nullptr), FolderPath(InFolder)
 		{}
 	};
 
diff --git a/Engine/Source/Editor/UnrealEd/Private/Commandlets/ParticleSystemAuditCommandlet.cpp b/Engine/Source/Editor/UnrealEd/Private/Commandlets/ParticleSystemAuditCommandlet.cpp
--- a/Engine/Source/Editor/UnrealEd/Private/Commandlets/ParticleSystemAuditCommandlet.cpp
+++ b/Engine/Source/Editor/UnrealEd/Private/Commandlets/ParticleSystemAuditCommandlet.cpp
@@ -422,7 +422,6 @@ void UParticleSystemAuditCommandlet::CheckPSysForDuplicateModules(UParticleSyste
 									ModuleCount = ModuleList->Find(Module);
 								}
 								check(ModuleCount);
-								*ModuleCount++;
 							}
 						}
 					}
diff --git a/Engine/Source/Editor/UnrealEd/Private/EditorModes.cpp b/Engine/Source/Editor/UnrealEd/Private/EditorModes.cpp
--- a/Engine/Source/Editor/UnrealEd/Private/EditorModes.cpp
+++ b/Engine/Source/Editor/UnrealEd/Private/EditorModes.cpp
@@ -185,13 +185,13 @@ namespace
 // FEdMode
 
 FEdMode::FEdMode()
-	: CurrentTool( NULL )
+	/** Deprecated members */
+	: Name(Info.Name), ID(Info.ID), IconBrush(Info.IconBrush), bVisible(Info.bVisible), PriorityOrder(Info.PriorityOrder)
+	, bPendingDeletion( false )
+	, CurrentTool( NULL )
 	, EditedPropertyName(TEXT(""))
 	, EditedPropertyIndex( INDEX_NONE )
 	, bEditedPropertyIsTransform( false )
-	, bPendingDeletion( false )
-	/** Deprecated members */
-	, Name(Info.Name), ID(Info.ID), IconBrush(Info.IconBrush), bVisible(Info.bVisible), PriorityOrder(Info.PriorityOrder)
 {
 	bDrawKillZ = true;
 }
@@ -1161,13 +1161,13 @@ FEditorModeTools::FEditorModeTools()
 	,	Snapping( 0 )
 	,	SnappedActor( 0 )
 	,	TranslateRotateXAxisAngle(0)
+	,	DefaultID(FBuiltinEditorModes::EM_Default)
 	,	WidgetMode( FWidget::WM_Translate )
 	,	OverrideWidgetMode( FWidget::WM_None )
 	,	bShowWidget( 1 )
 	,	bHideViewportUI(false)
 	,	CoordSystem(COORD_World)
 	,	bIsTracking(false)
-	,	DefaultID(FBuiltinEditorModes::EM_Default)
 {
 	// Load the last used settings
 	LoadConfig();
diff --git a/Engine/Source/Editor/UnrealEd/Private/Fbx/FbxMainExport.cpp b/Engine/Source/Editor/UnrealEd/Private/Fbx/FbxMainExport.cpp
--- a/Engine/Source/Editor/UnrealEd/Private/Fbx/FbxMainExport.cpp
+++ b/Engine/Source/Editor/UnrealEd/Private/Fbx/FbxMainExport.cpp
@@ -267,7 +267,7 @@ void FFbxExporter::ExportLevelMesh( ULevel* InLevel, AMatineeActor* InMatineeAct
 	for (int32 ActorIndex = 0; ActorIndex < ActorCount; ++ActorIndex)
 	{
 		AActor* Actor = World->GetCurrentLevel()->Actors[ActorIndex];
-		if ( Actor != NULL && ( !bSelectedOnly || bSelectedOnly && Actor->IsSelected() ) )
+		if ( Actor != NULL && ( !bSelectedOnly || ( bSelectedOnly && Actor->IsSelected() ) ) )
 		{
 			if (Actor->IsA(ALight::StaticClass()))
 			{
diff --git a/Engine/Source/Editor/WorkspaceMenuStructure/Private/WorkspaceMenuStructureModule.cpp b/Engine/Source/Editor/WorkspaceMenuStructure/Private/WorkspaceMenuStructureModule.cpp
--- a/Engine/Source/Editor/WorkspaceMenuStructure/Private/WorkspaceMenuStructureModule.cpp
+++ b/Engine/Source/Editor/WorkspaceMenuStructure/Private/WorkspaceMenuStructureModule.cpp
@@ -93,6 +93,8 @@ public:
 		ResetToolsCategory();
 	}
 
+	virtual ~FWorkspaceMenuStructure() {}
+
 private:
 	TSharedPtr<FWorkspaceItem> MenuRoot;
 	TSharedPtr<FWorkspaceItem> AssetEditorCategory;
diff --git a/Engine/Source/Editor/WorkspaceMenuStructure/Public/WorkspaceMenuStructureModule.h b/Engine/Source/Editor/WorkspaceMenuStructure/Public/WorkspaceMenuStructureModule.h
--- a/Engine/Source/Editor/WorkspaceMenuStructure/Public/WorkspaceMenuStructureModule.h
+++ b/Engine/Source/Editor/WorkspaceMenuStructure/Public/WorkspaceMenuStructureModule.h
@@ -14,6 +14,8 @@ class FWorkspaceMenuStructureModule : public IModuleInterface
 
 public:
 
+	virtual ~FWorkspaceMenuStructureModule() {}
+
 	virtual void StartupModule();
 
 	virtual void ShutdownModule();
