Add go to C++ file and documentation to right clicking on content browser assets

diff --git a/Engine/Source/Editor/ContentBrowser/Private/AssetContextMenu.cpp b/Engine/Source/Editor/ContentBrowser/Private/AssetContextMenu.cpp
--- a/Engine/Source/Editor/ContentBrowser/Private/AssetContextMenu.cpp
+++ b/Engine/Source/Editor/ContentBrowser/Private/AssetContextMenu.cpp
@@ -23,6 +23,10 @@
 #include "AssetToolsModule.h"
 #include "ComponentAssetBroker.h"
 
+#include "SourceCodeNavigation.h"
+#include "IDocumentation.h"
+#include "EditorClassUtils.h"
+
 #define LOCTEXT_NAMESPACE "ContentBrowser"
 
 FAssetContextMenu::FAssetContextMenu(const TWeakPtr<SAssetView>& InAssetView)
@@ -88,6 +92,9 @@ TSharedRef<SWidget> FAssetContextMenu::MakeContextMenu(const TArray<FAssetData>&
 		// Add reference options
 		AddReferenceMenuOptions(MenuBuilder);
 
+		// Add documentation options
+		AddDocumentationMenuOptions(MenuBuilder);
+
 		// Add source control options
 		AddSourceControlMenuOptions(MenuBuilder);
 
@@ -325,6 +332,78 @@ bool FAssetContextMenu::AddReferenceMenuOptions(FMenuBuilder& MenuBuilder)
 	return true;
 }
 
+bool FAssetContextMenu::AddDocumentationMenuOptions(FMenuBuilder& MenuBuilder)
+{
+	bool bAddedOption = false;
+
+	// Objects must be loaded for this operation... for now
+	UClass* SelectedClass = (SelectedAssets.Num() > 0 ? SelectedAssets[0].GetClass() : nullptr);
+	for (const FAssetData& AssetData : SelectedAssets)
+	{
+		if (SelectedClass != AssetData.GetClass())
+		{
+			SelectedClass = nullptr;
+			break;
+		}
+	}
+
+	// Go to C++ Code
+	if( SelectedClass != nullptr )
+	{
+		// Blueprints are special.  We won't link to C++ and for documentation we'll use the class it is generated from
+		const bool bIsBlueprint = SelectedClass->IsChildOf<UBlueprint>();
+		if (bIsBlueprint)
+		{
+			FString* ParentClassPath = SelectedAssets[0].TagsAndValues.Find(GET_MEMBER_NAME_CHECKED(UBlueprint,ParentClass));
+			if (ParentClassPath)
+			{
+				SelectedClass = FindObject<UClass>(nullptr,**ParentClassPath);
+			}
+		}
+
+		if ( !bIsBlueprint && FSourceCodeNavigation::IsCompilerAvailable() )
+		{
+			FString ClassHeaderPath;
+			if( FSourceCodeNavigation::FindClassHeaderPath( SelectedClass, ClassHeaderPath ) && IFileManager::Get().FileSize( *ClassHeaderPath ) != INDEX_NONE )
+			{
+				bAddedOption = true;
+
+				const FString CodeFileName = FPaths::GetCleanFilename( *ClassHeaderPath );
+
+				MenuBuilder.BeginSection( "ActorCode", LOCTEXT("ActorCodeHeading", "C++") );
+				{
+					MenuBuilder.AddMenuEntry(
+						FText::Format( LOCTEXT("GoToCodeForActor", "Open {0}"), FText::FromString( CodeFileName ) ),
+						FText::Format( LOCTEXT("GoToCodeForActor_ToolTip", "Opens the header file for this actor ({0}) in a code editing program"), FText::FromString( CodeFileName ) ),
+						FSlateIcon(),
+						FUIAction( FExecuteAction::CreateSP( this, &FAssetContextMenu::ExecuteGoToCodeForAsset, SelectedClass ) )
+						);
+				}
+				MenuBuilder.EndSection();
+			}
+		}
+
+		const FString DocumentationLink = FEditorClassUtils::GetDocumentationLink(SelectedClass);
+		if (!DocumentationLink.IsEmpty())
+		{
+			bAddedOption = true;
+
+			MenuBuilder.BeginSection( "ActorDocumentation", LOCTEXT("ActorDocsHeading", "Documentation") );
+			{
+					MenuBuilder.AddMenuEntry(
+						LOCTEXT("GoToDocsForActor", "Full documentation"),
+						LOCTEXT("GoToDocsForActor_ToolTip", "Click to open documentation for this actor"),
+						FSlateIcon(),
+						FUIAction( FExecuteAction::CreateSP( this, &FAssetContextMenu::ExecuteGoToDocsForAsset, SelectedClass ) )
+						);
+			}
+			MenuBuilder.EndSection();
+		}
+	}
+
+	return bAddedOption;
+}
+
 bool FAssetContextMenu::AddAssetTypeMenuOptions(FMenuBuilder& MenuBuilder)
 {
 	bool bAnyTypeOptions = false;
@@ -917,6 +996,31 @@ void FAssetContextMenu::ExecuteShowReferenceViewer()
 	}
 }
 
+void FAssetContextMenu::ExecuteGoToCodeForAsset(UClass* SelectedClass)
+{
+	if (SelectedClass)
+	{
+		FString ClassHeaderPath;
+		if( FSourceCodeNavigation::FindClassHeaderPath( SelectedClass, ClassHeaderPath ) && IFileManager::Get().FileSize( *ClassHeaderPath ) != INDEX_NONE )
+		{
+			const FString AbsoluteHeaderPath = IFileManager::Get().ConvertToAbsolutePathForExternalAppForRead(*ClassHeaderPath);
+			FSourceCodeNavigation::OpenSourceFile( AbsoluteHeaderPath );
+		}
+	}
+}
+
+void FAssetContextMenu::ExecuteGoToDocsForAsset(UClass* SelectedClass)
+{
+	if (SelectedClass)
+	{
+		FString DocumentationLink = FEditorClassUtils::GetDocumentationLink(SelectedClass);
+		if (!DocumentationLink.IsEmpty())
+		{
+			IDocumentation::Get()->Open( DocumentationLink );
+		}
+	}
+}
+
 void FAssetContextMenu::ExecuteCopyReference()
 {
 	ContentBrowserUtils::CopyAssetReferencesToClipboard(SelectedAssets);
diff --git a/Engine/Source/Editor/ContentBrowser/Private/AssetContextMenu.h b/Engine/Source/Editor/ContentBrowser/Private/AssetContextMenu.h
--- a/Engine/Source/Editor/ContentBrowser/Private/AssetContextMenu.h
+++ b/Engine/Source/Editor/ContentBrowser/Private/AssetContextMenu.h
@@ -40,6 +40,9 @@ private:
 	/** Adds asset reference menu options to a menu builder. Returns true if any options were added. */
 	bool AddReferenceMenuOptions(FMenuBuilder& MenuBuilder);
 
+	/** Adds asset documentation menu options to a menu builder. Returns true if any options were added. */
+	bool AddDocumentationMenuOptions(FMenuBuilder& MenuBuilder);
+
 	/** Adds asset type-specific menu options to a menu builder. Returns true if any options were added. */
 	bool AddAssetTypeMenuOptions(FMenuBuilder& MenuBuilder);
 
@@ -100,6 +103,12 @@ private:
 	/** Handler for ShowReferenceViewer */
 	void ExecuteShowReferenceViewer();
 
+	/** Handler for GoToAssetCode */
+	void ExecuteGoToCodeForAsset(UClass* SelectedClass);
+
+	/** Handler for GoToAssetDocs */
+	void ExecuteGoToDocsForAsset(UClass* SelectedClass);
+
 	/** Handler for CopyReference */
 	void ExecuteCopyReference();
 
