Folder Rename can now be done in the Content Browser Path View

#ttp 342267 - Editor Usability: Select a folder in Content Browser and pressing F2 to rename it, triggers a rename of an actor or asset instead!

#branch UE4

SContentBrowser
Added functions for Can/Execute Rename/Delete which forward to the active context menu
Moved Rename/Delete bind commands to the content browser, so it can forward the request to the correct context menu
(this was necessary as it wasn't possible to bind to both context menus without the code complaining about it already being bound)
Modified GetFolderContextMenu to take an extra param to indicate whether it was the path view or asset view which called it - it then uses this to clear the asset view selection when the path view context menu is requested - we can then use the selection information to determine which context menu was opened later on when processing rename/delete (ideally it should clear when the path view recieves focus, but that's harder to determine).

FAssetContextMenu
Made Can/Execute Rename/Delete public - so they can be called by the Content Browser
Removed BindCommands as it's no longer needed
Removed condition for adding Rename to the menu, as it now always is and is greyed out if not possible.

FPathContextMenu
Added delegate for when rename folder is requested - the same as FAssetContextMenu
Added handling for whether folder renaming is possible, and executing.
Removed condition for adding Rename to the menu, as it now always is and is greyed out if not possible.
Fixed issue with Delete not mapping to keyboard shortcut

SPathView
Added function to rename folder (similar to SAssetView)
Fixed issue with VerifyFolderNameChanged not generating the correct path when checking to see if a folder exists (it was previous just checking to old path - which would always return true).
Modified FolderNameChanged to take 'OldPath' as a param (inline with the delegate change). This function is called when a folder is created or renamed... if it's the latter it has to handle moving the contents of the folder to the new location, which it could only do if it knew where the previous location was (again, similar to SAssetView)

SAssetTreeItem
Added extra param to FOnNameChanged (OldPath) - so we can move assets when a folder is renamed

FTreeItem
Renamed variable bNewFolder to bNamingFolder - as it's true whenever the folder is being named, not just when it's new

diff --git a/Engine/Source/Editor/ContentBrowser/Private/AssetContextMenu.cpp b/Engine/Source/Editor/ContentBrowser/Private/AssetContextMenu.cpp
--- a/Engine/Source/Editor/ContentBrowser/Private/AssetContextMenu.cpp
+++ b/Engine/Source/Editor/ContentBrowser/Private/AssetContextMenu.cpp
@@ -42,19 +42,6 @@ FAssetContextMenu::FAssetContextMenu(const TWeakPtr<SAssetView>& InAssetView)
 
 }
 
-void FAssetContextMenu::BindCommands(TSharedPtr< FUICommandList > InCommandList)
-{
-	InCommandList->MapAction( FGenericCommands::Get().Rename, FUIAction(
-		FExecuteAction::CreateSP( this, &FAssetContextMenu::ExecuteRename ),
-		FCanExecuteAction::CreateSP( this, &FAssetContextMenu::CanExecuteRename )
-		));
-
-	InCommandList->MapAction( FGenericCommands::Get().Delete, FUIAction(
-		FExecuteAction::CreateSP( this, &FAssetContextMenu::ExecuteDelete ),
-		FCanExecuteAction::CreateSP( this, &FAssetContextMenu::CanExecuteDelete )
-		));
-}
-
 TSharedRef<SWidget> FAssetContextMenu::MakeContextMenu(const TArray<FAssetData>& InSelectedAssets, const FSourcesData& InSourcesData, TSharedPtr< FUICommandList > InCommandList)
 {
 	SelectedAssets = InSelectedAssets;
@@ -184,13 +171,10 @@ bool FAssetContextMenu::AddCommonMenuOptions(FMenuBuilder& MenuBuilder)
 			)
 		);
 
-	if ( SelectedAssets.Num() == 1 )
-	{
-		MenuBuilder.AddMenuEntry( FGenericCommands::Get().Rename, NAME_None, 
-			LOCTEXT("Rename", "Rename"),
-			LOCTEXT("RenameTooltip", "Rename the selected asset.")
-			);
-	}
+	MenuBuilder.AddMenuEntry( FGenericCommands::Get().Rename, NAME_None, 
+		LOCTEXT("Rename", "Rename"),
+		LOCTEXT("RenameTooltip", "Rename the selected asset.")
+		);
 
 	MenuBuilder.AddMenuEntry( FGenericCommands::Get().Delete, NAME_None, 
 		LOCTEXT("Delete", "Delete"),
@@ -852,7 +836,7 @@ void FAssetContextMenu::ExecuteDuplicate()
 	}
 }
 
-void FAssetContextMenu::ExecuteRename() 
+void FAssetContextMenu::ExecuteRename()
 {
 	TArray< FAssetData > AssetViewSelectedAssets = AssetView.Pin()->GetSelectedAssets();
 	TArray< FString > SelectedFolders = AssetView.Pin()->GetSelectedFolders();
@@ -872,7 +856,7 @@ void FAssetContextMenu::ExecuteRename()
 	}
 }
 
-void FAssetContextMenu::ExecuteDelete() 
+void FAssetContextMenu::ExecuteDelete()
 {
 	TArray< FAssetData > AssetViewSelectedAssets = AssetView.Pin()->GetSelectedAssets();
 	if(AssetViewSelectedAssets.Num() > 0)
diff --git a/Engine/Source/Editor/ContentBrowser/Private/AssetContextMenu.h b/Engine/Source/Editor/ContentBrowser/Private/AssetContextMenu.h
--- a/Engine/Source/Editor/ContentBrowser/Private/AssetContextMenu.h
+++ b/Engine/Source/Editor/ContentBrowser/Private/AssetContextMenu.h
@@ -8,9 +8,6 @@ public:
 	/** Constructor */
 	FAssetContextMenu(const TWeakPtr<SAssetView>& InAssetView);
 
-	/** Bind menu selection commands to the command list */
-	void BindCommands(TSharedPtr< FUICommandList > InCommandList);
-
 	/** Makes the context menu widget */
 	TSharedRef<SWidget> MakeContextMenu(const TArray<FAssetData>& SelectedAssets, const FSourcesData& InSourcesData, TSharedPtr< FUICommandList > InCommandList);
 
@@ -33,6 +30,18 @@ public:
 	DECLARE_DELEGATE(FOnAssetViewRefreshRequested);
 	void SetOnAssetViewRefreshRequested(const FOnAssetViewRefreshRequested& InOnAssetViewRefreshRequested);
 
+	/** Handler to check to see if a rename command is allowed */
+	bool CanExecuteRename() const;
+
+	/** Handler for Rename */
+	void ExecuteRename();
+
+	/** Handler to check to see if a delete command is allowed */
+	bool CanExecuteDelete() const;
+
+	/** Handler for Delete */
+	void ExecuteDelete();
+
 private:
 	/** Adds common menu options to a menu builder. Returns true if any options were added. */
 	bool AddCommonMenuOptions(FMenuBuilder& MenuBuilder);
@@ -79,12 +88,6 @@ private:
 	/** Handler for Duplicate */
 	void ExecuteDuplicate();
 
-	/** Handler for Rename */
-	void ExecuteRename();
-
-	/** Handler for Delete */
-	void ExecuteDelete();
-
 	/** Handler for confirmation of folder deletion */
 	FReply ExecuteDeleteFolderConfirmed();
 
@@ -169,12 +172,6 @@ private:
 	/** Handler to check to see if a duplicate command is allowed */
 	bool CanExecuteDuplicate() const;
 
-	/** Handler to check to see if a rename command is allowed */
-	bool CanExecuteRename() const;
-
-	/** Handler to check to see if a delete command is allowed */
-	bool CanExecuteDelete() const;
-
 	/** Handler to check to see if a "Remove from collection" command is allowed */
 	bool CanExecuteRemoveFromCollection() const;
 
diff --git a/Engine/Source/Editor/ContentBrowser/Private/PathContextMenu.cpp b/Engine/Source/Editor/ContentBrowser/Private/PathContextMenu.cpp
--- a/Engine/Source/Editor/ContentBrowser/Private/PathContextMenu.cpp
+++ b/Engine/Source/Editor/ContentBrowser/Private/PathContextMenu.cpp
@@ -26,6 +26,11 @@ void FPathContextMenu::SetOnNewAssetRequested(const FNewAssetContextMenu::FOnNew
 	OnNewAssetRequested = InOnNewAssetRequested;
 }
 
+void FPathContextMenu::SetOnRenameFolderRequested(const FOnRenameFolderRequested& InOnRenameFolderRequested)
+{
+	OnRenameFolderRequested = InOnRenameFolderRequested;
+}
+
 TSharedRef<FExtender> FPathContextMenu::MakePathViewContextMenuExtender(const TArray<FString>& InSelectedPaths)
 {
 	SelectedPaths = InSelectedPaths;
@@ -78,13 +83,10 @@ void FPathContextMenu::MakePathViewContextMenu(FMenuBuilder& MenuBuilder)
 				FUIAction( FExecuteAction::CreateSP( this, &FPathContextMenu::ExecuteExplore ) )
 				);
 
-			if ( SelectedPaths.Num() == 1 )
-			{
-				MenuBuilder.AddMenuEntry( FGenericCommands::Get().Rename, NAME_None,
-					LOCTEXT("RenameFolder", "Rename"),
-					LOCTEXT("RenameFolderTooltip", "Rename the selected folder.")
-					);
-			}
+			MenuBuilder.AddMenuEntry(FGenericCommands::Get().Rename, NAME_None,
+				LOCTEXT("RenameFolder", "Rename"),
+				LOCTEXT("RenameFolderTooltip", "Rename the selected folder.")
+				);
 
 			// If any colors have already been set, display color options as a sub menu
 			if ( ContentBrowserUtils::HasCustomColors() )
@@ -121,11 +123,9 @@ void FPathContextMenu::MakePathViewContextMenu(FMenuBuilder& MenuBuilder)
 			    );
     
 		    // Delete
-		    MenuBuilder.AddMenuEntry(
+			MenuBuilder.AddMenuEntry(FGenericCommands::Get().Delete, NAME_None,
 			    LOCTEXT("DeleteFolder", "Delete"),
-			    LOCTEXT("DeleteFolderTooltip", "Removes this folder and all assets it contains."),
-			    FSlateIcon(),
-			    FUIAction( FExecuteAction::CreateSP( this, &FPathContextMenu::ExecuteDeleteFolder ) )
+			    LOCTEXT("DeleteFolderTooltip", "Removes this folder and all assets it contains.")
 			    );
 
 			// Reference Viewer
@@ -329,6 +329,24 @@ void FPathContextMenu::ExecuteExplore()
 	}
 }
 
+bool FPathContextMenu::CanExecuteRename() const
+{
+	if (SelectedPaths.Num() == 1 && !ContentBrowserUtils::IsAssetRootDir(SelectedPaths[0]))
+	{
+		return true;
+	}
+	return false;
+}
+
+void FPathContextMenu::ExecuteRename()
+{
+	check(SelectedPaths.Num() == 1);
+	if (OnRenameFolderRequested.IsBound())
+	{
+		OnRenameFolderRequested.Execute(SelectedPaths[0]);
+	}
+}
+
 void FPathContextMenu::ExecuteResetColor()
 {
 	ResetColors();
@@ -437,9 +455,15 @@ void FPathContextMenu::ExecuteSaveFolder()
 	}
 }
 
-void FPathContextMenu::ExecuteDeleteFolder()
+bool FPathContextMenu::CanExecuteDelete() const
+{
+	return SelectedPaths.Num() > 0;
+}
+
+void FPathContextMenu::ExecuteDelete()
 {
-	if ( ParentContent.IsValid() )
+	check(SelectedPaths.Num() > 0);
+	if (ParentContent.IsValid())
 	{
 		FText Prompt;
 		if ( SelectedPaths.Num() == 1 )
@@ -450,7 +474,7 @@ void FPathContextMenu::ExecuteDeleteFolder()
 		{
 			Prompt = FText::Format(LOCTEXT("FolderDeleteConfirm_Multiple", "Delete {0} folders?"), FText::AsNumber(SelectedPaths.Num()));
 		}
-
+		
 		// Spawn a confirmation dialog since this is potentially a highly destructive operation
 		FOnClicked OnYesClicked = FOnClicked::CreateSP( this, &FPathContextMenu::ExecuteDeleteFolderConfirmed );
  		ContentBrowserUtils::DisplayConfirmationPopup(
diff --git a/Engine/Source/Editor/ContentBrowser/Private/PathContextMenu.h b/Engine/Source/Editor/ContentBrowser/Private/PathContextMenu.h
--- a/Engine/Source/Editor/ContentBrowser/Private/PathContextMenu.h
+++ b/Engine/Source/Editor/ContentBrowser/Private/PathContextMenu.h
@@ -13,6 +13,10 @@ public:
 	/** Sets the handler for when new assets are requested */
 	void SetOnNewAssetRequested(const FNewAssetContextMenu::FOnNewAssetRequested& InOnNewAssetRequested);
 
+	/** Delegate for when the context menu requests a rename of a folder */
+	DECLARE_DELEGATE_OneParam(FOnRenameFolderRequested, const FString& /*FolderToRename*/);
+	void SetOnRenameFolderRequested(const FOnRenameFolderRequested& InOnRenameFolderRequested);
+
 	/** Makes the asset tree context menu extender */
 	TSharedRef<FExtender> MakePathViewContextMenuExtender(const TArray<FString>& InSelectedPaths);
 
@@ -31,6 +35,18 @@ public:
 	/** Handler for when "Explore" is selected */
 	void ExecuteExplore();
 
+	/** Handler to check to see if a rename command is allowed */
+	bool CanExecuteRename() const;
+
+	/** Handler for Rename */
+	void ExecuteRename();
+
+	/** Handler to check to see if a delete command is allowed */
+	bool CanExecuteDelete() const;
+
+	/** Handler for Delete */
+	void ExecuteDelete();
+
 	/** Handler for when reset color is selected */
 	void ExecuteResetColor();
 
@@ -40,9 +56,6 @@ public:
 	/** Handler for when "Save" is selected */
 	void ExecuteSaveFolder();
 
-	/** Handler for when "Delete" is selected */
-	void ExecuteDeleteFolder();
-
 	/** Handler for when "ReferenceViewer" is selected */
 	void ExecuteReferenceViewer();
 
@@ -105,6 +118,7 @@ private:
 	TArray<FString> SelectedPaths;
 	TWeakPtr<SWidget> ParentContent;
 	FNewAssetContextMenu::FOnNewAssetRequested OnNewAssetRequested;
+	FOnRenameFolderRequested OnRenameFolderRequested;
 
 	/** Cached SCC CanExecute vars */
 	bool bCanExecuteSCCCheckOut;
diff --git a/Engine/Source/Editor/ContentBrowser/Private/PathViewTypes.h b/Engine/Source/Editor/ContentBrowser/Private/PathViewTypes.h
--- a/Engine/Source/Editor/ContentBrowser/Private/PathViewTypes.h
+++ b/Engine/Source/Editor/ContentBrowser/Private/PathViewTypes.h
@@ -17,8 +17,8 @@ struct FTreeItem : public TSharedFromThis<FTreeItem>
 	FString FolderName;
 	/** The path of the tree item including the name */
 	FString FolderPath;
-	/** If true, this folder is in the process of being created */
-	bool bNewFolder;
+	/** If true, this folder is in the process of being named */
+	bool bNamingFolder;
 
 	/** The children of this tree item */
 	TArray< TSharedPtr<FTreeItem> > Children;
@@ -33,10 +33,10 @@ struct FTreeItem : public TSharedFromThis<FTreeItem>
 	FRenamedRequestEvent OnRenamedRequestEvent;
 
 	/** Constructor */
-	FTreeItem(const FString& InFolderName, const FString& InFolderPath, TSharedPtr<FTreeItem> InParent, bool InNewFolder = false)
+	FTreeItem(const FString& InFolderName, const FString& InFolderPath, TSharedPtr<FTreeItem> InParent, bool InNamingFolder = false)
 		: FolderName(InFolderName)
 		, FolderPath(InFolderPath)
-		, bNewFolder(InNewFolder)
+		, bNamingFolder(InNamingFolder)
 		, Parent(InParent)
 	{}
 
diff --git a/Engine/Source/Editor/ContentBrowser/Private/SContentBrowser.cpp b/Engine/Source/Editor/ContentBrowser/Private/SContentBrowser.cpp
--- a/Engine/Source/Editor/ContentBrowser/Private/SContentBrowser.cpp
+++ b/Engine/Source/Editor/ContentBrowser/Private/SContentBrowser.cpp
@@ -57,6 +57,7 @@ void SContentBrowser::Construct( const FArguments& InArgs, const FName& InInstan
 
 	PathContextMenu = MakeShareable(new FPathContextMenu( AsShared() ));
 	PathContextMenu->SetOnNewAssetRequested( FNewAssetContextMenu::FOnNewAssetRequested::CreateSP(this, &SContentBrowser::NewAssetRequested) );
+	PathContextMenu->SetOnRenameFolderRequested(FPathContextMenu::FOnRenameFolderRequested::CreateSP(this, &SContentBrowser::OnRenameFolderRequested));
 
 	FrontendFilters = MakeShareable(new AssetFilterCollectionType());
 	TextFilter = MakeShareable( new FFrontendFilter_Text() );
@@ -384,7 +385,7 @@ void SContentBrowser::Construct( const FArguments& InArgs, const FName& InInstan
 							[
 								SAssignNew( PathViewPtr, SPathView )
 								.OnPathSelected( this, &SContentBrowser::PathSelected )
-								.OnGetFolderContextMenu( this, &SContentBrowser::GetFolderContextMenu )
+								.OnGetFolderContextMenu( this, &SContentBrowser::GetFolderContextMenu, true )
 								.OnGetPathContextMenuExtender( this, &SContentBrowser::GetPathContextMenuExtender )
 								.FocusSearchBoxWhenOpened( false )
 								.ShowTreeTitle( false )
@@ -546,7 +547,7 @@ void SContentBrowser::Construct( const FArguments& InArgs, const FName& InInstan
 							.OnAssetSelected(this, &SContentBrowser::OnAssetSelectionChanged)
 							.OnAssetsActivated(this, &SContentBrowser::OnAssetsActivated)
 							.OnGetAssetContextMenu(this, &SContentBrowser::OnGetAssetContextMenu)
-							.OnGetFolderContextMenu(this, &SContentBrowser::GetFolderContextMenu)
+							.OnGetFolderContextMenu(this, &SContentBrowser::GetFolderContextMenu, false)
 							.OnGetPathContextMenuExtender(this, &SContentBrowser::GetPathContextMenuExtender)
 							.OnFindInAssetTreeRequested(this, &SContentBrowser::OnFindInAssetTreeRequested)
 							.OnAssetRenameCommitted(this, &SContentBrowser::OnAssetRenameCommitted)
@@ -567,7 +568,6 @@ void SContentBrowser::Construct( const FArguments& InArgs, const FName& InInstan
 	];
 
 	AssetContextMenu = MakeShareable(new FAssetContextMenu(AssetViewPtr));
-	AssetContextMenu->BindCommands(Commands);
 	AssetContextMenu->SetOnFindInAssetTreeRequested( FOnFindInAssetTreeRequested::CreateSP(this, &SContentBrowser::OnFindInAssetTreeRequested) );
 	AssetContextMenu->SetOnRenameRequested( FAssetContextMenu::FOnRenameRequested::CreateSP(this, &SContentBrowser::OnRenameRequested) );
 	AssetContextMenu->SetOnRenameFolderRequested( FAssetContextMenu::FOnRenameFolderRequested::CreateSP(this, &SContentBrowser::OnRenameFolderRequested) );
@@ -606,6 +606,16 @@ void SContentBrowser::BindCommands()
 {
 	Commands = TSharedPtr< FUICommandList >(new FUICommandList);
 
+	Commands->MapAction(FGenericCommands::Get().Rename, FUIAction(
+		FExecuteAction::CreateSP(this, &SContentBrowser::OnRename),
+		FCanExecuteAction::CreateSP(this, &SContentBrowser::CanRename)
+		));
+
+	Commands->MapAction(FGenericCommands::Get().Delete, FUIAction(
+		FExecuteAction::CreateSP(this, &SContentBrowser::OnDelete),
+		FCanExecuteAction::CreateSP(this, &SContentBrowser::CanDelete)
+		));
+
 	Commands->MapAction(FContentBrowserCommands::Get().OpenAssetsOrFolders, FUIAction(
 		FExecuteAction::CreateSP(this, &SContentBrowser::OnOpenAssetsOrFolders)
 		));
@@ -1494,6 +1504,76 @@ FReply SContentBrowser::ForwardClicked()
 	return FReply::Handled();
 }
 
+bool SContentBrowser::CanRename() const
+{
+	const TArray<TSharedPtr<FAssetViewItem>>& SelectedItems = AssetViewPtr->GetSelectedItems();
+	if (SelectedItems.Num() > 0)
+	{
+		return AssetContextMenu->CanExecuteRename();
+	}
+	else
+	{
+		const TArray<FString>& SelectedPaths = PathViewPtr->GetSelectedPaths();
+		if (SelectedPaths.Num() > 0)
+		{
+			return PathContextMenu->CanExecuteRename();
+		}
+	}
+	return false;
+}
+
+void SContentBrowser::OnRename()
+{
+	const TArray<TSharedPtr<FAssetViewItem>>& SelectedItems = AssetViewPtr->GetSelectedItems();
+	if (SelectedItems.Num() > 0)
+	{
+		AssetContextMenu->ExecuteRename();
+	}
+	else
+	{
+		const TArray<FString>& SelectedPaths = PathViewPtr->GetSelectedPaths();
+		if (SelectedPaths.Num() > 0)
+		{
+			PathContextMenu->ExecuteRename();
+		}
+	}
+}
+
+bool SContentBrowser::CanDelete() const
+{
+	const TArray<TSharedPtr<FAssetViewItem>>& SelectedItems = AssetViewPtr->GetSelectedItems();
+	if (SelectedItems.Num() > 0)
+	{
+		return AssetContextMenu->CanExecuteDelete();
+	}
+	else
+	{
+		const TArray<FString>& SelectedPaths = PathViewPtr->GetSelectedPaths();
+		if (SelectedPaths.Num() > 0)
+		{
+			return PathContextMenu->CanExecuteDelete();
+		}
+	}
+	return false;
+}
+
+void SContentBrowser::OnDelete()
+{
+	const TArray<TSharedPtr<FAssetViewItem>>& SelectedItems = AssetViewPtr->GetSelectedItems();
+	if (SelectedItems.Num() > 0)
+	{
+		AssetContextMenu->ExecuteDelete();
+	}
+	else
+	{
+		const TArray<FString>& SelectedPaths = PathViewPtr->GetSelectedPaths();
+		if (SelectedPaths.Num() > 0)
+		{
+			PathContextMenu->ExecuteDelete();
+		}
+	}
+}
+
 void SContentBrowser::OnOpenAssetsOrFolders()
 {
 	AssetViewPtr->OnOpenAssetsOrFolders();
@@ -1743,7 +1823,19 @@ void SContentBrowser::OnRenameRequested(const FAssetData& AssetData)
 
 void SContentBrowser::OnRenameFolderRequested(const FString& FolderToRename)
 {
-	AssetViewPtr->RenameFolder(FolderToRename);
+	const TArray<FString>& SelectedFolders = AssetViewPtr->GetSelectedFolders();
+	if (SelectedFolders.Num() > 0)
+	{
+		AssetViewPtr->RenameFolder(FolderToRename);
+	}
+	else
+	{
+		const TArray<FString>& SelectedPaths = PathViewPtr->GetSelectedPaths();
+		if (SelectedPaths.Num() > 0)
+		{
+			PathViewPtr->RenameFolder(FolderToRename);
+		}
+	}
 }
 
 void SContentBrowser::OnDuplicateRequested(const TWeakObjectPtr<UObject>& OriginalObject)
@@ -1880,8 +1972,15 @@ TArray<FString> SContentBrowser::GetAssetSearchSuggestions() const
 	return AllSuggestions;
 }
 
-TSharedPtr<SWidget> SContentBrowser::GetFolderContextMenu(const TArray<FString>& SelectedPaths, FContentBrowserMenuExtender_SelectedPaths InMenuExtender, FOnCreateNewFolder InOnCreateNewFolder)
+TSharedPtr<SWidget> SContentBrowser::GetFolderContextMenu(const TArray<FString>& SelectedPaths, FContentBrowserMenuExtender_SelectedPaths InMenuExtender, FOnCreateNewFolder InOnCreateNewFolder, bool bPathView)
 {
+	// Clear any selection in the asset view, as it'll conflict with other view info
+	// This is important for determining which context menu may be open based on the asset selection for rename/delete operations
+	if (bPathView)
+	{
+		AssetViewPtr->ClearSelection();
+	}
+	
 	TSharedPtr<FExtender> Extender;
 	if(InMenuExtender.IsBound())
 	{
@@ -1890,7 +1989,7 @@ TSharedPtr<SWidget> SContentBrowser::GetFolderContextMenu(const TArray<FString>&
 
 	const bool bInShouldCloseWindowAfterSelection = true;
 	FMenuBuilder MenuBuilder(bInShouldCloseWindowAfterSelection, Commands, Extender, true);
-
+	
 	// New Folder
 	MenuBuilder.AddMenuEntry(
 		LOCTEXT("NewFolder", "New Folder"),
diff --git a/Engine/Source/Editor/ContentBrowser/Private/SContentBrowser.h b/Engine/Source/Editor/ContentBrowser/Private/SContentBrowser.h
--- a/Engine/Source/Editor/ContentBrowser/Private/SContentBrowser.h
+++ b/Engine/Source/Editor/ContentBrowser/Private/SContentBrowser.h
@@ -200,6 +200,18 @@ private:
 	/** Handler for clicking the history forward button */
 	FReply ForwardClicked();
 
+	/** Handler to check to see if a rename command is allowed */
+	bool CanRename() const;
+
+	/** Handler for Rename */
+	void OnRename();
+
+	/** Handler to check to see if a delete command is allowed */
+	bool CanDelete() const;
+
+	/** Handler for Delete */
+	void OnDelete();
+
 	/** Handler for opening assets or folders */
 	void OnOpenAssetsOrFolders();
 
@@ -288,7 +300,7 @@ private:
 	FText GetSearchAssetsHintText() const;
 
 	/** Delegate called when generating the context menu for a folder */
-	TSharedPtr<SWidget> GetFolderContextMenu(const TArray<FString>& SelectedPaths, FContentBrowserMenuExtender_SelectedPaths InMenuExtender, FOnCreateNewFolder OnCreateNewFolder);
+	TSharedPtr<SWidget> GetFolderContextMenu(const TArray<FString>& SelectedPaths, FContentBrowserMenuExtender_SelectedPaths InMenuExtender, FOnCreateNewFolder OnCreateNewFolder, bool bPathView);
 
 	/** Sets up an inline-name for the creation of a default-named folder the specified path */
 	void CreateNewFolder(FString FolderPath, FOnCreateNewFolder OnCreateNewFolder);
diff --git a/Engine/Source/Editor/ContentBrowser/Private/SPathView.cpp b/Engine/Source/Editor/ContentBrowser/Private/SPathView.cpp
--- a/Engine/Source/Editor/ContentBrowser/Private/SPathView.cpp
+++ b/Engine/Source/Editor/ContentBrowser/Private/SPathView.cpp
@@ -333,7 +333,7 @@ TSharedPtr<FTreeItem> SPathView::AddPath(const FString& Path, bool bUserNamed)
 			}
 			else
 			{
-				CurrentItem->bNewFolder = false;
+				CurrentItem->bNamingFolder = false;
 			}
 		}
 
@@ -386,6 +386,26 @@ bool SPathView::RemovePath(const FString& Path)
 	}
 }
 
+void SPathView::RenameFolder(const FString& FolderToRename)
+{
+	TArray<TSharedPtr<FTreeItem>> Items = TreeViewPtr->GetSelectedItems();
+	for (int32 ItemIdx = 0; ItemIdx < Items.Num(); ++ItemIdx)
+	{
+		TSharedPtr<FTreeItem>& Item = Items[ItemIdx];
+		if (Item.IsValid())
+		{
+			if (Item->FolderPath == FolderToRename)
+			{
+				Item->bNamingFolder = true;
+
+				TreeViewPtr->SetSelection(Item);
+				TreeViewPtr->RequestScrollIntoView(Item);
+				break;
+			}
+		}
+	}
+}
+
 void SPathView::SyncToAssets( const TArray<FAssetData>& AssetDataList, const bool bAllowImplicitSync )
 {
 	TArray<TSharedPtr<FTreeItem>> SyncTreeItems;
@@ -747,7 +767,7 @@ TSharedRef<ITableRow> SPathView::GenerateTreeRow( TSharedPtr<FTreeItem> TreeItem
 
 void SPathView::TreeItemScrolledIntoView( TSharedPtr<FTreeItem> TreeItem, const TSharedPtr<ITableRow>& Widget )
 {
-	if ( TreeItem->bNewFolder && Widget.IsValid() && Widget->GetContent().IsValid() )
+	if ( TreeItem->bNamingFolder && Widget.IsValid() && Widget->GetContent().IsValid() )
 	{
 		TreeItem->OnRenamedRequestEvent.Broadcast();
 	}
@@ -976,7 +996,9 @@ bool SPathView::VerifyFolderNameChanged(const FText& InName, FText& OutErrorMess
 	{
 		return false;
 	}
-	else if( ContentBrowserUtils::DoesFolderExist(InFolderPath) )
+
+	const FString NewPath = FPaths::GetPath(InFolderPath) / InName.ToString();
+	if (ContentBrowserUtils::DoesFolderExist(NewPath))
 	{
 		OutErrorMessage = LOCTEXT("RenameFolderAlreadyExists", "A folder already exists at this location with this name.");
 		return false;
@@ -985,7 +1007,7 @@ bool SPathView::VerifyFolderNameChanged(const FText& InName, FText& OutErrorMess
 	return true;
 }
 
-void SPathView::FolderNameChanged( const TSharedPtr< FTreeItem >& TreeItem, const FVector2D& MessageLocation )
+void SPathView::FolderNameChanged( const TSharedPtr< FTreeItem >& TreeItem, const FString& OldPath, const FVector2D& MessageLocation )
 {
 	// Verify the name of the folder
 	FText Reason;
@@ -1026,7 +1048,25 @@ void SPathView::FolderNameChanged( const TSharedPtr< FTreeItem >& TreeItem, cons
 
 		// Update the asset registry so this folder will persist
 		FAssetRegistryModule& AssetRegistryModule = FModuleManager::LoadModuleChecked<FAssetRegistryModule>("AssetRegistry");
-		AssetRegistryModule.Get().AddPath(TreeItem->FolderPath);
+		if (AssetRegistryModule.Get().AddPath(TreeItem->FolderPath) && TreeItem->FolderPath != OldPath)
+		{
+			// move any assets in our folder
+			TArray<FAssetData> AssetsInFolder;
+			AssetRegistryModule.Get().GetAssetsByPath(*OldPath, AssetsInFolder, true);
+			TArray<UObject*> ObjectsInFolder;
+			ContentBrowserUtils::GetObjectsInAssetData(AssetsInFolder, ObjectsInFolder);
+			ContentBrowserUtils::MoveAssets(ObjectsInFolder, TreeItem->FolderPath, OldPath);
+
+			// Now check to see if the original folder is empty, if so we can delete it
+			TArray<FAssetData> AssetsInOriginalFolder;
+			AssetRegistryModule.Get().GetAssetsByPath(*OldPath, AssetsInOriginalFolder, true);
+			if (AssetsInOriginalFolder.Num() == 0)
+			{
+				TArray<FString> FoldersToDelete;
+				FoldersToDelete.Add(OldPath);
+				ContentBrowserUtils::DeleteFolders(FoldersToDelete);
+			}
+		}
 	}
 	else
 	{
diff --git a/Engine/Source/Editor/ContentBrowser/Private/SPathView.h b/Engine/Source/Editor/ContentBrowser/Private/SPathView.h
--- a/Engine/Source/Editor/ContentBrowser/Private/SPathView.h
+++ b/Engine/Source/Editor/ContentBrowser/Private/SPathView.h
@@ -85,6 +85,9 @@ public:
 	/** Attempts to removed the folder at the end of the specified path from the tree. Returns true when successful. */
 	bool RemovePath(const FString& Path);
 
+	/** Sets up an inline rename for the specified folder */
+	void RenameFolder(const FString& FolderToRename);
+
 	/** Selects the paths containing the specified assets.
 	 *
 	 *	@param AssetDataList		- A list of assets to sync the view to
@@ -169,7 +172,7 @@ private:
 	bool VerifyFolderNameChanged(const FText& InName, FText& OutErrorMessage, const FString& InFolderPath) const;
 
 	/** Handler for when a name was given to a new folder */
-	void FolderNameChanged( const TSharedPtr< FTreeItem >& TreeItem, const FVector2D& MessageLocation );
+	void FolderNameChanged( const TSharedPtr< FTreeItem >& TreeItem, const FString& OldPath, const FVector2D& MessageLocation );
 
 	/** Returns true if the supplied folder item already exists in the tree. If so, ExistingItem will be set to the found item. */
 	bool FolderAlreadyExists(const TSharedPtr< FTreeItem >& TreeItem, TSharedPtr< FTreeItem >& ExistingItem);
diff --git a/Engine/Source/Editor/ContentBrowser/Private/SourcesViewWidgets.cpp b/Engine/Source/Editor/ContentBrowser/Private/SourcesViewWidgets.cpp
--- a/Engine/Source/Editor/ContentBrowser/Private/SourcesViewWidgets.cpp
+++ b/Engine/Source/Editor/ContentBrowser/Private/SourcesViewWidgets.cpp
@@ -279,10 +279,11 @@ void SAssetTreeItem::HandleNameCommitted( const FText& NewText, ETextCommit::Typ
 	{
 		TSharedPtr<FTreeItem> TreeItemPtr = TreeItem.Pin();
 
-		if ( TreeItemPtr->bNewFolder )
+		if ( TreeItemPtr->bNamingFolder )
 		{
-			TreeItemPtr->bNewFolder = false;
+			TreeItemPtr->bNamingFolder = false;
 
+			const FString OldPath = TreeItemPtr->FolderPath;
 			FString Path;
 			TreeItemPtr->FolderPath.Split(TEXT("/"), &Path, NULL, ESearchCase::CaseSensitive, ESearchDir::FromEnd);
 			TreeItemPtr->FolderName = NewText.ToString();
@@ -292,7 +293,7 @@ void SAssetTreeItem::HandleNameCommitted( const FText& NewText, ETextCommit::Typ
 			MessageLoc.X = LastGeometry.AbsolutePosition.X;
 			MessageLoc.Y = LastGeometry.AbsolutePosition.Y + LastGeometry.Size.Y * LastGeometry.Scale;
 
-			OnNameChanged.ExecuteIfBound(TreeItemPtr, MessageLoc);
+			OnNameChanged.ExecuteIfBound(TreeItemPtr, OldPath, MessageLoc);
 		}
 	}
 }
@@ -301,7 +302,7 @@ bool SAssetTreeItem::IsReadOnly() const
 {
 	if ( TreeItem.IsValid() )
 	{
-		return !TreeItem.Pin()->bNewFolder;
+		return !TreeItem.Pin()->bNamingFolder;
 	}
 	else
 	{
diff --git a/Engine/Source/Editor/ContentBrowser/Private/SourcesViewWidgets.h b/Engine/Source/Editor/ContentBrowser/Private/SourcesViewWidgets.h
--- a/Engine/Source/Editor/ContentBrowser/Private/SourcesViewWidgets.h
+++ b/Engine/Source/Editor/ContentBrowser/Private/SourcesViewWidgets.h
@@ -6,7 +6,7 @@
 class SAssetTreeItem : public SCompoundWidget
 {
 public:
-	DECLARE_DELEGATE_TwoParams( FOnNameChanged, const TSharedPtr<FTreeItem>& /*TreeItem*/, const FVector2D& /*MessageLocation*/);
+	DECLARE_DELEGATE_ThreeParams( FOnNameChanged, const TSharedPtr<FTreeItem>& /*TreeItem*/, const FString& /*OldPath*/, const FVector2D& /*MessageLocation*/);
 	DECLARE_DELEGATE_RetVal_ThreeParams( bool, FOnVerifyNameChanged, const FText& /*InName*/, FText& /*OutErrorMessage*/, const FString& /*FolderPath*/);
 	DECLARE_DELEGATE_TwoParams( FOnAssetsDragDropped, const TArray<FAssetData>& /*AssetList*/, const TSharedPtr<FTreeItem>& /*TreeItem*/);
 	DECLARE_DELEGATE_TwoParams( FOnPathsDragDropped, const TArray<FString>& /*PathNames*/, const TSharedPtr<FTreeItem>& /*TreeItem*/);
