Tutorials 2.0 - Initial version

NOTE: Old tutorials not deprecated (yet), but widget highlights in old tutorials will stop working with this change!

Added new Blueprintable UEditorTutorial object.
Added suite of widgets and details customizations to display tutorials.
New system is available on command line switch -NewTutorials.

Slate changes:
Tag names are now stored in SWidgets, rather than simply being discarded.
Removed STutorialWrapper in favour of using Tags.
Added Tags to more multibox widgets, so virtually all can now be picked.
Added SWindow::HasOverlay so we dont attempt to add overlays to widows that cannot have them.

diff --git a/Engine/Source/Developer/LogVisualizer/Private/SLogVisualizer.cpp b/Engine/Source/Developer/LogVisualizer/Private/SLogVisualizer.cpp
--- a/Engine/Source/Developer/LogVisualizer/Private/SLogVisualizer.cpp
+++ b/Engine/Source/Developer/LogVisualizer/Private/SLogVisualizer.cpp
@@ -357,13 +357,11 @@ void SLogVisualizer::Construct(const FArguments& InArgs, FLogVisualizer* InLogVi
 			+ SVerticalBox::Slot()
 			.AutoHeight()
 			[
-				SNew(STutorialWrapper, TEXT("CategoryFilters"))
-				[
-					SAssignNew(FilterListPtr, SLogFilterList)
-					.OnFilterChanged(this, &SLogVisualizer::OnLogCategoryFiltersChanged)
-					/*.OnGetContextMenu(this, &SLogVisualizer::GetFilterContextMenu)*/
-					/*.FrontendFilters(FrontendFilters)*/
-				]
+				SAssignNew(FilterListPtr, SLogFilterList)
+				.OnFilterChanged(this, &SLogVisualizer::OnLogCategoryFiltersChanged)
+				.Tag(TEXT("CategoryFilters"))
+				/*.OnGetContextMenu(this, &SLogVisualizer::GetFilterContextMenu)*/
+				/*.FrontendFilters(FrontendFilters)*/
 			]
 
 			+SVerticalBox::Slot()
diff --git a/Engine/Source/Editor/ContentBrowser/Private/AssetViewWidgets.cpp b/Engine/Source/Editor/ContentBrowser/Private/AssetViewWidgets.cpp
--- a/Engine/Source/Editor/ContentBrowser/Private/AssetViewWidgets.cpp
+++ b/Engine/Source/Editor/ContentBrowser/Private/AssetViewWidgets.cpp
@@ -962,6 +962,7 @@ void SAssetListItem::Construct( const FArguments& InArgs )
 		SNew(SBorder)
 		.BorderImage(this, &SAssetViewItem::GetBorderImage)
 		.Padding(0)
+		.Tag(AssetItem->GetType() == EAssetItemType::Normal ? StaticCastSharedPtr<FAssetViewAsset>(AssetItem)->Data.ObjectPath : NAME_None)
 		[
 			SNew(SHorizontalBox)
 
@@ -1173,6 +1174,7 @@ void SAssetTileItem::Construct( const FArguments& InArgs )
 		SNew(SBorder)
 		.BorderImage(this, &SAssetViewItem::GetBorderImage)
 		.Padding(0)
+		.Tag(AssetItem->GetType() == EAssetItemType::Normal ? StaticCastSharedPtr<FAssetViewAsset>(AssetItem)->Data.ObjectPath : NAME_None)
 		[
 			SNew(SVerticalBox)
 
@@ -1364,6 +1366,7 @@ TSharedRef<SWidget> SAssetColumnItem::GenerateWidgetForColumn( const FName& Colu
 		const float IconOverlaySize = IconBrush->ImageSize.X * 0.6f;
 
 		Content = SNew(SHorizontalBox)
+			.Tag(AssetItem->GetType() == EAssetItemType::Normal ? StaticCastSharedPtr<FAssetViewAsset>(AssetItem)->Data.ObjectPath : NAME_None)
 			// Icon
 			+SHorizontalBox::Slot()
 			.AutoWidth()
diff --git a/Engine/Source/Editor/ContentBrowser/Private/ContentBrowserSingleton.cpp b/Engine/Source/Editor/ContentBrowser/Private/ContentBrowserSingleton.cpp
--- a/Engine/Source/Editor/ContentBrowser/Private/ContentBrowserSingleton.cpp
+++ b/Engine/Source/Editor/ContentBrowser/Private/ContentBrowserSingleton.cpp
@@ -5,7 +5,6 @@
 #include "ContentBrowserModule.h"
 #include "ObjectTools.h"
 #include "WorkspaceMenuStructureModule.h"
-#include "STutorialWrapper.h"
 #include "IDocumentation.h"
 
 #define LOCTEXT_NAMESPACE "ContentBrowser"
@@ -362,8 +361,9 @@ TSharedRef<SDockTab> FContentBrowserSingleton::SpawnContentBrowserTab( const FSp
 	}
 
 	// Add wrapper for tutorial highlighting
-	TSharedRef<STutorialWrapper> Wrapper = 
-		SNew( STutorialWrapper, TEXT("ContentBrowser") )
+	TSharedRef<SBox> Wrapper = 
+		SNew( SBox )
+		.Tag(TEXT("ContentBrowser"))
 		[
 			NewBrowser
 		];
diff --git a/Engine/Source/Editor/ContentBrowser/Private/SAssetPicker.cpp b/Engine/Source/Editor/ContentBrowser/Private/SAssetPicker.cpp
--- a/Engine/Source/Editor/ContentBrowser/Private/SAssetPicker.cpp
+++ b/Engine/Source/Editor/ContentBrowser/Private/SAssetPicker.cpp
@@ -69,20 +69,18 @@ void SAssetPicker::Construct( const FArguments& InArgs )
 			HorizontalBox->AddSlot()
 			.AutoWidth()
 			[
-				SNew( STutorialWrapper, TEXT("ContentBrowserFiltersCombo") )
+				SNew( SComboButton )
+				.ComboButtonStyle( FEditorStyle::Get(), "ContentBrowser.Filters.Style" )
+				.ToolTipText( LOCTEXT( "AddFilterToolTip", "Add an asset filter." ) )
+				.OnGetMenuContent( this, &SAssetPicker::MakeAddFilterMenu )
+				.HasDownArrow( true )
+				.ContentPadding( FMargin( 1, 0 ) )
+				.Tag(TEXT("ContentBrowserFiltersCombo"))
+				.ButtonContent()
 				[
-					SNew( SComboButton )
-					.ComboButtonStyle( FEditorStyle::Get(), "ContentBrowser.Filters.Style" )
-					.ToolTipText( LOCTEXT( "AddFilterToolTip", "Add an asset filter." ) )
-					.OnGetMenuContent( this, &SAssetPicker::MakeAddFilterMenu )
-					.HasDownArrow( true )
-					.ContentPadding( FMargin( 1, 0 ) )
-					.ButtonContent()
-					[
-						SNew( STextBlock )
-						.TextStyle( FEditorStyle::Get(), "ContentBrowser.Filters.Text" )
-						.Text( LOCTEXT( "Filters", "Filters" ) )
-					]
+					SNew( STextBlock )
+					.TextStyle( FEditorStyle::Get(), "ContentBrowser.Filters.Text" )
+					.Text( LOCTEXT( "Filters", "Filters" ) )
 				]
 			];
 		}
diff --git a/Engine/Source/Editor/ContentBrowser/Private/SContentBrowser.cpp b/Engine/Source/Editor/ContentBrowser/Private/SContentBrowser.cpp
--- a/Engine/Source/Editor/ContentBrowser/Private/SContentBrowser.cpp
+++ b/Engine/Source/Editor/ContentBrowser/Private/SContentBrowser.cpp
@@ -6,7 +6,6 @@
 #include "NewAssetContextMenu.h"
 #include "PathContextMenu.h"
 #include "ContentBrowserModule.h"
-#include "STutorialWrapper.h"
 #include "ContentBrowserCommands.h"
 #include "CollectionManagerModule.h"
 #include "AssetRegistryModule.h"
@@ -101,38 +100,36 @@ void SContentBrowser::Construct( const FArguments& InArgs, const FName& InInstan
 						.HAlign( HAlign_Left )
 						.Padding(0,0,4,0)
 						[
-							SNew( STutorialWrapper, TEXT("ContentBrowserNewAsset") )
+							SNew( SComboButton )
+							.ComboButtonStyle( FEditorStyle::Get(), "ContentBrowser.NewAsset.Style" )
+							.ForegroundColor(FLinearColor::White)
+							.ContentPadding(0)
+							.OnGetMenuContent( this, &SContentBrowser::MakeCreateAssetContextMenu )
+							.ToolTipText( this, &SContentBrowser::GetNewAssetToolTipText )
+							.IsEnabled( this, &SContentBrowser::IsAssetPathSelected )
+							.Tag(TEXT("ContentBrowserNewAsset"))
+							.ButtonContent()
 							[
-								SNew( SComboButton )
-								.ComboButtonStyle( FEditorStyle::Get(), "ContentBrowser.NewAsset.Style" )
-								.ForegroundColor(FLinearColor::White)
-								.ContentPadding(0)
-								.OnGetMenuContent( this, &SContentBrowser::MakeCreateAssetContextMenu )
-								.ToolTipText( this, &SContentBrowser::GetNewAssetToolTipText )
-								.IsEnabled( this, &SContentBrowser::IsAssetPathSelected )
-								.ButtonContent()
-								[
-									SNew( SHorizontalBox )
+								SNew( SHorizontalBox )
 
-									// New Icon
-									+ SHorizontalBox::Slot()
-									.AutoWidth()
-									.VAlign( VAlign_Center )
-									[
-										SNew( SImage )
-										.Image( FEditorStyle::GetBrush( "ContentBrowser.NewAsset" ) )
-									]
+								// New Icon
+								+ SHorizontalBox::Slot()
+								.AutoWidth()
+								.VAlign( VAlign_Center )
+								[
+									SNew( SImage )
+									.Image( FEditorStyle::GetBrush( "ContentBrowser.NewAsset" ) )
+								]
 
-									// New Text
-									+ SHorizontalBox::Slot()
-									.AutoWidth()
-									.VAlign(VAlign_Center)
-									.Padding(0,0,2,0)
-									[
-										SNew( STextBlock )
-										.TextStyle( FEditorStyle::Get(), "ContentBrowser.TopBar.Font" )
-										.Text( LOCTEXT( "NewButton", "New" ) )
-									]
+								// New Text
+								+ SHorizontalBox::Slot()
+								.AutoWidth()
+								.VAlign(VAlign_Center)
+								.Padding(0,0,2,0)
+								[
+									SNew( STextBlock )
+									.TextStyle( FEditorStyle::Get(), "ContentBrowser.TopBar.Font" )
+									.Text( LOCTEXT( "NewButton", "New" ) )
 								]
 							]
 						]
@@ -144,36 +141,34 @@ void SContentBrowser::Construct( const FArguments& InArgs, const FName& InInstan
 						.HAlign( HAlign_Left )
 						.Padding(0,0,10,0)
 						[
-							SNew( STutorialWrapper, TEXT("ContentBrowserImportAsset") )
+							SNew( SButton )
+							.ButtonStyle( FEditorStyle::Get(), "ToggleButton" )
+							.ToolTipText( this, &SContentBrowser::GetImportTooltipText )
+							.IsEnabled( this, &SContentBrowser::IsAssetPathSelected )
+							.OnClicked( this, &SContentBrowser::HandleImportClicked )
+							.ContentPadding( 0 )
+							.Tag(TEXT("ContentBrowserImportAsset"))
 							[
-								SNew( SButton )
-								.ButtonStyle( FEditorStyle::Get(), "ToggleButton" )
-								.ToolTipText( this, &SContentBrowser::GetImportTooltipText )
-								.IsEnabled( this, &SContentBrowser::IsAssetPathSelected )
-								.OnClicked( this, &SContentBrowser::HandleImportClicked )
-								.ContentPadding( 0 )
-								[
-									SNew( SHorizontalBox )
+								SNew( SHorizontalBox )
 
-									// Import Icon
-									+ SHorizontalBox::Slot()
-									.AutoWidth()
-									.VAlign( VAlign_Center )
-									[
-										SNew( SImage )
-										.Image( FEditorStyle::GetBrush( "ContentBrowser.ImportPackage" ) )
-									]
+								// Import Icon
+								+ SHorizontalBox::Slot()
+								.AutoWidth()
+								.VAlign( VAlign_Center )
+								[
+									SNew( SImage )
+									.Image( FEditorStyle::GetBrush( "ContentBrowser.ImportPackage" ) )
+								]
 
-									// Import Text
-									+ SHorizontalBox::Slot()
-									.AutoWidth()
-									.VAlign(VAlign_Center)
-									.Padding(0,0,2,0)
-									[
-										SNew( STextBlock )
-										.TextStyle( FEditorStyle::Get(), "ContentBrowser.TopBar.Font" )
-										.Text( LOCTEXT( "Import", "Import" ) )
-									]
+								// Import Text
+								+ SHorizontalBox::Slot()
+								.AutoWidth()
+								.VAlign(VAlign_Center)
+								.Padding(0,0,2,0)
+								[
+									SNew( STextBlock )
+									.TextStyle( FEditorStyle::Get(), "ContentBrowser.TopBar.Font" )
+									.Text( LOCTEXT( "Import", "Import" ) )
 								]
 							]
 						]
@@ -184,17 +179,15 @@ void SContentBrowser::Construct( const FArguments& InArgs, const FName& InInstan
 						.VAlign(VAlign_Center)
 						.HAlign(HAlign_Left)
 						[
-							SNew( STutorialWrapper, TEXT("ContentBrowserSaveDirtyPackages") )
+							SNew( SButton )
+							.ButtonStyle( FEditorStyle::Get(), "ToggleButton" )
+							.ToolTipText( LOCTEXT( "SaveDirtyPackagesTooltip", "Save all modified assets." ) )
+							.ContentPadding( 0 )
+							.OnClicked( this, &SContentBrowser::OnSaveClicked )
+							.Tag(TEXT("ContentBrowserSaveDirtyPackages"))
 							[
-								SNew( SButton )
-								.ButtonStyle( FEditorStyle::Get(), "ToggleButton" )
-								.ToolTipText( LOCTEXT( "SaveDirtyPackagesTooltip", "Save all modified assets." ) )
-								.ContentPadding( 0 )
-								.OnClicked( this, &SContentBrowser::OnSaveClicked )
-								[
-									SNew( SImage )
-									.Image( FEditorStyle::GetBrush( "ContentBrowser.SaveDirtyPackages" ) )
-								]
+								SNew( SImage )
+								.Image( FEditorStyle::GetBrush( "ContentBrowser.SaveDirtyPackages" ) )
 							]
 						]
 					]
@@ -220,20 +213,18 @@ void SContentBrowser::Construct( const FArguments& InArgs, const FName& InInstan
 						+ SVerticalBox::Slot()
 						.FillHeight(1.0f)
 						[
-							SNew( STutorialWrapper, TEXT("ContentBrowserHistoryBack") )
+							SNew(SButton)
+							.VAlign(EVerticalAlignment::VAlign_Center)
+							.ButtonStyle( FEditorStyle::Get(), "ToggleButton" )
+							.ForegroundColor( FEditorStyle::GetSlateColor("DefaultForeground") )
+							.ToolTipText( this, &SContentBrowser::GetHistoryBackTooltip )
+							.ContentPadding( FMargin(1, 0) )
+							.OnClicked(this, &SContentBrowser::BackClicked)
+							.IsEnabled(this, &SContentBrowser::IsBackEnabled)
+							.Tag(TEXT("ContentBrowserHistoryBack"))
 							[
-								SNew(SButton)
-								.VAlign(EVerticalAlignment::VAlign_Center)
-								.ButtonStyle( FEditorStyle::Get(), "ToggleButton" )
-								.ForegroundColor( FEditorStyle::GetSlateColor("DefaultForeground") )
-								.ToolTipText( this, &SContentBrowser::GetHistoryBackTooltip )
-								.ContentPadding( FMargin(1, 0) )
-								.OnClicked(this, &SContentBrowser::BackClicked)
-								.IsEnabled(this, &SContentBrowser::IsBackEnabled)
-								[
-									SNew(SImage)
-									.Image(FEditorStyle::GetBrush("ContentBrowser.HistoryBack"))
-								]
+								SNew(SImage)
+								.Image(FEditorStyle::GetBrush("ContentBrowser.HistoryBack"))
 							]
 						]
 					]
@@ -248,20 +239,18 @@ void SContentBrowser::Construct( const FArguments& InArgs, const FName& InInstan
 						+ SVerticalBox::Slot()
 						.FillHeight(1.0f)
 						[
-							SNew(STutorialWrapper, "ContentBrowserHistoryForward")
+							SNew(SButton)
+							.VAlign(EVerticalAlignment::VAlign_Center)
+							.ButtonStyle( FEditorStyle::Get(), "ToggleButton" )
+							.ForegroundColor( FEditorStyle::GetSlateColor("DefaultForeground") )
+							.ToolTipText( this, &SContentBrowser::GetHistoryForwardTooltip )
+							.ContentPadding( FMargin(1, 0) )
+							.OnClicked(this, &SContentBrowser::ForwardClicked)
+							.IsEnabled(this, &SContentBrowser::IsForwardEnabled)
+							.Tag(TEXT("ContentBrowserHistoryForward"))
 							[
-								SNew(SButton)
-								.VAlign(EVerticalAlignment::VAlign_Center)
-								.ButtonStyle( FEditorStyle::Get(), "ToggleButton" )
-								.ForegroundColor( FEditorStyle::GetSlateColor("DefaultForeground") )
-								.ToolTipText( this, &SContentBrowser::GetHistoryForwardTooltip )
-								.ContentPadding( FMargin(1, 0) )
-								.OnClicked(this, &SContentBrowser::ForwardClicked)
-								.IsEnabled(this, &SContentBrowser::IsForwardEnabled)
-								[
-									SNew(SImage)
-									.Image(FEditorStyle::GetBrush("ContentBrowser.HistoryForward"))
-								]
+								SNew(SImage)
+								.Image(FEditorStyle::GetBrush("ContentBrowser.HistoryForward"))
 							]
 						]
 					]
@@ -280,19 +269,17 @@ void SContentBrowser::Construct( const FArguments& InArgs, const FName& InInstan
 					.AutoWidth()
 					.VAlign( VAlign_Fill )
 					[
-						SNew( STutorialWrapper, TEXT("ContentBrowserPathPicker") )
+						SAssignNew( PathPickerButton, SComboButton )
+						.ComboButtonStyle( FEditorStyle::Get(), "ToolbarComboButton" )
+						.ForegroundColor(FLinearColor::White)
+						.ToolTipText( LOCTEXT( "PathPickerTooltip", "Choose a path" ) )
+						.OnGetMenuContent( this, &SContentBrowser::GetPathPickerContent )
+						.HasDownArrow( false )
+						.Tag(TEXT("ContentBrowserPathPicker"))
+						.ButtonContent()
 						[
-							SAssignNew( PathPickerButton, SComboButton )
-							.ComboButtonStyle( FEditorStyle::Get(), "ToolbarComboButton" )
-							.ForegroundColor(FLinearColor::White)
-							.ToolTipText( LOCTEXT( "PathPickerTooltip", "Choose a path" ) )
-							.OnGetMenuContent( this, &SContentBrowser::GetPathPickerContent )
-							.HasDownArrow( false )
-							.ButtonContent()
-							[
-								SNew( SImage )
-								.Image( FEditorStyle::GetBrush( "ContentBrowser.Sources" ) )
-							]
+							SNew( SImage )
+							.Image( FEditorStyle::GetBrush( "ContentBrowser.Sources" ) )
 						]
 					]
 
@@ -308,18 +295,16 @@ void SContentBrowser::Construct( const FArguments& InArgs, const FName& InInstan
 						.VAlign(VAlign_Center)
 						.FillWidth(1.0f)
 						[
-							SNew( STutorialWrapper, TEXT("ContentBrowserPath") )
-							[
-								SAssignNew(PathBreadcrumbTrail, SBreadcrumbTrail<FString>)
-								//.ToolTipText( LOCTEXT("PathTooltip", "Content Path") )
-								.ButtonContentPadding(FMargin(3.0f, 3.0f))
-								.DelimiterImage(FEditorStyle::GetBrush("ContentBrowser.PathDelimiter"))
-								.TextStyle(FEditorStyle::Get(), "ContentBrowser.PathText")
-								.ShowLeadingDelimiter( false )
-								.InvertTextColorOnHover( false )
-								.OnCrumbClicked(this, &SContentBrowser::OnPathClicked)
-								.GetCrumbMenuContent(this, &SContentBrowser::OnGetCrumbDelimiterContent)
-							]
+							SAssignNew(PathBreadcrumbTrail, SBreadcrumbTrail<FString>)
+							//.ToolTipText( LOCTEXT("PathTooltip", "Content Path") )
+							.ButtonContentPadding(FMargin(3.0f, 3.0f))
+							.DelimiterImage(FEditorStyle::GetBrush("ContentBrowser.PathDelimiter"))
+							.TextStyle(FEditorStyle::Get(), "ContentBrowser.PathText")
+							.ShowLeadingDelimiter( false )
+							.InvertTextColorOnHover( false )
+							.OnCrumbClicked(this, &SContentBrowser::OnPathClicked)
+							.GetCrumbMenuContent(this, &SContentBrowser::OnGetCrumbDelimiterContent)
+							.Tag(TEXT("ContentBrowserPath"))
 						]
 					]
 
@@ -333,18 +318,16 @@ void SContentBrowser::Construct( const FArguments& InArgs, const FName& InInstan
 						+ SVerticalBox::Slot()
 						.FillHeight(1.0f)
 						[
-							SNew( STutorialWrapper, TEXT("ContentBrowserLock") )
+							SNew(SButton)
+							.VAlign(EVerticalAlignment::VAlign_Center)
+							.ButtonStyle( FEditorStyle::Get(), "NoBorder" )
+							.ToolTipText( LOCTEXT("LockToggleTooltip", "Toggle lock. If locked, this browser will ignore Find in Content Browser requests.") )
+							.ContentPadding( FMargin(1, 0) )
+							.OnClicked(this, &SContentBrowser::ToggleLockClicked)
+							.Tag(TEXT("ContentBrowserLock"))
 							[
-								SNew(SButton)
-								.VAlign(EVerticalAlignment::VAlign_Center)
-								.ButtonStyle( FEditorStyle::Get(), "NoBorder" )
-								.ToolTipText( LOCTEXT("LockToggleTooltip", "Toggle lock. If locked, this browser will ignore Find in Content Browser requests.") )
-								.ContentPadding( FMargin(1, 0) )
-								.OnClicked(this, &SContentBrowser::ToggleLockClicked)
-								[
-									SNew(SImage)
-									.Image( this, &SContentBrowser::GetToggleLockImage)
-								]
+								SNew(SImage)
+								.Image( this, &SContentBrowser::GetToggleLockImage)
 							]
 						]
 					]
@@ -381,37 +364,32 @@ void SContentBrowser::Construct( const FArguments& InArgs, const FName& InInstan
 						.Padding(FMargin(3))
 						.BorderImage( FEditorStyle::GetBrush("ToolPanel.GroupBorder") )
 						[
-							SNew( STutorialWrapper, TEXT("ContentBrowserSources") )
+							SAssignNew( PathViewPtr, SPathView )
+							.OnPathSelected( this, &SContentBrowser::PathSelected )
+							.OnGetFolderContextMenu( this, &SContentBrowser::GetFolderContextMenu, true )
+							.OnGetPathContextMenuExtender( this, &SContentBrowser::GetPathContextMenuExtender )
+							.FocusSearchBoxWhenOpened( false )
+							.ShowTreeTitle( false )
+							.ShowSeparator( false )
+							.Tag(TEXT("ContentBrowserSources"))
+							.SearchContent()
 							[
-								SAssignNew( PathViewPtr, SPathView )
-								.OnPathSelected( this, &SContentBrowser::PathSelected )
-								.OnGetFolderContextMenu( this, &SContentBrowser::GetFolderContextMenu, true )
-								.OnGetPathContextMenuExtender( this, &SContentBrowser::GetPathContextMenuExtender )
-								.FocusSearchBoxWhenOpened( false )
-								.ShowTreeTitle( false )
-								.ShowSeparator( false )
-								.SearchContent()
+								SNew( SVerticalBox )
+								.Tag(TEXT("ContentBrowserSourcesToggle"))
+								+ SVerticalBox::Slot()
+								.FillHeight( 1.0f )
+								.Padding(0,0,2,0)
 								[
-									SNew( STutorialWrapper, TEXT("ContentBrowserSourcesToggle") )
+									SNew( SButton )
+									.VAlign( EVerticalAlignment::VAlign_Center )
+									.ButtonStyle( FEditorStyle::Get(), "ToggleButton" )
+									.ToolTipText( LOCTEXT( "SourcesTreeToggleTooltip", "Show or hide the sources panel" ) )
+									.ContentPadding( FMargin( 1, 0 ) )
+									.ForegroundColor( FEditorStyle::GetSlateColor( "DefaultForeground" ) )
+									.OnClicked( this, &SContentBrowser::SourcesViewExpandClicked )
 									[
-										SNew( SVerticalBox )
-
-										+ SVerticalBox::Slot()
-										.FillHeight( 1.0f )
-										.Padding(0,0,2,0)
-										[
-											SNew( SButton )
-											.VAlign( EVerticalAlignment::VAlign_Center )
-											.ButtonStyle( FEditorStyle::Get(), "ToggleButton" )
-											.ToolTipText( LOCTEXT( "SourcesTreeToggleTooltip", "Show or hide the sources panel" ) )
-											.ContentPadding( FMargin( 1, 0 ) )
-											.ForegroundColor( FEditorStyle::GetSlateColor( "DefaultForeground" ) )
-											.OnClicked( this, &SContentBrowser::SourcesViewExpandClicked )
-											[
-												SNew( SImage )
-												.Image( this, &SContentBrowser::GetSourcesToggleImage )
-											]
-										]
+										SNew( SImage )
+										.Image( this, &SContentBrowser::GetSourcesToggleImage )
 									]
 								]
 							]
@@ -426,11 +404,9 @@ void SContentBrowser::Construct( const FArguments& InArgs, const FName& InInstan
 						.Padding(FMargin(3))
 						.BorderImage( FEditorStyle::GetBrush("ToolPanel.GroupBorder") )
 						[
-							SNew( STutorialWrapper, TEXT("ContentBrowserCollections") )
-							[
-								SAssignNew(CollectionViewPtr, SCollectionView)
-								.OnCollectionSelected(this, &SContentBrowser::CollectionSelected)
-							]
+							SAssignNew(CollectionViewPtr, SCollectionView)
+							.OnCollectionSelected(this, &SContentBrowser::CollectionSelected)
+							.Tag(TEXT("ContentBrowserCollections"))
 						]
 					]
 				]
@@ -457,25 +433,22 @@ void SContentBrowser::Construct( const FArguments& InArgs, const FName& InInstan
 						.AutoWidth()
 						.Padding( 0, 0, 4, 0 )
 						[
-							SNew( STutorialWrapper, TEXT("ContentBrowserSourcesToggle") )
+							SNew( SVerticalBox )
+							.Tag(TEXT("ContentBrowserSourcesToggle"))
+							+ SVerticalBox::Slot()
+							.FillHeight( 1.0f )
 							[
-								SNew( SVerticalBox )
-
-								+ SVerticalBox::Slot()
-								.FillHeight( 1.0f )
+								SNew( SButton )
+								.VAlign( EVerticalAlignment::VAlign_Center )
+								.ButtonStyle( FEditorStyle::Get(), "ToggleButton" )
+								.ToolTipText( LOCTEXT( "SourcesTreeToggleTooltip", "Show or hide the sources panel" ) )
+								.ContentPadding( FMargin( 1, 0 ) )
+								.ForegroundColor( FEditorStyle::GetSlateColor( "DefaultForeground" ) )
+								.OnClicked( this, &SContentBrowser::SourcesViewExpandClicked )
+								.Visibility( this, &SContentBrowser::GetPathExpanderVisibility )
 								[
-									SNew( SButton )
-									.VAlign( EVerticalAlignment::VAlign_Center )
-									.ButtonStyle( FEditorStyle::Get(), "ToggleButton" )
-									.ToolTipText( LOCTEXT( "SourcesTreeToggleTooltip", "Show or hide the sources panel" ) )
-									.ContentPadding( FMargin( 1, 0 ) )
-									.ForegroundColor( FEditorStyle::GetSlateColor( "DefaultForeground" ) )
-									.OnClicked( this, &SContentBrowser::SourcesViewExpandClicked )
-									.Visibility( this, &SContentBrowser::GetPathExpanderVisibility )
-									[
-										SNew( SImage )
-										.Image( this, &SContentBrowser::GetSourcesToggleImage )
-									]
+									SNew( SImage )
+									.Image( this, &SContentBrowser::GetSourcesToggleImage )
 								]
 							]
 						]
@@ -484,22 +457,20 @@ void SContentBrowser::Construct( const FArguments& InArgs, const FName& InInstan
 						+ SHorizontalBox::Slot()
 						.AutoWidth()
 						[
-							SNew( STutorialWrapper, TEXT("ContentBrowserFiltersCombo") )
+							SNew( SComboButton )
+							.ComboButtonStyle( FEditorStyle::Get(), "ContentBrowser.Filters.Style" )
+							.ForegroundColor(FLinearColor::White)
+							.ContentPadding(0)
+							.ToolTipText( LOCTEXT( "AddFilterToolTip", "Add an asset filter." ) )
+							.OnGetMenuContent( this, &SContentBrowser::MakeAddFilterMenu )
+							.HasDownArrow( true )
+							.ContentPadding( FMargin( 1, 0 ) )
+							.Tag(TEXT("ContentBrowserFiltersCombo"))
+							.ButtonContent()
 							[
-								SNew( SComboButton )
-								.ComboButtonStyle( FEditorStyle::Get(), "ContentBrowser.Filters.Style" )
-								.ForegroundColor(FLinearColor::White)
-								.ContentPadding(0)
-								.ToolTipText( LOCTEXT( "AddFilterToolTip", "Add an asset filter." ) )
-								.OnGetMenuContent( this, &SContentBrowser::MakeAddFilterMenu )
-								.HasDownArrow( true )
-								.ContentPadding( FMargin( 1, 0 ) )
-								.ButtonContent()
-								[
-									SNew( STextBlock )
-									.TextStyle( FEditorStyle::Get(), "ContentBrowser.Filters.Text" )
-									.Text( LOCTEXT( "Filters", "Filters" ) )
-								]
+								SNew( STextBlock )
+								.TextStyle( FEditorStyle::Get(), "ContentBrowser.Filters.Text" )
+								.Text( LOCTEXT( "Filters", "Filters" ) )
 							]
 						]
 
@@ -509,15 +480,13 @@ void SContentBrowser::Construct( const FArguments& InArgs, const FName& InInstan
 						.VAlign(VAlign_Center)
 						.FillWidth(1.0f)
 						[
-							SNew( STutorialWrapper, TEXT("ContentBrowserSearchAssets") )
-							[
-								SAssignNew(SearchBoxPtr, SAssetSearchBox)
-								.HintText( this, &SContentBrowser::GetSearchAssetsHintText )
-								.OnTextChanged( this, &SContentBrowser::OnSearchBoxChanged )
-								.OnTextCommitted( this, &SContentBrowser::OnSearchBoxCommitted )
-								.PossibleSuggestions( this, &SContentBrowser::GetAssetSearchSuggestions )
-								.DelayChangeNotificationsWhileTyping( true )
-							]
+							SAssignNew(SearchBoxPtr, SAssetSearchBox)
+							.HintText( this, &SContentBrowser::GetSearchAssetsHintText )
+							.OnTextChanged( this, &SContentBrowser::OnSearchBoxChanged )
+							.OnTextCommitted( this, &SContentBrowser::OnSearchBoxCommitted )
+							.PossibleSuggestions( this, &SContentBrowser::GetAssetSearchSuggestions )
+							.DelayChangeNotificationsWhileTyping( true )
+							.Tag(TEXT("ContentBrowserSearchAssets"))
 						]
 					]
 
@@ -525,13 +494,11 @@ void SContentBrowser::Construct( const FArguments& InArgs, const FName& InInstan
 					+ SVerticalBox::Slot()
 					.AutoHeight()
 					[
-						SNew( STutorialWrapper, TEXT("ContentBrowserFilters") )
-						[
-							SAssignNew(FilterListPtr, SFilterList)
-							.OnFilterChanged(this, &SContentBrowser::OnFilterChanged)
-							.OnGetContextMenu(this, &SContentBrowser::GetFilterContextMenu)
-							.FrontendFilters(FrontendFilters)
-						]
+						SAssignNew(FilterListPtr, SFilterList)
+						.OnFilterChanged(this, &SContentBrowser::OnFilterChanged)
+						.OnGetContextMenu(this, &SContentBrowser::GetFilterContextMenu)
+						.FrontendFilters(FrontendFilters)
+						.Tag(TEXT("ContentBrowserFilters"))
 					]
 
 					// Assets
@@ -539,28 +506,26 @@ void SContentBrowser::Construct( const FArguments& InArgs, const FName& InInstan
 					.FillHeight( 1.0f )
 					.Padding( 0 )
 					[
-						SNew( STutorialWrapper, TEXT("ContentBrowserAssets") )
-						[
-							SAssignNew(AssetViewPtr, SAssetView)
-							.ThumbnailScale( 0.0f )
-							.OnPathSelected(this, &SContentBrowser::FolderEntered)
-							.OnAssetSelected(this, &SContentBrowser::OnAssetSelectionChanged)
-							.OnAssetsActivated(this, &SContentBrowser::OnAssetsActivated)
-							.OnGetAssetContextMenu(this, &SContentBrowser::OnGetAssetContextMenu)
-							.OnGetFolderContextMenu(this, &SContentBrowser::GetFolderContextMenu, false)
-							.OnGetPathContextMenuExtender(this, &SContentBrowser::GetPathContextMenuExtender)
-							.OnFindInAssetTreeRequested(this, &SContentBrowser::OnFindInAssetTreeRequested)
-							.OnAssetRenameCommitted(this, &SContentBrowser::OnAssetRenameCommitted)
-							.AreRealTimeThumbnailsAllowed(this, &SContentBrowser::IsHovered)
-							.FrontendFilters(FrontendFilters)
-							.HighlightedText(this, &SContentBrowser::GetHighlightedText)
-							.AllowThumbnailEditMode(true)
-							.AllowThumbnailHintLabel(false)
-							.CanShowFolders(true)
-							.CanShowOnlyAssetsInSelectedFolders(true)
-							.CanShowRealTimeThumbnails(true)
-							.CanShowDevelopersFolder(true)
-						]
+						SAssignNew(AssetViewPtr, SAssetView)
+						.ThumbnailScale( 0.0f )
+						.OnPathSelected(this, &SContentBrowser::FolderEntered)
+						.OnAssetSelected(this, &SContentBrowser::OnAssetSelectionChanged)
+						.OnAssetsActivated(this, &SContentBrowser::OnAssetsActivated)
+						.OnGetAssetContextMenu(this, &SContentBrowser::OnGetAssetContextMenu)
+						.OnGetFolderContextMenu(this, &SContentBrowser::GetFolderContextMenu, false)
+						.OnGetPathContextMenuExtender(this, &SContentBrowser::GetPathContextMenuExtender)
+						.OnFindInAssetTreeRequested(this, &SContentBrowser::OnFindInAssetTreeRequested)
+						.OnAssetRenameCommitted(this, &SContentBrowser::OnAssetRenameCommitted)
+						.AreRealTimeThumbnailsAllowed(this, &SContentBrowser::IsHovered)
+						.FrontendFilters(FrontendFilters)
+						.HighlightedText(this, &SContentBrowser::GetHighlightedText)
+						.AllowThumbnailEditMode(true)
+						.AllowThumbnailHintLabel(false)
+						.CanShowFolders(true)
+						.CanShowOnlyAssetsInSelectedFolders(true)
+						.CanShowRealTimeThumbnails(true)
+						.CanShowDevelopersFolder(true)
+						.Tag(TEXT("ContentBrowserAssets"))
 					]
 				]
 			]
diff --git a/Engine/Source/Editor/EditorStyle/Private/SlateEditorStyle.cpp b/Engine/Source/Editor/EditorStyle/Private/SlateEditorStyle.cpp
--- a/Engine/Source/Editor/EditorStyle/Private/SlateEditorStyle.cpp
+++ b/Engine/Source/Editor/EditorStyle/Private/SlateEditorStyle.cpp
@@ -2358,6 +2358,33 @@ void FSlateEditorStyle::FStyle::SetupTutorialStyles()
 		const FTextBlockStyle TutorialHeaderText = FTextBlockStyle(DocumentationHeaderText)
 			.SetColorAndOpacity(FLinearColor::Black);
 
+		Set( "Tutorials.Border", new BOX_BRUSH( "Tutorials/OverlayFrame", FMargin(18.0f/64.0f), FLinearColor(1.0f, 1.0f, 1.0f, 0.8f) ) );
+		Set( "Tutorials.Browser.WelcomeHeader", FTextBlockStyle(TutorialHeaderText)
+			.SetFontSize(20));
+
+		Set( "Tutorials.Browser.SummaryHeader", FTextBlockStyle(TutorialHeaderText)
+			.SetFontSize(16));
+
+		Set( "Tutorials.Browser.SummaryText", FTextBlockStyle(TutorialText)
+			.SetFontSize(10));
+
+		Set( "Tutorials.Browser.HighlightTextColor", TutorialSelectionColor );
+
+		Set( "Tutorials.Browser.Button", FButtonStyle()
+			.SetNormal( BOX_BRUSH( "Common/ButtonHoverHint", FMargin(4/16.0f), FLinearColor(1,1,1,0.15f) ) )
+			.SetHovered( BOX_BRUSH( "Common/ButtonHoverHint", FMargin(4/16.0f), FLinearColor(1,1,1,1) ) )
+			.SetPressed( BOX_BRUSH( "Common/ButtonHoverHint", FMargin(4/16.0f), FLinearColor(1,1,1,1) ) )
+			.SetNormalPadding( FMargin(0,0,0,1))
+			.SetPressedPadding( FMargin(0,1,0,0)));
+
+		Set( "Tutorials.Browser.CategoryArrow", new IMAGE_BRUSH( "Tutorials/BrowserCategoryArrow", FVector2D(24.0f, 24.0f), FLinearColor(0.0f, 0.0f, 0.0f, 1.0f) ) );
+
+		Set( "Tutorials.Browser.BackButton", new IMAGE_BRUSH( "Tutorials/BrowserBack", FVector2D(32.0f, 32.0f), FLinearColor(0.0f, 0.0f, 0.0f, 1.0f) ) );
+		Set( "Tutorials.Browser.RestartButton", new IMAGE_BRUSH( "Tutorials/BrowserRestart", FVector2D(16.0f, 16.0f), FLinearColor(0.0f, 0.0f, 0.0f, 1.0f) ) );
+
+		Set( "Tutorials.WidgetContent", FTextBlockStyle(TutorialText)
+			.SetFontSize(10));
+
 		Set( "Tutorials.ButtonColor", TutorialButtonColor );
 		Set( "Tutorials.ButtonHighlightColor", TutorialSelectionColor );
 		Set( "Tutorials.ButtonDisabledColor", SelectionColor_Inactive );
diff --git a/Engine/Source/Editor/IntroTutorials/IntroTutorials.Build.cs b/Engine/Source/Editor/IntroTutorials/IntroTutorials.Build.cs
--- a/Engine/Source/Editor/IntroTutorials/IntroTutorials.Build.cs
+++ b/Engine/Source/Editor/IntroTutorials/IntroTutorials.Build.cs
@@ -9,7 +9,7 @@ namespace UnrealBuildTool.Rules
 			PublicDependencyModuleNames.AddRange(
 				new string[] {
 					"Core",
-					"CoreUObject", // @todo Mac: for some reason CoreUObject and Engine are needed to link in debug on Mac
+					"CoreUObject",
 					"Engine", // @todo Mac: for some reason CoreUObject and Engine are needed to link in debug on Mac
                     "InputCore",
                     "Slate",
@@ -33,6 +33,8 @@ namespace UnrealBuildTool.Rules
                     "Kismet",
                     "PlacementMode",
 					"SlateCore",
+					"Settings",
+					"PropertyEditor"
 				}
 			);
 
diff --git a/Engine/Source/Editor/IntroTutorials/Private/EditorTutorial.cpp b/Engine/Source/Editor/IntroTutorials/Private/EditorTutorial.cpp
new file mode 100644
--- /dev/null
+++ b/Engine/Source/Editor/IntroTutorials/Private/EditorTutorial.cpp
@@ -0,0 +1,27 @@
+// Copyright 1998-2014 Epic Games, Inc. All Rights Reserved.
+
+#include "IntroTutorialsPrivatePCH.h"
+#include "EditorTutorial.h"
+
+UEditorTutorial::UEditorTutorial(const class FPostConstructInitializeProperties& PCIP)
+	: Super(PCIP)
+{
+}
+
+
+void UEditorTutorial::GotoNextTutorialStage()
+{
+
+}
+
+
+void UEditorTutorial::GotoPreviousTutorialStage()
+{
+
+}
+
+
+void UEditorTutorial::BeginTutorial(UEditorTutorial* TutorialToStart)
+{
+
+}
diff --git a/Engine/Source/Editor/IntroTutorials/Private/EditorTutorialDetailsCustomization.cpp b/Engine/Source/Editor/IntroTutorials/Private/EditorTutorialDetailsCustomization.cpp
new file mode 100644
--- /dev/null
+++ b/Engine/Source/Editor/IntroTutorials/Private/EditorTutorialDetailsCustomization.cpp
@@ -0,0 +1,48 @@
+// Copyright 1998-2014 Epic Games, Inc. All Rights Reserved.
+
+#include "IntroTutorialsPrivatePCH.h"
+#include "EditorTutorialDetailsCustomization.h"
+#include "Editor/LevelEditor/Public/LevelEditor.h"
+#include "EditorTutorial.h"
+
+#define LOCTEXT_NAMESPACE "EditorTutorialDetailsCustomization"
+
+TSharedRef<IDetailCustomization> FEditorTutorialDetailsCustomization::MakeInstance()
+{
+	return MakeShareable( new FEditorTutorialDetailsCustomization );
+}
+
+void FEditorTutorialDetailsCustomization::CustomizeDetails( IDetailLayoutBuilder& DetailLayout )
+{
+	struct Local
+	{
+		static FReply OnLaunchClicked(UEditorTutorial* Tutorial)
+		{
+			FLevelEditorModule& LevelEditorModule = FModuleManager::Get().GetModuleChecked<FLevelEditorModule>("LevelEditor");
+			IIntroTutorials& IntroTutorials = FModuleManager::Get().GetModuleChecked<IIntroTutorials>("IntroTutorials");
+			IntroTutorials.LaunchTutorial(Tutorial, true, LevelEditorModule.GetLevelEditorTab()->GetParentWindow());
+			return FReply::Handled();
+		}
+	};
+
+	TArray<TWeakObjectPtr<UObject>> Objects;
+	DetailLayout.GetObjectsBeingCustomized(Objects);
+	check(Objects.Num() > 0);
+	UEditorTutorial* Tutorial = CastChecked<UEditorTutorial>(Objects[0].Get());
+
+	IDetailCategoryBuilder& CategoryBuilder = DetailLayout.EditCategory(TEXT("Testing"), LOCTEXT("TestingSection", "Testing").ToString(), ECategoryPriority::Important);
+
+	CategoryBuilder.AddCustomRow(TEXT("Launch"))
+	.WholeRowContent()
+	.VAlign(VAlign_Center)
+	.HAlign(HAlign_Left)
+	[
+		SNew(SButton)
+		.OnClicked(FOnClicked::CreateStatic(&Local::OnLaunchClicked, Tutorial))
+		.Text(LOCTEXT("LaunchButtonLabel", "Launch"))
+		.ToolTipText(LOCTEXT("LaunchButtonTooltip", "Test this tutorial."))
+	];
+}
+
+
+#undef LOCTEXT_NAMESPACE
\ No newline at end of file
diff --git a/Engine/Source/Editor/IntroTutorials/Private/EditorTutorialDetailsCustomization.h b/Engine/Source/Editor/IntroTutorials/Private/EditorTutorialDetailsCustomization.h
new file mode 100644
--- /dev/null
+++ b/Engine/Source/Editor/IntroTutorials/Private/EditorTutorialDetailsCustomization.h
@@ -0,0 +1,13 @@
+// Copyright 1998-2014 Epic Games, Inc. All Rights Reserved.
+
+#pragma once
+
+class FEditorTutorialDetailsCustomization : public IDetailCustomization
+{
+public:
+	/** Makes a new instance of this detail layout class for a specific detail view requesting it */
+	static TSharedRef<IDetailCustomization> MakeInstance();
+
+	/** IDetailCustomization interface */
+	virtual void CustomizeDetails( IDetailLayoutBuilder& DetailLayout ) override;
+};
diff --git a/Engine/Source/Editor/IntroTutorials/Private/EditorTutorialSettings.cpp b/Engine/Source/Editor/IntroTutorials/Private/EditorTutorialSettings.cpp
new file mode 100644
--- /dev/null
+++ b/Engine/Source/Editor/IntroTutorials/Private/EditorTutorialSettings.cpp
@@ -0,0 +1,75 @@
+// Copyright 1998-2014 Epic Games, Inc. All Rights Reserved.
+
+#include "IntroTutorialsPrivatePCH.h"
+#include "EditorTutorialSettings.h"
+
+UEditorTutorialSettings::UEditorTutorialSettings(const class FPostConstructInitializeProperties& PCIP)
+	: Super(PCIP)
+{
+
+}
+
+void UEditorTutorialSettings::PostInitProperties()
+{
+	Super::PostInitProperties();
+
+	for(const auto& Progress : TutorialsProgress)
+	{
+		TSubclassOf<UEditorTutorial> TutorialClass = LoadClass<UEditorTutorial>(NULL, *Progress.Tutorial.AssetLongPathname, NULL, LOAD_None, NULL);
+		if(TutorialClass != nullptr)
+		{
+			UEditorTutorial* Tutorial = TutorialClass->GetDefaultObject<UEditorTutorial>();
+			ProgressMap.Add(Tutorial, Progress);
+		}
+	}
+}
+
+int32 UEditorTutorialSettings::GetProgress(UEditorTutorial* InTutorial, bool& bOutHaveSeenTutorial) const
+{
+	const FTutorialProgress* FoundProgress = ProgressMap.Find(InTutorial);
+	if(FoundProgress != nullptr)
+	{
+		bOutHaveSeenTutorial = true;
+		return FoundProgress->CurrentStage;
+	}
+
+	bOutHaveSeenTutorial = false;
+	return 0;
+}
+
+bool UEditorTutorialSettings::HaveSeenTutorial(UEditorTutorial* InTutorial) const
+{
+	return ProgressMap.Find(InTutorial) != nullptr;
+}
+
+void UEditorTutorialSettings::RecordProgress(UEditorTutorial* InTutorial, int32 CurrentStage)
+{
+	if(InTutorial != nullptr)
+	{
+		FTutorialProgress* FoundProgress = ProgressMap.Find(InTutorial);
+		if(FoundProgress != nullptr)
+		{
+			FoundProgress->CurrentStage = FMath::Max(FoundProgress->CurrentStage, CurrentStage);
+		}
+		else
+		{
+			FTutorialProgress Progress;
+			Progress.Tutorial = FStringClassReference(InTutorial->GetClass());
+			Progress.CurrentStage = CurrentStage;
+
+			ProgressMap.Add(InTutorial, Progress);
+		}
+	}
+}
+
+void UEditorTutorialSettings::SaveProgress()
+{
+	TutorialsProgress.Empty();
+
+	for(const auto& ProgressMapEntry : ProgressMap)
+	{
+		TutorialsProgress.Add(ProgressMapEntry.Value);
+	}
+
+	SaveConfig();
+}
\ No newline at end of file
diff --git a/Engine/Source/Editor/IntroTutorials/Private/EditorTutorialSettings.h b/Engine/Source/Editor/IntroTutorials/Private/EditorTutorialSettings.h
new file mode 100644
--- /dev/null
+++ b/Engine/Source/Editor/IntroTutorials/Private/EditorTutorialSettings.h
@@ -0,0 +1,59 @@
+// Copyright 1998-2014 Epic Games, Inc. All Rights Reserved.
+
+#pragma once
+
+#include "EditorTutorial.h"
+#include "EditorTutorialSettings.generated.h"
+
+/** Track the progress of an individual tutorial */
+USTRUCT()
+struct FTutorialProgress
+{
+	GENERATED_USTRUCT_BODY()
+
+	UPROPERTY()
+	FStringClassReference Tutorial;
+
+	UPROPERTY()
+	int32 CurrentStage;
+};
+
+/** Editor-wide tutorial settings */
+UCLASS(config=EditorGameAgnostic)
+class UEditorTutorialSettings : public UObject
+{
+	GENERATED_UCLASS_BODY()
+
+	/** Categories for tutorials */
+	UPROPERTY(Config, EditAnywhere, Category="Tutorials")
+	TArray<FTutorialCategory> Categories;
+
+	/** The tutorials available in the editor */
+	UPROPERTY(Config, EditAnywhere, Category="Tutorials", meta=(MetaClass="EditorTutorial"))
+	TArray<FStringClassReference> Tutorials;
+
+	UPROPERTY(Config)
+	TArray<FTutorialProgress> TutorialsProgress;
+
+	/** Tutorial to start on Editor startup */
+	UPROPERTY(Config, EditAnywhere, Category="Tutorials", meta=(MetaClass="EditorTutorial"))
+	FStringClassReference StartupTutorial;
+
+	/** UObject interface */
+	virtual void PostInitProperties() override;
+
+	/** Get the recorded progress of the pass-in tutorial */
+	int32 GetProgress(UEditorTutorial* InTutorial, bool& bOutHaveSeenTutorial) const;
+
+	/** Check if we have seen the passed-in tutorial before */
+	bool HaveSeenTutorial(UEditorTutorial* InTutorial) const;
+
+	/** Record the progress of the passed-in tutorial */
+	void RecordProgress(UEditorTutorial* InTutorial, int32 CurrentStage);
+
+	/** Save the progress of all our tutorials */
+	void SaveProgress();
+
+	/** Recorded progress */
+	TMap<UEditorTutorial*, FTutorialProgress> ProgressMap;
+};
\ No newline at end of file
diff --git a/Engine/Source/Editor/IntroTutorials/Private/InteractiveTutorials.cpp b/Engine/Source/Editor/IntroTutorials/Private/InteractiveTutorials.cpp
--- a/Engine/Source/Editor/IntroTutorials/Private/InteractiveTutorials.cpp
+++ b/Engine/Source/Editor/IntroTutorials/Private/InteractiveTutorials.cpp
@@ -65,11 +65,6 @@ void FInteractiveTutorials::SetupEditorHooks()
 
 	UK2Node_TutorialExcerptComplete::OnTutorialExcerptComplete.AddSP(this, &FInteractiveTutorials::CompleteExcerpt);
 
-
-	// Hook us up to the tutorial system to allow us to highlight particular widgets
-	STutorialWrapper::OnWidgetHighlight().BindSP(this, &FInteractiveTutorials::IsWidgetHighlighted);
-
-
 	// Then make sure to collect all our actual tutorial data
 	PopulateInteractivityData(SharedThis(this));
 }
diff --git a/Engine/Source/Editor/IntroTutorials/Private/IntroTutorials.cpp b/Engine/Source/Editor/IntroTutorials/Private/IntroTutorials.cpp
--- a/Engine/Source/Editor/IntroTutorials/Private/IntroTutorials.cpp
+++ b/Engine/Source/Editor/IntroTutorials/Private/IntroTutorials.cpp
@@ -15,6 +15,17 @@
 #include "EngineAnalytics.h"
 #include "Runtime/Analytics/Analytics/Public/Interfaces/IAnalyticsProvider.h"
 #include "Particles/ParticleSystem.h"
+#include "EditorTutorialSettings.h"
+#include "TutorialSettings.h"
+#include "Settings.h"
+#include "EditorTutorial.h"
+#include "SEditorTutorials.h"
+#include "STutorialsBrowser.h"
+#include "STutorialNavigation.h"
+#include "STutorialOverlay.h"
+#include "TutorialStructCustomization.h"
+#include "EditorTutorialDetailsCustomization.h"
+#include "STutorialRoot.h"
 
 #define LOCTEXT_NAMESPACE "IntroTutorials"
 
@@ -185,6 +196,29 @@ void FIntroTutorials::StartupModule()
 			ResetWelcomeTutorials();
 		}
 	}
+
+	// Register to display our settings
+	ISettingsModule* SettingsModule = ISettingsModule::Get();
+	if (SettingsModule != nullptr)
+	{
+		SettingsModule->RegisterSettings("Editor", "General", "Tutorials",
+			LOCTEXT("EditorTutorialSettingsName", "Tutorials"),
+			LOCTEXT("EditorTutorialSettingsDescription", "Control what tutorials are available in the Editor."),
+			GetMutableDefault<UEditorTutorialSettings>()
+		);
+
+		SettingsModule->RegisterSettings("Project", "Engine", "Tutorials",
+			LOCTEXT("TutorialSettingsName", "Tutorials"),
+			LOCTEXT("TutorialSettingsDescription", "Control what tutorials are available in this project."),
+			GetMutableDefault<UTutorialSettings>()
+		);
+	}
+
+	// register details customizations
+	FPropertyEditorModule& PropertyEditorModule = FModuleManager::GetModuleChecked<FPropertyEditorModule>("PropertyEditor");
+	PropertyEditorModule.RegisterCustomPropertyTypeLayout("TutorialContent", FOnGetPropertyTypeCustomizationInstance::CreateStatic(&FTutorialContentCustomization::MakeInstance));
+	PropertyEditorModule.RegisterCustomPropertyTypeLayout("TutorialContentAnchor", FOnGetPropertyTypeCustomizationInstance::CreateStatic(&FTutorialContentAnchorCustomization::MakeInstance));
+	PropertyEditorModule.RegisterCustomClassLayout("EditorTutorial", FOnGetDetailCustomizationInstance::CreateStatic(&FEditorTutorialDetailsCustomization::MakeInstance));
 }
 
 void FIntroTutorials::ShutdownModule()
@@ -208,6 +242,21 @@ void FIntroTutorials::ShutdownModule()
 		IMainFrameModule& MainFrameModule = FModuleManager::LoadModuleChecked<IMainFrameModule>(TEXT("MainFrame"));
 		MainFrameModule.OnMainFrameCreationFinished().RemoveAll(this);
 	}
+
+	ISettingsModule* SettingsModule = ISettingsModule::Get();
+	if (SettingsModule != nullptr)
+	{
+		SettingsModule->UnregisterSettings("Editor", "General", "Tutorials");
+		SettingsModule->UnregisterSettings("Project", "Engine", "Tutorials");
+	}
+
+	if(FModuleManager::Get().IsModuleLoaded("PropertyEditor"))
+	{
+		FPropertyEditorModule& PropertyEditorModule = FModuleManager::GetModuleChecked<FPropertyEditorModule>("PropertyEditor");
+		PropertyEditorModule.UnregisterCustomPropertyTypeLayout("TutorialContent");
+		PropertyEditorModule.UnregisterCustomPropertyTypeLayout("TutorialWidgetReference");
+		PropertyEditorModule.UnregisterCustomClassLayout("EditorTutorial");
+	}
 }
 
 void FIntroTutorials::AddSummonTutorialsMenuExtension(FMenuBuilder& MenuBuilder)
@@ -258,6 +307,15 @@ void FIntroTutorials::MainFrameLoad(TSharedPtr<SWindow> InRootWindow, bool bIsNe
 			RootWindow = InRootWindow;
 		}
 
+		// install a root widget for the tutorial overlays to hang off
+		if(InRootWindow.IsValid() && !TutorialRoot.IsValid())
+		{
+			InRootWindow->AddOverlaySlot()
+				[
+					SAssignNew(TutorialRoot, STutorialRoot)
+				];
+		}
+
 		// See if we should show 'welcome' screen
 		MaybeOpenWelcomeTutorial(UE4WelcomeTutorial.TutorialPath, UE4WelcomeTutorial.SeenOnceSettingName);
 	}
@@ -265,8 +323,16 @@ void FIntroTutorials::MainFrameLoad(TSharedPtr<SWindow> InRootWindow, bool bIsNe
 
 void FIntroTutorials::SummonTutorialHome()
 {
-	CurrentObjectClass = nullptr;
-	SummonTutorialWindowForPage(HomePath);
+	if(FParse::Param(FCommandLine::Get(), TEXT("NewTutorials")))
+	{
+		FLevelEditorModule& LevelEditorModule = FModuleManager::LoadModuleChecked<FLevelEditorModule>( "LevelEditor" );
+		SummonTutorialBrowser(LevelEditorModule.GetLevelEditorTab()->GetParentWindow().ToSharedRef());
+	}
+	else
+	{
+		CurrentObjectClass = nullptr;
+		SummonTutorialWindowForPage(HomePath);
+	}
 }
 
 void FIntroTutorials::SummonBlueprintTutorialHome(UObject* Asset, bool bForceWelcome)
@@ -471,6 +537,35 @@ FWelcomeTutorialProperties const* FIntroTutorials::FindAssetEditorTutorialProper
 
 bool FIntroTutorials::MaybeOpenWelcomeTutorial(const FString& TutorialPath, const FString& ConfigSettingName)
 {
+	if(FParse::Param(FCommandLine::Get(), TEXT("NewTutorials")))
+	{
+		// try editor startup tutorial
+		TSubclassOf<UEditorTutorial> EditorStartupTutorialClass = LoadClass<UEditorTutorial>(NULL, *GetDefault<UEditorTutorialSettings>()->StartupTutorial.AssetLongPathname, NULL, LOAD_None, NULL);
+		if(EditorStartupTutorialClass != nullptr)
+		{
+			UEditorTutorial* Tutorial = EditorStartupTutorialClass->GetDefaultObject<UEditorTutorial>();
+			if(!GetDefault<UEditorTutorialSettings>()->HaveSeenTutorial(Tutorial))
+			{
+				LaunchTutorial(Tutorial);
+				return true;
+			}
+		}
+
+		// Try project startup tutorial
+		TSubclassOf<UEditorTutorial> ProjectStartupTutorialClass = LoadClass<UEditorTutorial>(NULL, *GetDefault<UTutorialSettings>()->StartupTutorial.AssetLongPathname, NULL, LOAD_None, NULL);
+		if(ProjectStartupTutorialClass != nullptr)
+		{
+			UEditorTutorial* Tutorial = ProjectStartupTutorialClass->GetDefaultObject<UEditorTutorial>();
+			if(!GetDefault<UEditorTutorialSettings>()->HaveSeenTutorial(Tutorial))
+			{
+				LaunchTutorial(Tutorial);
+				return true;
+			}
+		}
+
+		return false;
+	}
+
 	// don't open if viewing any tutorial other than the index
 	if (TutorialWidget.IsValid() && (TutorialWidget.Pin()->GetCurrentPagePath() != HomePath))
 	{
@@ -632,4 +727,25 @@ void FIntroTutorials::UnregisterTutorialForAssetEditor(UClass* AssetClass)
 	AssetEditorTutorialPropertyMap.Remove(AssetClass);
 }
 
+FOnIsPicking& FIntroTutorials::OnIsPicking()
+{
+	return OnIsPickingDelegate;
+}
+
+void FIntroTutorials::SummonTutorialBrowser(TSharedRef<SWindow> InWindow, const FString& InFilter)
+{
+	if(TutorialRoot.IsValid())
+	{
+		TutorialRoot->SummonTutorialBrowser(InWindow, InFilter);
+	}
+}
+
+void FIntroTutorials::LaunchTutorial(UEditorTutorial* InTutorial, bool bInRestart, TWeakPtr<SWindow> InNavigationWindow)
+{
+	if(TutorialRoot.IsValid())
+	{
+		TutorialRoot->LaunchTutorial(InTutorial, bInRestart, InNavigationWindow);
+	}
+}
+
 #undef LOCTEXT_NAMESPACE
\ No newline at end of file
diff --git a/Engine/Source/Editor/IntroTutorials/Private/IntroTutorials.h b/Engine/Source/Editor/IntroTutorials/Private/IntroTutorials.h
--- a/Engine/Source/Editor/IntroTutorials/Private/IntroTutorials.h
+++ b/Engine/Source/Editor/IntroTutorials/Private/IntroTutorials.h
@@ -2,6 +2,13 @@
 
 #pragma once
 
+/** Query delegate to see if we are in picking mode */
+DECLARE_DELEGATE_RetVal_OneParam(bool, FOnIsPicking, FName& /* OutWidgetNameToHighlight */);
+
+class UEditorTutorial;
+struct FTutorialContent;
+class STutorialRoot;
+
 /** Contains properties defining a "welcome" tutorial, which is auto-shown once.*/
 struct FWelcomeTutorialProperties
 {
@@ -58,7 +65,11 @@ public:
 	// ctor
 	FIntroTutorials();
 
+	/** Get the delegate used to check for whether we are picking widgets */
+	FOnIsPicking& OnIsPicking();
+
 private:
+
 	/** IModuleInterface implementation */
 	virtual void StartupModule() override;
 	virtual void ShutdownModule() override;
@@ -126,11 +137,14 @@ private:
 	/** Handle linking between tutorials */
 	FString HandleGotoNextTutorial(const FString& InCurrentPage) const;
 
+	void SummonTutorialBrowser(TSharedRef<SWindow> InWindow, const FString& InFilter = TEXT(""));
+
 public:
 
 	// IIntroTutorials interface
 	virtual void RegisterTutorialForAssetEditor(UClass* AssetClass, const FString& TutorialDocPath, const FString& TutorialHasBeenSeenSettingName, const FString& SurveyGUIDString) override;
 	virtual void UnregisterTutorialForAssetEditor(UClass* AssetClass) override;
+	virtual void LaunchTutorial(UEditorTutorial* Tutorial, bool bRestart = true, TWeakPtr<SWindow> InNavigationWindow = nullptr) override;
 	// End of IIntroTutorials interface
 private:
 	/** The tab id of the tutorial tab */
@@ -168,4 +182,10 @@ private:
 
 	/** should we be clearing the 'have seen this tutorial' flag? (controlled by -tutorials on the command line) */
 	bool bDesireResettingTutorialSeenFlagOnLoad;
+
+	/** Delegate used to determine whether we are in picking mode */
+	FOnIsPicking OnIsPickingDelegate;
+
+	/** Root widget for tutorial overlay system */
+	TSharedPtr<STutorialRoot> TutorialRoot;
 };
\ No newline at end of file
diff --git a/Engine/Source/Editor/IntroTutorials/Private/IntroTutorialsPrivatePCH.h b/Engine/Source/Editor/IntroTutorials/Private/IntroTutorialsPrivatePCH.h
--- a/Engine/Source/Editor/IntroTutorials/Private/IntroTutorialsPrivatePCH.h
+++ b/Engine/Source/Editor/IntroTutorials/Private/IntroTutorialsPrivatePCH.h
@@ -1,9 +1,10 @@
 // Copyright 1998-2014 Epic Games, Inc. All Rights Reserved.
 
 #include "UnrealEd.h"
+#include "PropertyEditing.h"
 #include "IIntroTutorials.h"
 
 // You should place include statements to your module's private header files here.  You only need to
 // add includes for headers that are used in most of your module's source files though.
 
-#include "IntroTutorials.h"
\ No newline at end of file
+#include "IntroTutorials.h"
diff --git a/Engine/Source/Editor/IntroTutorials/Private/SEditorTutorials.cpp b/Engine/Source/Editor/IntroTutorials/Private/SEditorTutorials.cpp
new file mode 100644
--- /dev/null
+++ b/Engine/Source/Editor/IntroTutorials/Private/SEditorTutorials.cpp
@@ -0,0 +1,202 @@
+// Copyright 1998-2014 Epic Games, Inc. All Rights Reserved.
+
+#include "IntroTutorialsPrivatePCH.h"
+#include "SEditorTutorials.h"
+#include "STutorialsBrowser.h"
+#include "STutorialNavigation.h"
+#include "STutorialOverlay.h"
+#include "EditorTutorial.h"
+
+#define LOCTEXT_NAMESPACE "TutorialsBrowser"
+
+void SEditorTutorials::Construct(const FArguments& InArgs)
+{
+	bBrowserVisible = false;
+	bShowNavigation = false;
+	ParentWindow = InArgs._ParentWindow;
+	OnNextClicked = InArgs._OnNextClicked;
+	OnBackClicked = InArgs._OnBackClicked;
+	OnHomeClicked = InArgs._OnHomeClicked;
+	OnGetCurrentTutorial = InArgs._OnGetCurrentTutorial;
+	OnGetCurrentTutorialStage = InArgs._OnGetCurrentTutorialStage;
+
+	TutorialHome = SNew(STutorialsBrowser)
+		.Visibility(this, &SEditorTutorials::GetBrowserVisibility)
+		.OnLaunchTutorial(InArgs._OnLaunchTutorial)
+		.OnClosed(FSimpleDelegate::CreateSP(this, &SEditorTutorials::HandleCloseClicked))
+		.ParentWindow(InArgs._ParentWindow);
+
+	NavigationWidget = SNew(STutorialNavigation)
+		.Visibility(this, &SEditorTutorials::GetNavigationVisibility)
+		.OnBackClicked(FSimpleDelegate::CreateSP(this, &SEditorTutorials::HandleBackClicked))
+		.OnHomeClicked(FSimpleDelegate::CreateSP(this, &SEditorTutorials::HandleHomeClicked))
+		.OnNextClicked(FSimpleDelegate::CreateSP(this, &SEditorTutorials::HandleNextClicked))
+		.IsBackEnabled(this, &SEditorTutorials::IsBackButtonEnabled)
+		.IsHomeEnabled(this, &SEditorTutorials::IsHomeButtonEnabled)
+		.IsNextEnabled(this, &SEditorTutorials::IsNextButtonEnabled)
+		.OnGetProgress(this, &SEditorTutorials::GetProgress)
+		.ParentWindow(InArgs._ParentWindow);
+
+	ContentBox = SNew(SHorizontalBox);
+
+	ChildSlot
+	[
+		SNew(SOverlay)
+		+SOverlay::Slot()
+		[
+			SNew(SVerticalBox)
+			+SVerticalBox::Slot()
+			.VAlign(VAlign_Center)
+			.HAlign(HAlign_Center)
+			[
+				TutorialHome.ToSharedRef()
+			]
+		]
+		+SOverlay::Slot()
+		.VAlign(VAlign_Fill)
+		.HAlign(HAlign_Fill)
+		[
+			SNew(SVerticalBox)
+			+SVerticalBox::Slot()
+			.FillHeight(0.75f)
+			[
+				SNullWidget::NullWidget
+			]
+			+SVerticalBox::Slot()
+			.FillHeight(0.25f)
+			.VAlign(VAlign_Center)
+			.HAlign(HAlign_Center)
+			[
+				NavigationWidget.ToSharedRef()
+			]
+		]
+		+SOverlay::Slot()
+		.VAlign(VAlign_Fill)
+		.HAlign(HAlign_Fill)
+		[
+			ContentBox.ToSharedRef()
+		]
+	];
+}
+
+void SEditorTutorials::LaunchTutorial(bool bInIsNavigationWindow)
+{
+	bShowNavigation = bInIsNavigationWindow;
+
+	RebuildCurrentContent();
+}
+
+
+void SEditorTutorials::ShowBrowser(const FString& InFilter)
+{
+	TutorialHome->SetFilter(InFilter);
+	HandleHomeClicked();
+}
+
+EVisibility SEditorTutorials::GetBrowserVisibility() const
+{
+	return bBrowserVisible && OnGetCurrentTutorial.Execute() == nullptr ? EVisibility::Visible : EVisibility::Collapsed;
+}
+
+EVisibility SEditorTutorials::GetNavigationVisibility() const
+{
+	UEditorTutorial* CurrentTutorial = OnGetCurrentTutorial.Execute();
+
+	return bShowNavigation && CurrentTutorial != nullptr && !CurrentTutorial->bIsStandalone ? EVisibility::Visible : EVisibility::Collapsed;
+}
+
+void SEditorTutorials::HandleCloseClicked()
+{
+	HandleHomeClicked();
+	bBrowserVisible = false;
+}
+
+void SEditorTutorials::HandleBackClicked()
+{
+	// forward to other overlays so they can rebuild their widgets as well
+	OnBackClicked.ExecuteIfBound();
+}
+
+void SEditorTutorials::HandleHomeClicked()
+{
+	OnHomeClicked.ExecuteIfBound();
+
+	ContentBox->ClearChildren();
+	bBrowserVisible = true;
+}
+
+void SEditorTutorials::HandleNextClicked()
+{
+	// forward to other overlays so they can rebuild their widgets as well
+	OnNextClicked.ExecuteIfBound(ParentWindow);
+}
+
+bool SEditorTutorials::IsBackButtonEnabled() const
+{
+	UEditorTutorial* CurrentTutorial = OnGetCurrentTutorial.Execute();
+	int32 CurrentTutorialStage = OnGetCurrentTutorialStage.Execute();
+
+	return CurrentTutorial != nullptr && CurrentTutorialStage > 0;
+}
+
+bool SEditorTutorials::IsHomeButtonEnabled() const
+{
+	return true;
+}
+
+bool SEditorTutorials::IsNextButtonEnabled() const
+{
+	UEditorTutorial* CurrentTutorial = OnGetCurrentTutorial.Execute();
+	int32 CurrentTutorialStage = OnGetCurrentTutorialStage.Execute();
+
+	const bool bStepsRemaining = CurrentTutorial != nullptr && CurrentTutorialStage + 1 < CurrentTutorial->Stages.Num();
+	const bool bNextTutorialAvailable = CurrentTutorial != nullptr && CurrentTutorialStage + 1 >= CurrentTutorial->Stages.Num() && CurrentTutorial->NextTutorial != nullptr;
+	return bStepsRemaining || bNextTutorialAvailable;
+}
+
+float SEditorTutorials::GetProgress() const
+{
+	UEditorTutorial* CurrentTutorial = OnGetCurrentTutorial.Execute();
+	int32 CurrentTutorialStage = OnGetCurrentTutorialStage.Execute();
+
+	return (CurrentTutorial != nullptr && CurrentTutorial->Stages.Num() > 0) ? (float)(CurrentTutorialStage + 1) / (float)CurrentTutorial->Stages.Num() : 0.0f;
+}
+
+void SEditorTutorials::RebuildCurrentContent()
+{
+	UEditorTutorial* CurrentTutorial = OnGetCurrentTutorial.Execute();
+	int32 CurrentTutorialStage = OnGetCurrentTutorialStage.Execute();
+
+	OverlayContent = nullptr;
+	ContentBox->ClearChildren();
+	if(CurrentTutorial != nullptr && CurrentTutorialStage < CurrentTutorial->Stages.Num())
+	{
+		ContentBox->AddSlot()
+		[
+			SAssignNew(OverlayContent, STutorialOverlay, &CurrentTutorial->Stages[CurrentTutorialStage])
+			.OnClosed(FSimpleDelegate::CreateSP(this, &SEditorTutorials::HandleCloseClicked))
+			.IsStandalone(CurrentTutorial->bIsStandalone)
+			.ParentWindow(ParentWindow)
+			.AllowNonWidgetContent(bShowNavigation)
+		];
+	}
+	else
+	{
+		// create 'empty' overlay, as we may need this for picking visualization
+		ContentBox->AddSlot()
+		[
+			SAssignNew(OverlayContent, STutorialOverlay, nullptr)
+			.OnClosed(FSimpleDelegate::CreateSP(this, &SEditorTutorials::HandleCloseClicked))
+			.IsStandalone(false)
+			.ParentWindow(ParentWindow)
+			.AllowNonWidgetContent(false)
+		];	
+	}
+}
+
+TSharedPtr<SWindow> SEditorTutorials::GetParentWindow() const
+{
+	return ParentWindow.IsValid() ? ParentWindow.Pin() : TSharedPtr<SWindow>();
+}
+
+#undef LOCTEXT_NAMESPACE
\ No newline at end of file
diff --git a/Engine/Source/Editor/IntroTutorials/Private/SEditorTutorials.h b/Engine/Source/Editor/IntroTutorials/Private/SEditorTutorials.h
new file mode 100644
--- /dev/null
+++ b/Engine/Source/Editor/IntroTutorials/Private/SEditorTutorials.h
@@ -0,0 +1,121 @@
+// Copyright 1998-2014 Epic Games, Inc. All Rights Reserved.
+
+#pragma once
+
+#include "STutorialsBrowser.h"
+
+class STutorialsBrowser;
+class STutorialOverlay;
+
+/** Delegate fired when next button is clicked */
+DECLARE_DELEGATE_OneParam(FOnNextClicked, TWeakPtr<SWindow> /*InNavigationWindow*/);
+
+/** Delegate fired to retrieve current tutorial */
+DECLARE_DELEGATE_RetVal(UEditorTutorial*, FOnGetCurrentTutorial);
+
+/** Delegate fired to retrieve current tutorial stage */
+DECLARE_DELEGATE_RetVal(int32, FOnGetCurrentTutorialStage);
+
+/**
+ * Container widget for all tutorial-related widgets
+ */
+class SEditorTutorials : public SCompoundWidget
+{
+	SLATE_BEGIN_ARGS( SEditorTutorials )
+	{
+		_Visibility = EVisibility::SelfHitTestInvisible;
+	}
+
+	SLATE_ARGUMENT(FSimpleDelegate, OnClosed)
+	SLATE_ARGUMENT(TWeakPtr<SWindow>, ParentWindow)
+	SLATE_ARGUMENT(FOnNextClicked, OnNextClicked)
+	SLATE_ARGUMENT(FSimpleDelegate, OnBackClicked)
+	SLATE_ARGUMENT(FSimpleDelegate, OnHomeClicked)
+	SLATE_ARGUMENT(FOnGetCurrentTutorial, OnGetCurrentTutorial)
+	SLATE_ARGUMENT(FOnGetCurrentTutorialStage, OnGetCurrentTutorialStage)
+	SLATE_ARGUMENT(FOnLaunchTutorial, OnLaunchTutorial)
+
+	SLATE_END_ARGS()
+
+	void Construct(const FArguments& InArgs);
+
+	/** Launch a tutorial - will interrogate parent to get the tutorial data to display */
+	void LaunchTutorial(bool bInIsNavigationWindow);
+
+	/** Show the tutorials browser in this window */
+	void ShowBrowser(const FString& InFilter);
+
+	/** Get parent window for this widget */
+	TSharedPtr<SWindow> GetParentWindow() const;
+
+	/** Rebuild content according ot the current tutorial state */
+	void RebuildCurrentContent();
+
+private:
+	/** Handle whether we should display the browser */
+	EVisibility GetBrowserVisibility() const;
+
+	/** Handle whether we should display the navigation buttons */
+	EVisibility GetNavigationVisibility() const;
+
+	/** Handle hitting the close button on the browser or a standalone widget */
+	void HandleCloseClicked();
+
+	/** Handle hitting the home button in the navigation window */
+	void HandleHomeClicked();
+
+	/** Handle hitting the back button in the navigation window */
+	void HandleBackClicked();
+
+	/** Handle hitting the next button in the navigation window */
+	void HandleNextClicked();
+
+	/** Get enabled state of the back button */
+	bool IsBackButtonEnabled() const;
+
+	/** Get enabled state of the home button */
+	bool IsHomeButtonEnabled() const;
+
+	/** Get enabled state of the next button */
+	bool IsNextButtonEnabled() const;
+
+	/** Get the current progress to display in the progress bar */
+	float GetProgress() const;
+
+private:
+	/** The root widget of the tutorial home page */
+	TSharedPtr<STutorialsBrowser> TutorialHome;
+
+	/** Navigation widget */
+	TSharedPtr<SWidget> NavigationWidget;
+
+	/** Box that contains varied content for current tutorial */
+	TSharedPtr<SHorizontalBox> ContentBox;
+
+	/** Content widget for current tutorial */
+	TSharedPtr<STutorialOverlay> OverlayContent;
+
+	/** Window that contains this widget */
+	TWeakPtr<SWindow> ParentWindow;
+
+	/** Whether the browser is visible */
+	bool bBrowserVisible;
+
+	/** Whether we should display navigation */
+	bool bShowNavigation;
+
+	/** Delegate fired when next button is clicked */
+	FOnNextClicked OnNextClicked;
+
+	/** Delegate fired when back button is clicked */
+	FSimpleDelegate OnBackClicked;
+
+	/** Delegate fired when home button is clicked */
+	FSimpleDelegate OnHomeClicked;
+
+	/** Delegate fired to retrive the current tutorial */
+	FOnGetCurrentTutorial OnGetCurrentTutorial;
+
+	/** Delegate fired to retrive the current tutorial stage */
+	FOnGetCurrentTutorialStage OnGetCurrentTutorialStage;
+};
\ No newline at end of file
diff --git a/Engine/Source/Editor/IntroTutorials/Private/STutorialContent.cpp b/Engine/Source/Editor/IntroTutorials/Private/STutorialContent.cpp
new file mode 100644
--- /dev/null
+++ b/Engine/Source/Editor/IntroTutorials/Private/STutorialContent.cpp
@@ -0,0 +1,299 @@
+// Copyright 1998-2014 Epic Games, Inc. All Rights Reserved.
+
+#include "IntroTutorialsPrivatePCH.h"
+#include "STutorialContent.h"
+
+namespace TutorialConstants
+{
+	const float BorderPulseAnimationLength = 1.0f;
+	const float BorderIntroAnimationLength = 0.5f;
+	const float MinBorderOpacity = 0.1f;
+	const float ShadowScale = 8.0f;
+	const float MaxBorderOffset = 8.0f;
+}
+
+const float ContentOffset = 10.0f;
+
+void STutorialContent::Construct(const FArguments& InArgs, const FTutorialContent& InContent)
+{
+	bIsVisible = Anchor.Type == ETutorialAnchorIdentifier::None;
+
+	VerticalAlignment = InArgs._VAlign;
+	HorizontalAlignment = InArgs._HAlign;
+	WidgetOffset = InArgs._Offset;
+	bIsStandalone = InArgs._IsStandalone;
+	OnClosed = InArgs._OnClosed;
+	Anchor = InArgs._Anchor;
+
+	BorderIntroAnimation.AddCurve(0.0f, TutorialConstants::BorderIntroAnimationLength, ECurveEaseFunction::CubicOut);
+	BorderPulseAnimation.AddCurve(0.0f, TutorialConstants::BorderPulseAnimationLength, ECurveEaseFunction::Linear);
+	BorderIntroAnimation.Play();
+	BorderPulseAnimation.Play();
+
+	ChildSlot
+	[
+		SAssignNew(ContentWidget, SBorder)
+
+		// Add more padding if the content is to be displayed centrally (i.e. not on a widget)
+		.Padding(Anchor.Type != ETutorialAnchorIdentifier::None ? 18.0f : 24.0f)
+		.BorderImage(FEditorStyle::GetBrush("Tutorials.Border"))
+		.Visibility(this, &STutorialContent::GetVisibility)
+		[
+			SNew(SHorizontalBox)
+			+SHorizontalBox::Slot()
+			.AutoWidth()
+			.MaxWidth(600.0f)
+			.VAlign(VAlign_Center)
+			[
+				GenerateContentWidget(InContent, InArgs._WrapTextAt, Anchor.Type != ETutorialAnchorIdentifier::None, DocumentationPage)
+			]
+			+SHorizontalBox::Slot()
+			.AutoWidth()
+			.VAlign(VAlign_Top)
+			.Padding(5.0f, 0.0f, 0.0f, 0.0f)
+			[
+				SNew(SButton)
+				.OnClicked(this, &STutorialContent::OnCloseButtonClicked)
+				.Visibility(this, &STutorialContent::GetCloseButtonVisibility)
+				.ButtonStyle(&FEditorStyle::Get().GetWidgetStyle<FButtonStyle>("Tutorials.Browser.Button"))
+				.ContentPadding(0.0f)
+				[
+					SNew(SImage)
+					.Image(FEditorStyle::GetBrush("Symbols.X"))
+					.ColorAndOpacity(FLinearColor(0.0f, 0.0f, 0.0f, 1.0f))
+				]
+			]
+		]
+	];
+}
+
+static void GetAnimationValues(bool bIsIntro, float InAnimationProgress, float& OutAlphaFactor, float& OutPulseFactor, FLinearColor& OutShadowTint, FLinearColor& OutBorderTint)
+{
+	if(bIsIntro)
+	{
+		OutAlphaFactor = InAnimationProgress;
+		OutPulseFactor = (1.0f - OutAlphaFactor) * 50.0f;
+		OutShadowTint = FLinearColor(1.0f, 1.0f, 0.0f, OutAlphaFactor);
+		OutBorderTint = FLinearColor(1.0f, 1.0f, 0.0f, OutAlphaFactor * OutAlphaFactor);
+	}
+	else
+	{
+		OutAlphaFactor = 1.0f - (0.5f + (FMath::Cos(2.0f * PI * InAnimationProgress) * 0.5f));
+		OutPulseFactor = 0.5f + (FMath::Cos(2.0f * PI * InAnimationProgress) * 0.5f);
+		OutShadowTint = FLinearColor(1.0f, 1.0f, 0.0f, 1.0f);
+		OutBorderTint = FLinearColor(1.0f, 1.0f, 0.0f, TutorialConstants::MinBorderOpacity + ((1.0f - TutorialConstants::MinBorderOpacity) * OutAlphaFactor));
+	}
+}
+
+int32 STutorialContent::OnPaint( const FPaintArgs& Args, const FGeometry& AllottedGeometry, const FSlateRect& MyClippingRect, FSlateWindowElementList& OutDrawElements, int32 LayerId, const FWidgetStyle& InWidgetStyle, bool bParentEnabled ) const
+{
+	if(bIsVisible && Anchor.Type != ETutorialAnchorIdentifier::None)
+	{
+		float AlphaFactor;
+		float PulseFactor;
+		FLinearColor ShadowTint;
+		FLinearColor BorderTint;
+		GetAnimationValues(BorderIntroAnimation.IsPlaying(),  BorderIntroAnimation.IsPlaying() ? BorderIntroAnimation.GetLerp() : BorderPulseAnimation.GetLerpLooping(), AlphaFactor, PulseFactor, ShadowTint, BorderTint);
+
+		const FSlateBrush* ShadowBrush = FCoreStyle::Get().GetBrush(TEXT("Tutorials.Shadow"));
+		const FSlateBrush* BorderBrush = FCoreStyle::Get().GetBrush(TEXT("Tutorials.Border"));
+					
+		const FGeometry& WidgetGeometry = CachedGeometry;
+		const FVector2D WindowSize = OutDrawElements.GetWindow()->GetSizeInScreen();
+
+		// We should be clipped by the window size, not our containing widget, as we want to draw outside the widget
+		FSlateRect WindowClippingRect(0.0f, 0.0f, WindowSize.X, WindowSize.Y);
+
+		FPaintGeometry ShadowGeometry((WidgetGeometry.AbsolutePosition - FVector2D(ShadowBrush->Margin.Left, ShadowBrush->Margin.Top) * ShadowBrush->ImageSize * WidgetGeometry.Scale * TutorialConstants::ShadowScale),
+										((WidgetGeometry.Size * WidgetGeometry.Scale) + (FVector2D(ShadowBrush->Margin.Right * 2.0f, ShadowBrush->Margin.Bottom * 2.0f) * ShadowBrush->ImageSize * WidgetGeometry.Scale * TutorialConstants::ShadowScale)),
+										WidgetGeometry.Scale * TutorialConstants::ShadowScale);
+		// draw highlight shadow
+		FSlateDrawElement::MakeBox(OutDrawElements, LayerId++, ShadowGeometry, ShadowBrush, WindowClippingRect, ESlateDrawEffect::None, ShadowTint);
+
+		FVector2D PulseOffset = FVector2D(PulseFactor * TutorialConstants::MaxBorderOffset, PulseFactor * TutorialConstants::MaxBorderOffset);
+
+		FVector2D BorderPosition = (WidgetGeometry.AbsolutePosition - ((FVector2D(BorderBrush->Margin.Left, BorderBrush->Margin.Top) * BorderBrush->ImageSize * WidgetGeometry.Scale) + PulseOffset));
+		FVector2D BorderSize = ((WidgetGeometry.Size * WidgetGeometry.Scale) + (PulseOffset * 2.0f) + (FVector2D(BorderBrush->Margin.Right * 2.0f, BorderBrush->Margin.Bottom * 2.0f) * BorderBrush->ImageSize * WidgetGeometry.Scale));
+
+		FPaintGeometry BorderGeometry(BorderPosition, BorderSize, WidgetGeometry.Scale);
+
+		// draw highlight border
+		FSlateDrawElement::MakeBox(OutDrawElements, LayerId++, BorderGeometry, BorderBrush, WindowClippingRect, ESlateDrawEffect::None, BorderTint);
+	}
+
+	return SCompoundWidget::OnPaint(Args, AllottedGeometry, MyClippingRect, OutDrawElements, LayerId, InWidgetStyle, bParentEnabled);
+}
+
+/** Helper function to generate title widget, if any */
+static TSharedRef<SWidget> GetStageTitle(const FExcerpt& InExcerpt, int32 InCurrentExcerptIndex)
+{
+	// First try for unadorned 'StageTitle'
+	FString VariableName(TEXT("StageTitle"));
+	const FString* VariableValue = InExcerpt.Variables.Find(VariableName);
+	if(VariableValue != NULL)
+	{
+		return SNew(STextBlock)
+			.Text(FText::FromString(*VariableValue))
+			.TextStyle(FEditorStyle::Get(), "Tutorials.CurrentExcerpt");
+	}
+
+	// Then try 'StageTitle<StageNum>'
+	VariableName = FString::Printf(TEXT("StageTitle%d"), InCurrentExcerptIndex + 1);
+	VariableValue = InExcerpt.Variables.Find(VariableName);
+	if(VariableValue != NULL)
+	{
+		return SNew(STextBlock)
+			.Text(FText::FromString(*VariableValue))
+			.TextStyle(FEditorStyle::Get(), "Tutorials.CurrentExcerpt");
+	}
+
+	return SNullWidget::NullWidget;
+}
+
+TSharedRef<SWidget> STutorialContent::GenerateContentWidget(const FTutorialContent& InContent, float WrapTextAt, bool bForWidget, TSharedPtr<IDocumentationPage>& OutDocumentationPage)
+{
+	// Style for the documentation
+	static FDocumentationStyle DocumentationStyle;
+	DocumentationStyle
+		.ContentStyle(TEXT("Tutorials.Content"))
+		.BoldContentStyle(TEXT("Tutorials.BoldContent"))
+		.NumberedContentStyle(TEXT("Tutorials.NumberedContent"))
+		.Header1Style(TEXT("Tutorials.Header1"))
+		.Header2Style(TEXT("Tutorials.Header2"))
+		.HyperlinkButtonStyle(TEXT("Tutorials.Hyperlink.Button"))
+		.HyperlinkTextStyle(TEXT("Tutorials.Hyperlink.Text"))
+		.SeparatorStyle(TEXT("Tutorials.Separator"));
+
+	switch(InContent.Type)
+	{
+	case ETutorialContent::Text:
+		{
+			OutDocumentationPage = nullptr;
+
+			return SNew(STextBlock)
+				.Visibility(EVisibility::SelfHitTestInvisible)
+				.WrapTextAt(WrapTextAt)
+				.Text(InContent.Text)
+				.TextStyle(FEditorStyle::Get(), bForWidget ? "Tutorials.WidgetContent" : "Tutorials.Content");
+		}
+
+	case ETutorialContent::UDNExcerpt:
+		if (IDocumentation::Get()->PageExists(InContent.Content))
+		{
+			OutDocumentationPage = IDocumentation::Get()->GetPage(InContent.Content, TSharedPtr<FParserConfiguration>(), DocumentationStyle);
+			FExcerpt Excerpt;
+			if(OutDocumentationPage->GetExcerpt(InContent.ExcerptName, Excerpt) && OutDocumentationPage->GetExcerptContent(Excerpt))
+			{
+				return SNew(SVerticalBox)
+					+SVerticalBox::Slot()
+					.HAlign(HAlign_Left)
+					.AutoHeight()
+					.Padding(0.0f, 0.0f, 0.0f, 5.0f)
+					[
+						GetStageTitle(Excerpt, 0)
+					]
+					+SVerticalBox::Slot()
+					.HAlign(HAlign_Fill)
+					.AutoHeight()
+					[
+						Excerpt.Content.ToSharedRef()
+					];
+			}
+		}
+		break;
+	}
+
+	return SNullWidget::NullWidget;
+}
+
+
+FVector2D STutorialContent::GetPosition() const
+{
+	float XOffset = 0.0f;
+	switch(HorizontalAlignment.Get())
+	{
+	case HAlign_Left:
+		XOffset = -(ContentWidget->GetDesiredSize().X - ContentOffset);
+		break;
+	default:
+	case HAlign_Fill:
+	case HAlign_Center:
+		XOffset = (CachedGeometry.Size.X * 0.5f) - (ContentWidget->GetDesiredSize().X * 0.5f);
+		break;
+	case HAlign_Right:
+		XOffset = CachedGeometry.Size.X - ContentOffset;
+		break;
+	}
+
+	XOffset += WidgetOffset.Get().X;
+
+	float YOffset = 0.0f;
+	switch(VerticalAlignment.Get())
+	{
+	case VAlign_Top:
+		YOffset = -(ContentWidget->GetDesiredSize().Y - ContentOffset);
+		break;
+	default:
+	case VAlign_Fill:
+	case VAlign_Center:
+		YOffset = (CachedGeometry.Size.Y * 0.5f) - (ContentWidget->GetDesiredSize().Y * 0.5f);
+		break;
+	case VAlign_Bottom:
+		YOffset = (CachedGeometry.Size.Y - ContentOffset);
+		break;
+	}
+
+	YOffset += WidgetOffset.Get().Y;
+
+	// now build & clamp to area
+	FVector2D BaseOffset = FVector2D(CachedGeometry.AbsolutePosition.X + XOffset, CachedGeometry.AbsolutePosition.Y + YOffset);
+	BaseOffset.X = FMath::Clamp(BaseOffset.X, 0.0f, CachedWindowSize.X - ContentWidget->GetDesiredSize().X);
+	BaseOffset.Y = FMath::Clamp(BaseOffset.Y, 0.0f, CachedWindowSize.Y - ContentWidget->GetDesiredSize().Y);
+	return BaseOffset;
+}
+
+FVector2D STutorialContent::GetSize() const
+{
+	return ContentWidget->GetDesiredSize();
+}
+
+FReply STutorialContent::OnCloseButtonClicked()
+{
+	OnClosed.ExecuteIfBound();
+
+	return FReply::Handled();
+}
+
+EVisibility STutorialContent::GetCloseButtonVisibility() const
+{
+	return bIsStandalone ? EVisibility::Visible : EVisibility::Collapsed;
+}
+
+void STutorialContent::HandlePaintNamedWidget(TSharedRef<SWidget> InWidget, const FGeometry& InGeometry)
+{
+	switch(Anchor.Type)
+	{
+	case ETutorialAnchorIdentifier::NamedWidget:
+		if(Anchor.WrapperIdentifier == InWidget->GetTag())
+		{
+			bIsVisible = true;
+			CachedGeometry = InGeometry;
+		}
+		break;
+	}
+}
+
+void STutorialContent::HandleResetNamedWidget()
+{
+	bIsVisible = false;
+}
+
+void STutorialContent::HandleCacheWindowSize(const FVector2D& InWindowSize)
+{
+	CachedWindowSize = InWindowSize;
+}
+
+EVisibility STutorialContent::GetVisibility() const
+{
+	return bIsVisible || Anchor.Type == ETutorialAnchorIdentifier::None ? EVisibility::SelfHitTestInvisible : EVisibility::Collapsed;
+}
\ No newline at end of file
diff --git a/Engine/Source/Editor/IntroTutorials/Private/STutorialContent.h b/Engine/Source/Editor/IntroTutorials/Private/STutorialContent.h
new file mode 100644
--- /dev/null
+++ b/Engine/Source/Editor/IntroTutorials/Private/STutorialContent.h
@@ -0,0 +1,113 @@
+// Copyright 1998-2014 Epic Games, Inc. All Rights Reserved.
+
+#pragma once
+
+#include "IDocumentation.h"
+#include "IDocumentationPage.h"
+#include "EditorTutorial.h"
+
+/**
+ * The widget which displays 'floating' content 
+ */
+class STutorialContent : public SCompoundWidget
+{
+	SLATE_BEGIN_ARGS( STutorialContent )
+	{
+		_Visibility = EVisibility::SelfHitTestInvisible;
+		_IsStandalone = false;
+	}
+
+	/** Alignment of content relative to widget, note "Fill" is not supported */
+	SLATE_ATTRIBUTE(EVerticalAlignment, VAlign)
+
+	/** Alignment of content relative to widget, note "Fill" is not supported */
+	SLATE_ATTRIBUTE(EHorizontalAlignment, HAlign)
+
+	/** Offset form the widget we annotate */
+	SLATE_ATTRIBUTE(FVector2D, Offset)
+
+	/** Whether this a standalone widget (with its own close button) or part of a group of other widgets, paired with tutorial navigation */
+	SLATE_ARGUMENT(bool, IsStandalone)
+
+	/** Delegate fired when the close button is clicked */
+	SLATE_ARGUMENT(FSimpleDelegate, OnClosed)
+
+	/** Where text should be wrapped */
+	SLATE_ARGUMENT(float, WrapTextAt)
+
+	/** Anchor if required */
+	SLATE_ARGUMENT(FTutorialContentAnchor, Anchor)
+
+	SLATE_END_ARGS()
+
+	void Construct(const FArguments& InArgs, const FTutorialContent& InContent);
+
+	/** SWidget implementation */
+	virtual int32 OnPaint( const FPaintArgs& Args, const FGeometry& AllottedGeometry, const FSlateRect& MyClippingRect, FSlateWindowElementList& OutDrawElements, int32 LayerId, const FWidgetStyle& InWidgetStyle, bool bParentEnabled ) const override;
+
+	/** Helper function to generate widgets from an FTutorialContent struct */
+	static TSharedRef<SWidget> GenerateContentWidget(const FTutorialContent& InContent, float WrapTextAt, bool bForWidget, TSharedPtr<IDocumentationPage>& OutDocumentationPage);
+
+	/** Handle repositioning the widget */
+	FVector2D GetPosition() const;
+
+	/** Handle resizing the widget */
+	FVector2D GetSize() const;
+
+	/** Delegate handler called back from the overlay paint routines to flag whether we should paint as well (i.e. if this widget content is highlighted for the current stage) */
+	void HandlePaintNamedWidget(TSharedRef<SWidget> InWidget, const FGeometry& InGeometry);
+
+	/** Called back from the overlay paint routines to reset the flag we check for painting with */
+	void HandleResetNamedWidget();
+
+	/** Handle caching window size - called back from overlay paint routine */
+	void HandleCacheWindowSize(const FVector2D& InWindowSize);
+
+private:
+
+	/** Get the visibility of this content */
+	EVisibility GetVisibility() const;
+
+	/** Handle close button clicked - forward to delegate */
+	FReply OnCloseButtonClicked();
+
+	/** Get close button visibility - varies depending on whether we are standalone or not */
+	EVisibility GetCloseButtonVisibility() const;
+
+	/** Copy of the window size we were last draw at */
+	FVector2D CachedWindowSize;
+
+	/** Copy of the geometry our widget was last drawn with */
+	FGeometry CachedGeometry;
+
+	/** Container for widget content */
+	TSharedPtr<SWidget> ContentWidget;
+
+	/** Alignment of content relative to widget, note "Fill" is not supported */
+	TAttribute<EVerticalAlignment> VerticalAlignment;
+
+	/** Alignment of content relative to widget, note "Fill" is not supported */
+	TAttribute<EHorizontalAlignment> HorizontalAlignment;
+
+	/** Offset form the widget we annotate */
+	TAttribute<FVector2D> WidgetOffset;
+
+	/** Copy of the anchor for this tutorial content */
+	FTutorialContentAnchor Anchor; 
+
+	/** Whether this a standalone widget (with its own close button) or part of a group of other widgets, paired with tutorial navigation */
+	bool bIsStandalone;
+
+	/** Whether this overlay is currently visible */
+	bool bIsVisible;
+
+	/** Delegate fired when the close button is clicked */
+	FSimpleDelegate OnClosed;
+
+	/** Animation curves for displaying border */
+	FCurveSequence BorderPulseAnimation;
+	FCurveSequence BorderIntroAnimation;
+
+	/** Documentation page reference to use if we are displaying a UDN doc - we need this otherwise the page will be freed */
+	TSharedPtr<IDocumentationPage> DocumentationPage;
+};
\ No newline at end of file
diff --git a/Engine/Source/Editor/IntroTutorials/Private/STutorialNavigation.cpp b/Engine/Source/Editor/IntroTutorials/Private/STutorialNavigation.cpp
new file mode 100644
--- /dev/null
+++ b/Engine/Source/Editor/IntroTutorials/Private/STutorialNavigation.cpp
@@ -0,0 +1,121 @@
+// Copyright 1998-2014 Epic Games, Inc. All Rights Reserved.
+
+#include "IntroTutorialsPrivatePCH.h"
+#include "STutorialNavigation.h"
+
+void STutorialNavigation::Construct(const FArguments& InArgs)
+{
+	OnBackClicked = InArgs._OnBackClicked;
+	OnHomeClicked = InArgs._OnHomeClicked;
+	OnNextClicked = InArgs._OnNextClicked;
+	IsBackEnabled = InArgs._IsBackEnabled;
+	IsHomeEnabled = InArgs._IsHomeEnabled;
+	IsNextEnabled = InArgs._IsNextEnabled;
+	OnGetProgress = InArgs._OnGetProgress;
+
+	ChildSlot
+	[
+		SNew(SBorder)
+		.Padding(24.0f)
+		.BorderImage(FEditorStyle::GetBrush("Tutorials.Border"))
+		[
+			SNew(SVerticalBox)
+			+SVerticalBox::Slot()
+			.AutoHeight()
+			[
+				SNew(SHorizontalBox)
+				+SHorizontalBox::Slot()
+				.Padding(2.0f)
+				.AutoWidth()
+				[
+					SNew(SButton)
+					.OnClicked(this, &STutorialNavigation::OnBackButtonClicked)
+					.IsEnabled(InArgs._IsBackEnabled)
+					.ButtonStyle(&FEditorStyle::Get().GetWidgetStyle<FButtonStyle>("Tutorials.Browser.Button"))
+					.Content()
+					[
+						SNew(SImage)
+						.Image(FEditorStyle::GetBrush("Tutorials.Back"))
+						.ColorAndOpacity(this, &STutorialNavigation::GetBackButtonColor)
+					]
+				]
+				+SHorizontalBox::Slot()
+				.Padding(2.0f)
+				.AutoWidth()
+				[
+					SNew(SButton)
+					.OnClicked(this, &STutorialNavigation::OnHomeButtonClicked)
+					.IsEnabled(InArgs._IsHomeEnabled)
+					.ButtonStyle(&FEditorStyle::Get().GetWidgetStyle<FButtonStyle>("Tutorials.Browser.Button"))
+					.Content()
+					[
+						SNew(SImage)
+						.Image(FEditorStyle::GetBrush("Tutorials.Home"))
+						.ColorAndOpacity(this, &STutorialNavigation::GetHomeButtonColor)
+					]
+				]
+				+SHorizontalBox::Slot()
+				.Padding(2.0f)
+				.AutoWidth()
+				[
+					SNew(SButton)
+					.OnClicked(this, &STutorialNavigation::OnNextButtonClicked)
+					.IsEnabled(InArgs._IsNextEnabled)
+					.ButtonStyle(&FEditorStyle::Get().GetWidgetStyle<FButtonStyle>("Tutorials.Browser.Button"))
+					.Content()
+					[
+						SNew(SImage)
+						.Image(FEditorStyle::GetBrush("Tutorials.Next"))
+						.ColorAndOpacity(this, &STutorialNavigation::GetNextButtonColor)
+					]
+				]
+			]
+			+SVerticalBox::Slot()
+			.AutoHeight()
+			.MaxHeight(4.0f)
+			.Padding(0.0f, 2.0f)
+			[
+				SNew(SProgressBar)
+				.Percent(this, &STutorialNavigation::GetPercentComplete)
+			]
+		]
+	];
+}
+
+FReply STutorialNavigation::OnBackButtonClicked()
+{
+	OnBackClicked.ExecuteIfBound();
+	return FReply::Handled();
+}
+
+FSlateColor STutorialNavigation::GetBackButtonColor() const
+{
+	return IsBackEnabled.Get() ? FLinearColor(0.0f, 0.0f, 0.0f, 1.0f) : FLinearColor(1.0f, 1.0f, 1.0f, 0.25f);
+}
+
+FReply STutorialNavigation::OnHomeButtonClicked()
+{
+	OnHomeClicked.ExecuteIfBound();
+	return FReply::Handled();
+}
+
+FSlateColor STutorialNavigation::GetHomeButtonColor() const
+{
+	return IsHomeEnabled.Get() ? FLinearColor(0.0f, 0.0f, 0.0f, 1.0f) : FLinearColor(1.0f, 1.0f, 1.0f, 0.25f);
+}
+
+FReply STutorialNavigation::OnNextButtonClicked()
+{
+	OnNextClicked.ExecuteIfBound();
+	return FReply::Handled();
+}
+
+FSlateColor STutorialNavigation::GetNextButtonColor() const
+{
+	return IsNextEnabled.Get() ? FLinearColor(0.0f, 0.0f, 0.0f, 1.0f) : FLinearColor(1.0f, 1.0f, 1.0f, 0.25f);
+}
+
+TOptional<float> STutorialNavigation::GetPercentComplete() const
+{
+	return OnGetProgress.Get();
+}
\ No newline at end of file
diff --git a/Engine/Source/Editor/IntroTutorials/Private/STutorialNavigation.h b/Engine/Source/Editor/IntroTutorials/Private/STutorialNavigation.h
new file mode 100644
--- /dev/null
+++ b/Engine/Source/Editor/IntroTutorials/Private/STutorialNavigation.h
@@ -0,0 +1,56 @@
+// Copyright 1998-2014 Epic Games, Inc. All Rights Reserved.
+
+#pragma once
+
+/**
+ * The widget which displays floating navigation controls
+ */
+class STutorialNavigation : public SCompoundWidget
+{
+	SLATE_BEGIN_ARGS(STutorialNavigation) {}
+
+	SLATE_ARGUMENT(FSimpleDelegate, OnBackClicked)
+	SLATE_ARGUMENT(FSimpleDelegate, OnHomeClicked)
+	SLATE_ARGUMENT(FSimpleDelegate, OnNextClicked)
+	SLATE_ATTRIBUTE(bool, IsBackEnabled)
+	SLATE_ATTRIBUTE(bool, IsHomeEnabled)
+	SLATE_ATTRIBUTE(bool, IsNextEnabled)
+	SLATE_ATTRIBUTE(float, OnGetProgress)
+	SLATE_ARGUMENT(TWeakPtr<SWindow>, ParentWindow)
+
+	SLATE_END_ARGS()
+
+	void Construct(const FArguments& InArgs);
+
+protected:
+	/** Handle back button being clicked */
+	FReply OnBackButtonClicked();
+
+	/** Handle back button color - used to dim the button when disabled */
+	FSlateColor GetBackButtonColor() const;
+
+	/** Handle home button being clicked */
+	FReply OnHomeButtonClicked();
+
+	/** Handle home button color - used to dim the button when disabled */
+	FSlateColor GetHomeButtonColor() const;
+
+	/** Handle next button being clicked */
+	FReply OnNextButtonClicked();
+
+	/** Handle next button color - used to dim the button when disabled */
+	FSlateColor GetNextButtonColor() const;
+
+	/** Handle display of the progress bar */
+	TOptional<float> GetPercentComplete() const;
+
+private:
+
+	FSimpleDelegate OnBackClicked;
+	FSimpleDelegate OnHomeClicked;
+	FSimpleDelegate OnNextClicked;
+	TAttribute<bool> IsBackEnabled;
+	TAttribute<bool> IsHomeEnabled;
+	TAttribute<bool> IsNextEnabled;
+	TAttribute<float> OnGetProgress;
+};
\ No newline at end of file
diff --git a/Engine/Source/Editor/IntroTutorials/Private/STutorialOverlay.cpp b/Engine/Source/Editor/IntroTutorials/Private/STutorialOverlay.cpp
new file mode 100644
--- /dev/null
+++ b/Engine/Source/Editor/IntroTutorials/Private/STutorialOverlay.cpp
@@ -0,0 +1,131 @@
+// Copyright 1998-2014 Epic Games, Inc. All Rights Reserved.
+
+#include "IntroTutorialsPrivatePCH.h"
+#include "STutorialOverlay.h"
+#include "STutorialContent.h"
+#include "EditorTutorial.h"
+#include "IntroTutorials.h"
+
+void STutorialOverlay::Construct(const FArguments& InArgs, FTutorialStage* const InStage)
+{
+	ParentWindow = InArgs._ParentWindow;
+	bIsStandalone = InArgs._IsStandalone;
+	OnClosed = InArgs._OnClosed;
+
+	TSharedPtr<SOverlay> Overlay;
+
+	ChildSlot
+	[
+		SAssignNew(Overlay, SOverlay)
+		+SOverlay::Slot()
+		[
+			SAssignNew(OverlayCanvas, SCanvas)
+		]
+	];
+
+	if(InStage != nullptr)
+	{
+		// add non-widget content, if any
+		if(InArgs._AllowNonWidgetContent && InStage->Content.Type != ETutorialContent::None)
+		{
+			Overlay->AddSlot()
+			[
+				SNew(SHorizontalBox)
+				+SHorizontalBox::Slot()
+				.VAlign(VAlign_Center)
+				.HAlign(HAlign_Center)
+				[
+					SNew(STutorialContent, InStage->Content)
+					.OnClosed(InArgs._OnClosed)
+					.IsStandalone(InArgs._IsStandalone)
+					.WrapTextAt(600.0f)
+				]
+			];
+		}
+
+		if(InStage->WidgetContent.Num() > 0)
+		{
+			FIntroTutorials& IntroTutorials = FModuleManager::Get().GetModuleChecked<FIntroTutorials>("IntroTutorials");
+
+			// now add canvas slots for widget-bound content
+			for(const FTutorialWidgetContent& WidgetContent : InStage->WidgetContent)
+			{
+				if(WidgetContent.Content.Type != ETutorialContent::None)
+				{
+					TSharedPtr<STutorialContent> ContentWidget = 
+						SNew(STutorialContent, WidgetContent.Content)
+						.HAlign(WidgetContent.HorizontalAlignment)
+						.VAlign(WidgetContent.VerticalAlignment)
+						.Offset(WidgetContent.Offset)
+						.IsStandalone(bIsStandalone)
+						.OnClosed(OnClosed)
+						.WrapTextAt(WidgetContent.ContentWidth)
+						.Anchor(WidgetContent.WidgetAnchor);
+
+					OverlayCanvas->AddSlot()
+					.Position(TAttribute<FVector2D>::Create(TAttribute<FVector2D>::FGetter::CreateSP(ContentWidget.Get(), &STutorialContent::GetPosition)))
+					.Size(TAttribute<FVector2D>::Create(TAttribute<FVector2D>::FGetter::CreateSP(ContentWidget.Get(), &STutorialContent::GetSize)))
+					[
+						ContentWidget.ToSharedRef()
+					];
+
+					OnPaintNamedWidget.AddSP(ContentWidget.Get(), &STutorialContent::HandlePaintNamedWidget);
+					OnResetNamedWidget.AddSP(ContentWidget.Get(), &STutorialContent::HandleResetNamedWidget);
+					OnCacheWindowSize.AddSP(ContentWidget.Get(), &STutorialContent::HandleCacheWindowSize);
+				}
+			}
+		}
+	}
+}
+
+int32 STutorialOverlay::OnPaint( const FPaintArgs& Args, const FGeometry& AllottedGeometry, const FSlateRect& MyClippingRect, FSlateWindowElementList& OutDrawElements, int32 LayerId, const FWidgetStyle& InWidgetStyle, bool bParentEnabled ) const
+{
+	if(ParentWindow.IsValid())
+	{
+		TSharedPtr<SWindow> PinnedWindow = ParentWindow.Pin();
+		OnResetNamedWidget.Broadcast();
+		OnCacheWindowSize.Broadcast(PinnedWindow->GetWindowGeometryInWindow().Size);
+		LayerId = TraverseWidgets(PinnedWindow.ToSharedRef(), PinnedWindow->GetWindowGeometryInWindow(), MyClippingRect, OutDrawElements, LayerId);
+	}
+	
+	return SCompoundWidget::OnPaint(Args, AllottedGeometry, MyClippingRect, OutDrawElements, LayerId, InWidgetStyle, bParentEnabled);
+}
+
+int32 STutorialOverlay::TraverseWidgets(TSharedRef<SWidget> InWidget, const FGeometry& InGeometry, const FSlateRect& MyClippingRect, FSlateWindowElementList& OutDrawElements, int32 LayerId) const
+{
+	const FName Tag = InWidget->GetTag();
+	if(Tag != NAME_None)
+	{
+		// we are a named widget - ask it to draw
+		OnPaintNamedWidget.Broadcast(InWidget, InGeometry);
+
+		// if we are picking, we need to draw an outline here
+		FName WidgetNameToHighlight = NAME_None;
+		bool bIsPicking = false;
+		FIntroTutorials& IntroTutorials = FModuleManager::Get().GetModuleChecked<FIntroTutorials>("IntroTutorials");
+		if(IntroTutorials.OnIsPicking().IsBound())
+		{
+			bIsPicking = IntroTutorials.OnIsPicking().Execute(WidgetNameToHighlight);
+		}
+	
+		if(WidgetNameToHighlight != NAME_None || bIsPicking)
+		{
+			const bool bHighlight = WidgetNameToHighlight != NAME_None && WidgetNameToHighlight == Tag;
+			if(bIsPicking || (!bIsPicking && bHighlight))
+			{
+				const FLinearColor Color = bIsPicking && bHighlight ? FLinearColor::Green : FLinearColor::White;
+				FSlateDrawElement::MakeBox(OutDrawElements, LayerId++, InGeometry.ToPaintGeometry(), FCoreStyle::Get().GetBrush(TEXT("Debug.Border")), MyClippingRect, ESlateDrawEffect::None, Color);
+			}
+		}	
+	}
+
+	FArrangedChildren ArrangedChildren(EVisibility::Visible);
+	InWidget->ArrangeChildren(InGeometry, ArrangedChildren);
+	for(int32 ChildIndex = 0; ChildIndex < ArrangedChildren.Num(); ChildIndex++)
+	{
+		const FArrangedWidget& ArrangedWidget = ArrangedChildren(ChildIndex);
+		LayerId = TraverseWidgets(ArrangedWidget.Widget, ArrangedWidget.Geometry, MyClippingRect, OutDrawElements, LayerId);
+	}
+
+	return LayerId;
+}
\ No newline at end of file
diff --git a/Engine/Source/Editor/IntroTutorials/Private/STutorialOverlay.h b/Engine/Source/Editor/IntroTutorials/Private/STutorialOverlay.h
new file mode 100644
--- /dev/null
+++ b/Engine/Source/Editor/IntroTutorials/Private/STutorialOverlay.h
@@ -0,0 +1,70 @@
+// Copyright 1998-2014 Epic Games, Inc. All Rights Reserved.
+
+#pragma once
+
+#include "EditorTutorial.h"
+
+class STutorialContent;
+
+/** Delegate used when drawing/arranging widgets */
+DECLARE_MULTICAST_DELEGATE_TwoParams(FOnPaintNamedWidget, TSharedRef<SWidget> /*InWidget*/, const FGeometry& /*InGeometry*/);
+
+/** Delegate used to inform widgets of the current window size, so they can auto-adjust layout */
+DECLARE_MULTICAST_DELEGATE_OneParam(FOnCacheWindowSize, const FVector2D& /*InWindowSize*/);
+
+/**
+ * The widget which displays multiple 'floating' pieces of content overlaid onto the editor
+ */
+class STutorialOverlay : public SCompoundWidget
+{
+public:
+	SLATE_BEGIN_ARGS( STutorialOverlay ) 
+	{
+		_Visibility = EVisibility::SelfHitTestInvisible;
+	}
+
+	/** The window this content is displayed over */
+	SLATE_ARGUMENT(TWeakPtr<SWindow>, ParentWindow)
+
+	/** Whether this a standalone widget (with its own close button) or part of a group of other widgets, paired with tutorial navigation */
+	SLATE_ARGUMENT(bool, IsStandalone)
+
+	/** Delegate fired when the close button is clicked */
+	SLATE_ARGUMENT(FSimpleDelegate, OnClosed)
+
+	/** Whether we can show full window content in this overlay (i.e. in the same window as the navigation controls) */
+	SLATE_ARGUMENT(bool, AllowNonWidgetContent)
+
+	SLATE_END_ARGS()
+
+	void Construct(const FArguments& InArgs, FTutorialStage* const InStage);
+
+	/** SWidget implementation */
+	virtual int32 OnPaint( const FPaintArgs& Args, const FGeometry& AllottedGeometry, const FSlateRect& MyClippingRect, FSlateWindowElementList& OutDrawElements, int32 LayerId, const FWidgetStyle& InWidgetStyle, bool bParentEnabled ) const override;
+
+private:
+	/** Recursive function used to re-generate widget geometry and forward the geometry of named widgets onto their respective content */
+	int32 TraverseWidgets(TSharedRef<SWidget> InWidget, const FGeometry& InGeometry, const FSlateRect& MyClippingRect, FSlateWindowElementList& OutDrawElements, int32 LayerId) const;
+
+private:
+	/** Reference to the canvas we use to position our content widgets */
+	TSharedPtr<SCanvas> OverlayCanvas;
+
+	/** The window this content is displayed over */
+	TWeakPtr<SWindow> ParentWindow;
+
+	/** Whether this a standalone widget (with its own close button) or part of a group of other widgets, paired with tutorial navigation */
+	bool bIsStandalone;
+
+	/** Delegate fired when a cos ebutton is clicked in tutorial content */
+	FSimpleDelegate OnClosed;
+
+	/** Delegate used when drawing/arranging widgets */
+	FOnPaintNamedWidget OnPaintNamedWidget;
+
+	/** Delegate used to reset drawing of named widgets */
+	FSimpleMulticastDelegate OnResetNamedWidget;
+
+	/** Delegate used to inform widgets of the current window size, so they can auto-adjust layout */
+	FOnCacheWindowSize OnCacheWindowSize;
+};
\ No newline at end of file
diff --git a/Engine/Source/Editor/IntroTutorials/Private/STutorialRoot.cpp b/Engine/Source/Editor/IntroTutorials/Private/STutorialRoot.cpp
new file mode 100644
--- /dev/null
+++ b/Engine/Source/Editor/IntroTutorials/Private/STutorialRoot.cpp
@@ -0,0 +1,199 @@
+// Copyright 1998-2014 Epic Games, Inc. All Rights Reserved.
+
+#include "IntroTutorialsPrivatePCH.h"
+#include "STutorialRoot.h"
+#include "SEditorTutorials.h"
+#include "EditorTutorialSettings.h"
+
+#define LOCTEXT_NAMESPACE "STutorialRoot"
+
+void STutorialRoot::Construct(const FArguments& InArgs)
+{
+	CurrentTutorial = nullptr;
+	CurrentTutorialStage = 0;
+
+	ChildSlot
+	[
+		SNullWidget::NullWidget
+	];
+}
+
+void STutorialRoot::MaybeAddOverlay(TSharedRef<SWindow> InWindow)
+{
+	if(InWindow->HasOverlay())
+	{
+		// check we dont already have a widget overlay for this window
+		TWeakPtr<SEditorTutorials>* FoundWidget = TutorialWidgets.Find(InWindow);
+		if(FoundWidget == nullptr)
+		{
+			TSharedPtr<SEditorTutorials> TutorialWidget = nullptr;
+			InWindow->AddOverlaySlot()
+			[
+				SNew(SHorizontalBox)
+				+SHorizontalBox::Slot()
+				.VAlign(VAlign_Fill)
+				.HAlign(HAlign_Fill)
+				[
+					SAssignNew(TutorialWidget, SEditorTutorials)
+					.ParentWindow(InWindow)
+					.OnNextClicked(FOnNextClicked::CreateSP(this, &STutorialRoot::HandleNextClicked))
+					.OnBackClicked(FSimpleDelegate::CreateSP(this, &STutorialRoot::HandleBackClicked))
+					.OnHomeClicked(FSimpleDelegate::CreateSP(this, &STutorialRoot::HandleHomeClicked))
+					.OnGetCurrentTutorial(FOnGetCurrentTutorial::CreateSP(this, &STutorialRoot::HandleGetCurrentTutorial))
+					.OnGetCurrentTutorialStage(FOnGetCurrentTutorialStage::CreateSP(this, &STutorialRoot::HandleGetCurrentTutorialStage))
+					.OnLaunchTutorial(FOnLaunchTutorial::CreateSP(this, &STutorialRoot::LaunchTutorial))
+				]
+			];
+
+			FoundWidget = &TutorialWidgets.Add(InWindow, TutorialWidget);
+
+			FoundWidget->Pin()->RebuildCurrentContent();
+		}
+	}
+
+	TArray<TSharedRef<SWindow>> ChildWindows = InWindow->GetChildWindows();
+	for(auto& ChildWindow : ChildWindows)
+	{
+		MaybeAddOverlay(ChildWindow);
+	}
+}
+
+void STutorialRoot::Tick(const FGeometry& AllottedGeometry, const double InCurrentTime, const float InDeltaTime)
+{
+	TArray<TSharedRef<SWindow>> Windows = FSlateApplication::Get().GetInteractiveTopLevelWindows();
+	for(auto& Window : Windows)
+	{
+		MaybeAddOverlay(Window);
+	}
+}
+
+void STutorialRoot::SummonTutorialBrowser(TSharedRef<SWindow> InWindow, const FString& InFilter)
+{
+	TWeakPtr<SEditorTutorials>* FoundWidget = TutorialWidgets.Find(InWindow);
+	if(FoundWidget != nullptr && FoundWidget->IsValid())
+	{
+		FoundWidget->Pin()->ShowBrowser(InFilter);
+	}
+}
+
+void STutorialRoot::LaunchTutorial(UEditorTutorial* InTutorial, bool bInRestart, TWeakPtr<SWindow> InNavigationWindow)
+{
+	if(InTutorial != nullptr)
+	{
+		CurrentTutorial = InTutorial;
+
+		bool bHaveSeenTutorial = false;
+		CurrentTutorialStage = bInRestart ? 0 : GetDefault<UEditorTutorialSettings>()->GetProgress(CurrentTutorial, bHaveSeenTutorial);
+
+		// launch tutorial for all windows we wrap - any tutorial can display over any window
+		for(auto& TutorialWidget : TutorialWidgets)
+		{
+			if(TutorialWidget.Value.IsValid())
+			{
+				bool bIsNavigationWindow = InNavigationWindow.IsValid() && TutorialWidget.Value.Pin()->GetParentWindow() == InNavigationWindow.Pin();
+				TutorialWidget.Value.Pin()->LaunchTutorial(bIsNavigationWindow);
+			}
+		}
+	}
+}
+
+void STutorialRoot::HandleNextClicked(TWeakPtr<SWindow> InNavigationWindow)
+{
+	GoToNextStage(InNavigationWindow);
+
+	for(auto& TutorialWidget : TutorialWidgets)
+	{
+		if(TutorialWidget.Value.IsValid())
+		{
+			TSharedPtr<SEditorTutorials> PinnedTutorialWidget = TutorialWidget.Value.Pin();
+			PinnedTutorialWidget->RebuildCurrentContent();
+		}
+	}
+}
+
+void STutorialRoot::HandleBackClicked()
+{
+	GoToPreviousStage();
+
+	for(auto& TutorialWidget : TutorialWidgets)
+	{
+		if(TutorialWidget.Value.IsValid())
+		{
+			TSharedPtr<SEditorTutorials> PinnedTutorialWidget = TutorialWidget.Value.Pin();
+			PinnedTutorialWidget->RebuildCurrentContent();
+		}
+	}
+}
+
+void STutorialRoot::HandleHomeClicked()
+{
+	if(CurrentTutorial != nullptr)
+	{
+		GetMutableDefault<UEditorTutorialSettings>()->RecordProgress(CurrentTutorial, CurrentTutorialStage);
+		GetMutableDefault<UEditorTutorialSettings>()->SaveProgress();
+	}
+
+	CurrentTutorial = nullptr;
+	CurrentTutorialStage = 0;
+
+	for(auto& TutorialWidget : TutorialWidgets)
+	{
+		if(TutorialWidget.Value.IsValid())
+		{
+			TSharedPtr<SEditorTutorials> PinnedTutorialWidget = TutorialWidget.Value.Pin();
+			PinnedTutorialWidget->RebuildCurrentContent();
+		}
+	}
+}
+
+UEditorTutorial* STutorialRoot::HandleGetCurrentTutorial()
+{
+	return CurrentTutorial;
+}
+
+int32 STutorialRoot::HandleGetCurrentTutorialStage()
+{
+	return CurrentTutorialStage;
+}
+
+void STutorialRoot::AddReferencedObjects( FReferenceCollector& Collector )
+{
+	if(CurrentTutorial != nullptr)
+	{
+		Collector.AddReferencedObject(CurrentTutorial);
+	}
+}
+
+void STutorialRoot::GoToPreviousStage()
+{
+	if(CurrentTutorial != nullptr)
+	{
+		CurrentTutorialStage = FMath::Max(0, CurrentTutorialStage - 1);
+	}
+}
+
+void STutorialRoot::GoToNextStage(TWeakPtr<SWindow> InNavigationWindow)
+{
+	if(CurrentTutorial != nullptr)
+	{
+		if(CurrentTutorialStage + 1 >= CurrentTutorial->Stages.Num() && FName(*CurrentTutorial->NextTutorial.AssetLongPathname) != NAME_None)
+		{
+			TSubclassOf<UEditorTutorial> NextTutorialClass = LoadClass<UEditorTutorial>(NULL, *CurrentTutorial->NextTutorial.AssetLongPathname, NULL, LOAD_None, NULL);
+			if(NextTutorialClass != nullptr)
+			{
+				LaunchTutorial(NextTutorialClass->GetDefaultObject<UEditorTutorial>(), true, InNavigationWindow);
+			}
+			else
+			{
+				FSlateNotificationManager::Get().AddNotification(FNotificationInfo(FText::Format(LOCTEXT("TutorialNotFound", "Could not start next tutorial {0}"), FText::FromString(CurrentTutorial->NextTutorial.AssetLongPathname))));
+			}
+		}
+		else
+		{
+			CurrentTutorialStage = FMath::Min(CurrentTutorialStage + 1, CurrentTutorial->Stages.Num() - 1);
+			GetMutableDefault<UEditorTutorialSettings>()->RecordProgress(CurrentTutorial, CurrentTutorialStage);
+		}
+	}
+}
+
+#undef LOCTEXT_NAMESPACE
\ No newline at end of file
diff --git a/Engine/Source/Editor/IntroTutorials/Private/STutorialRoot.h b/Engine/Source/Editor/IntroTutorials/Private/STutorialRoot.h
new file mode 100644
--- /dev/null
+++ b/Engine/Source/Editor/IntroTutorials/Private/STutorialRoot.h
@@ -0,0 +1,73 @@
+// Copyright 1998-2014 Epic Games, Inc. All Rights Reserved.
+
+#pragma once
+
+class UEditorTutorial;
+class SEditorTutorials;
+
+/**
+ * The widget which simply monitors windows in its tick function to see if we need to attach
+ * a tutorial overlay.
+ */
+class STutorialRoot : public SCompoundWidget, public FGCObject
+{
+public:
+	SLATE_BEGIN_ARGS( STutorialRoot ) 
+	{
+		_Visibility = EVisibility::HitTestInvisible;
+	}
+
+	SLATE_END_ARGS()
+
+	void Construct(const FArguments& InArgs);
+
+	/** SWidget implementation */
+	virtual void Tick( const FGeometry& AllottedGeometry, const double InCurrentTime, const float InDeltaTime ) override;
+
+	/** FGCObject implementation */
+	virtual void AddReferencedObjects( FReferenceCollector& Collector ) override;
+
+	/** Launch the specified tutorial from the specified window */
+	void LaunchTutorial(UEditorTutorial* InTutorial, bool bInRestart, TWeakPtr<SWindow> InNavigationWindow);
+
+	/** Summon the browser widget for the specified window */
+	void SummonTutorialBrowser(TSharedRef<SWindow> InWindow, const FString& InFilter = TEXT(""));
+
+	/** Reload tutorials that we know about */
+	void ReloadTutorials();
+
+private:
+	/** Handle when the next button is clicked - forward navigation to other overlays */
+	void HandleNextClicked(TWeakPtr<SWindow> InNavigationWindow);
+
+	/** Handle when the back button is clicked - forward navigation to other overlays */
+	void HandleBackClicked();
+
+	/** Handle when the home button is clicked - forward navigation to other overlays */
+	void HandleHomeClicked();
+
+	/** Handle retrieving the current tutorial */
+	UEditorTutorial* HandleGetCurrentTutorial();
+
+	/** Handle retrieving the current tutorial stage */
+	int32 HandleGetCurrentTutorialStage();
+
+	/** Got to the previous stage in the current tutorial */
+	void GoToPreviousStage();
+
+	/** Got to the next stage in the current tutorial */
+	void GoToNextStage(TWeakPtr<SWindow> InNavigationWindow);
+
+	/** Function called on Tick() to check active windows for whether they need an overlay adding */
+	void MaybeAddOverlay(TSharedRef<SWindow> InWindow);
+
+private:
+	/** Container widgets, inserted into window overlays */
+	TMap<TWeakPtr<SWindow>, TWeakPtr<SEditorTutorials>> TutorialWidgets;
+
+	/** Tutorial we are currently viewing */
+	UEditorTutorial* CurrentTutorial;
+
+	/** Current stage of tutorial */
+	int32 CurrentTutorialStage;
+};
\ No newline at end of file
diff --git a/Engine/Source/Editor/IntroTutorials/Private/STutorialsBrowser.cpp b/Engine/Source/Editor/IntroTutorials/Private/STutorialsBrowser.cpp
new file mode 100644
--- /dev/null
+++ b/Engine/Source/Editor/IntroTutorials/Private/STutorialsBrowser.cpp
@@ -0,0 +1,695 @@
+// Copyright 1998-2014 Epic Games, Inc. All Rights Reserved.
+
+#include "IntroTutorialsPrivatePCH.h"
+#include "STutorialsBrowser.h"
+#include "EditorTutorial.h"
+#include "STutorialContent.h"
+#include "TutorialSettings.h"
+#include "EditorTutorialSettings.h"
+
+#define LOCTEXT_NAMESPACE "TutorialsBrowser"
+
+class FTutorialListEntry_Tutorial;
+
+DECLARE_DELEGATE_OneParam(FOnCategorySelected, const FString& /* InCategory */);
+
+class FTutorialListEntry_Category : public ITutorialListEntry, public TSharedFromThis<FTutorialListEntry_Category>
+{
+public:
+	FTutorialListEntry_Category(FOnCategorySelected InOnCategorySelected)
+		:  OnCategorySelected(InOnCategorySelected)
+	{}
+
+	FTutorialListEntry_Category(const FTutorialCategory& InCategory, FOnCategorySelected InOnCategorySelected, const TAttribute<FText>& InHighlightText)
+		: Category(InCategory)
+		, OnCategorySelected(InOnCategorySelected)
+		, HighlightText(InHighlightText)
+	{
+		if(!Category.Identifier.IsEmpty())
+		{
+			int32 Index = INDEX_NONE;
+			if(Category.Identifier.FindLastChar(TEXT('.'), Index))
+			{
+				CategoryName = Category.Identifier.RightChop(Index + 1);
+			}
+			else
+			{
+				CategoryName = Category.Identifier;
+			}
+		}
+	}
+
+	virtual TSharedRef<ITableRow> OnGenerateTutorialRow(const TSharedRef<STableViewBase>& OwnerTable) const override
+	{
+		return SNew(STableRow<TSharedPtr<ITutorialListEntry>>, OwnerTable)
+		[
+			SNew(SButton)
+			.OnClicked(this, &FTutorialListEntry_Category::OnClicked)
+			.ButtonStyle(&FEditorStyle::Get().GetWidgetStyle<FButtonStyle>("Tutorials.Browser.Button"))
+			.Content()
+			[
+				SNew(SHorizontalBox)
+				+SHorizontalBox::Slot()
+				.AutoWidth()
+				.VAlign(VAlign_Center)
+				.HAlign(HAlign_Center)
+				.Padding(16.0f)
+				[
+					SNew(SImage)
+					.Image(FEditorStyle::Get().GetBrush(FName(*Category.Icon)))
+				]
+				+SHorizontalBox::Slot()
+				.FillWidth(1.0f)
+				.MaxWidth(400.0f)
+				.HAlign(HAlign_Left)
+				.VAlign(VAlign_Center)
+				[
+					SNew(SVerticalBox)
+					+SVerticalBox::Slot()
+					.HAlign(HAlign_Left)
+					.AutoHeight()
+					[
+						SNew(STextBlock)
+						.Text(!Category.Title.IsEmpty() ? Category.Title : FText::FromString(CategoryName))
+						.TextStyle(&FEditorStyle::Get().GetWidgetStyle<FTextBlockStyle>("Tutorials.Browser.SummaryHeader"))
+						.HighlightText(HighlightText)
+						.HighlightColor(FEditorStyle::Get().GetColor("Tutorials.Browser.HighlightTextColor"))
+						.HighlightShape(FEditorStyle::Get().GetBrush("TextBlock.HighlightShape"))
+					]
+					+SVerticalBox::Slot()
+					.HAlign(HAlign_Left)
+					.FillHeight(1.0f)
+					[
+						SNew(STextBlock)
+						.WrapTextAt(350.0f)
+						.Text(Category.Description)
+						.TextStyle(&FEditorStyle::Get().GetWidgetStyle<FTextBlockStyle>("Tutorials.Browser.SummaryText"))
+						.HighlightText(HighlightText)
+						.HighlightColor(FEditorStyle::Get().GetColor("Tutorials.Browser.HighlightTextColor"))
+						.HighlightShape(FEditorStyle::Get().GetBrush("TextBlock.HighlightShape"))
+					]
+				]
+				+SHorizontalBox::Slot()
+				.AutoWidth()
+				.HAlign(HAlign_Right)
+				.VAlign(VAlign_Center)
+				[
+					SNew(SImage)
+					.Visibility(this, &FTutorialListEntry_Category::OnGetArrowVisibility)
+					.Image(FEditorStyle::Get().GetBrush("Tutorials.Browser.CategoryArrow"))
+				]
+			]
+		];
+	}
+
+	bool PassesFilter(const FString& InCategoryFilter, const FString& InFilter) const override
+	{
+		const FString Title = !Category.Title.IsEmpty() ? Category.Title.ToString() : CategoryName;
+		const bool bPassesFilter = InFilter.IsEmpty() || (Title.Contains(InFilter) || Category.Description.ToString().Contains(InFilter));
+		const bool bPassesCategory = InCategoryFilter.IsEmpty() || Category.Identifier.StartsWith(InCategoryFilter);
+		return bPassesFilter && bPassesCategory;
+	}
+
+	void AddSubCategory(TSharedPtr<FTutorialListEntry_Category> InSubCategory)
+	{
+		SubCategories.Add(InSubCategory);
+	}
+
+	void AddTutorial(TSharedPtr<FTutorialListEntry_Tutorial> InTutorial)
+	{
+		Tutorials.Add(InTutorial);
+	}
+
+	FReply OnClicked() const
+	{
+		if(SubCategories.Num() > 0 || Tutorials.Num() > 0)
+		{
+			OnCategorySelected.ExecuteIfBound(Category.Identifier);
+		}
+		return FReply::Handled();
+	}
+
+	EVisibility OnGetArrowVisibility() const
+	{
+		return (SubCategories.Num() > 0 || Tutorials.Num() > 0) ? EVisibility::Visible : EVisibility::Collapsed;
+	}
+
+public:
+	/** Copy of the category info */
+	FTutorialCategory Category;
+
+	/** Parent category */
+	TWeakPtr<ITutorialListEntry> ParentCategory;
+
+	/** Sub-categories */
+	TArray<TSharedPtr<ITutorialListEntry>> SubCategories;
+
+	/** Tutorials in this category */
+	TArray<TSharedPtr<ITutorialListEntry>> Tutorials;
+
+	/** Selection delegate */
+	FOnCategorySelected OnCategorySelected;
+
+	/** Name of the category, empty if this category is at the root */
+	FString CategoryName;
+
+	/** Text to highlight */
+	TAttribute<FText> HighlightText;
+};
+
+DECLARE_DELEGATE_TwoParams(FOnTutorialSelected, UEditorTutorial* /* InTutorial */, bool /* bRestart */ );
+
+class FTutorialListEntry_Tutorial : public ITutorialListEntry, public TSharedFromThis<FTutorialListEntry_Tutorial>
+{
+public:
+	FTutorialListEntry_Tutorial(UEditorTutorial* InTutorial, FOnTutorialSelected InOnTutorialSelected, const TAttribute<FText>& InHighlightText)
+		: Tutorial(InTutorial)
+		, OnTutorialSelected(InOnTutorialSelected)
+		, HighlightText(InHighlightText)
+	{}
+
+	virtual TSharedRef<ITableRow> OnGenerateTutorialRow(const TSharedRef<STableViewBase>& OwnerTable) const override
+	{
+		return SNew(STableRow<TSharedPtr<ITutorialListEntry>>, OwnerTable)
+		[
+			SAssignNew(LaunchButton, SButton)
+			.OnClicked(this, &FTutorialListEntry_Tutorial::OnClicked, false)
+			.ButtonStyle(&FEditorStyle::Get().GetWidgetStyle<FButtonStyle>("Tutorials.Browser.Button"))
+			.Content()
+			[
+				SNew(SVerticalBox)
+				+SVerticalBox::Slot()
+				.AutoHeight()
+				[
+					SNew(SHorizontalBox)
+					+SHorizontalBox::Slot()
+					.AutoWidth()
+					.VAlign(VAlign_Center)
+					.HAlign(HAlign_Center)
+					[
+						SNew(SBox)
+						.WidthOverride(96.0f)
+						.HeightOverride(96.0f)
+						.VAlign(VAlign_Center)
+						.HAlign(HAlign_Center)
+						[
+							SNew(SImage)
+							.Image(FEditorStyle::Get().GetBrush(FName(*Tutorial->Icon)))
+						]
+					]
+					+SHorizontalBox::Slot()
+					.FillWidth(1.0f)
+					.MaxWidth(400.0f)
+					.HAlign(HAlign_Left)
+					.VAlign(VAlign_Center)
+					[
+						SNew(SVerticalBox)
+						+SVerticalBox::Slot()
+						.AutoHeight()
+						[
+							SNew(SHorizontalBox)
+							+SHorizontalBox::Slot()
+							.FillWidth(1.0f)
+							[
+								SNew(STextBlock)
+								.Text(Tutorial->Title)
+								.TextStyle(&FEditorStyle::Get().GetWidgetStyle<FTextBlockStyle>("Tutorials.Browser.SummaryHeader"))
+								.HighlightText(HighlightText)
+								.HighlightColor(FEditorStyle::Get().GetColor("Tutorials.Browser.HighlightTextColor"))
+								.HighlightShape(FEditorStyle::Get().GetBrush("TextBlock.HighlightShape"))
+							]
+							+SHorizontalBox::Slot()
+							.AutoWidth()
+							.VAlign(VAlign_Center)
+							[
+								SNew(SButton)
+								.Visibility(this, &FTutorialListEntry_Tutorial::GetProgressVisibility)
+								.OnClicked(this, &FTutorialListEntry_Tutorial::OnClicked, true)
+								.ButtonStyle(&FEditorStyle::Get().GetWidgetStyle<FButtonStyle>("Tutorials.Browser.Button"))
+								.Content()
+								[
+									SNew(SImage)
+									.Image(FEditorStyle::GetBrush("Tutorials.Browser.RestartButton"))
+								]
+							]
+						]
+						+SVerticalBox::Slot()
+						.AutoHeight()
+						[
+							SNew(SBox)
+							.Visibility(this, &FTutorialListEntry_Tutorial::GetProgressVisibility)
+							.HeightOverride(3.0f)
+							[
+								SNew(SProgressBar)
+								.Percent(this, &FTutorialListEntry_Tutorial::GetProgress)
+							]
+						]
+						+SVerticalBox::Slot()
+						.HAlign(HAlign_Left)
+						.FillHeight(1.0f)
+						[
+							STutorialContent::GenerateContentWidget(Tutorial->SummaryContent, 350.0f, false, DocumentationPage)
+						]
+					]
+				]
+			]
+		];
+	}
+
+	bool PassesFilter(const FString& InCategoryFilter, const FString& InFilter) const override
+	{
+		const bool bPassesFilter = InFilter.IsEmpty() || (Tutorial->Title.ToString().Contains(InFilter) || Tutorial->SummaryContent.Text.ToString().Contains(InFilter));
+		const bool bPassesCategory = InCategoryFilter.IsEmpty() || Tutorial->Category.StartsWith(InCategoryFilter);
+
+		return bPassesFilter && bPassesCategory;
+	}
+
+	FReply OnClicked(bool bRestart) const
+	{
+		OnTutorialSelected.ExecuteIfBound(Tutorial, bRestart);
+
+		return FReply::Handled();
+	}
+
+	TOptional<float> GetProgress() const
+	{
+		bool bHaveSeenTutorial = false;
+		const int32 CurrentStage = GetDefault<UEditorTutorialSettings>()->GetProgress(Tutorial, bHaveSeenTutorial);
+		return (Tutorial->Stages.Num() > 0) ? (float)(CurrentStage + 1) / (float)Tutorial->Stages.Num() : 0.0f;
+	}
+
+	EVisibility GetProgressVisibility() const
+	{
+		if(LaunchButton->IsHovered())
+		{
+			bool bHaveSeenTutorial = false;
+			const int32 CurrentStage = GetDefault<UEditorTutorialSettings>()->GetProgress(Tutorial, bHaveSeenTutorial);
+			return LaunchButton->IsHovered() && bHaveSeenTutorial ? EVisibility::Visible : EVisibility::Hidden;
+		}
+
+		return EVisibility::Hidden;
+	}
+
+public:
+	/** Parent category */
+	TWeakPtr<ITutorialListEntry> ParentCategory;
+
+	/** Tutorial that we will launch */
+	UEditorTutorial* Tutorial;
+
+	/** Selection delegate */
+	FOnTutorialSelected OnTutorialSelected;
+
+	/** Text to highlight */
+	TAttribute<FText> HighlightText;
+
+	/** Button clicked to launch tutorial */
+	mutable TSharedPtr<SWidget> LaunchButton;
+
+	/** Documentation page reference to use if we are displaying a UDN doc */
+	mutable TSharedPtr<IDocumentationPage> DocumentationPage;
+};
+
+void STutorialsBrowser::Construct(const FArguments& InArgs)
+{
+	OnClosed = InArgs._OnClosed;
+	OnLaunchTutorial = InArgs._OnLaunchTutorial;
+	ParentWindow = InArgs._ParentWindow;
+
+	ChildSlot
+	[
+		SNew(SVerticalBox)
+		+SVerticalBox::Slot()
+		.VAlign(VAlign_Center)
+		.HAlign(HAlign_Center)
+		[
+			SNew(SBorder)
+			.Padding(18.0f)
+			.BorderImage(FEditorStyle::GetBrush("Tutorials.Border"))
+			[
+				SNew(SVerticalBox)
+				+SVerticalBox::Slot()
+				.AutoHeight()
+				[
+					// Title area
+					SNew(SHorizontalBox)
+					+SHorizontalBox::Slot()
+					.AutoWidth()
+					.VAlign(VAlign_Center)
+					[
+						SNew(SButton)
+						.OnClicked(this, &STutorialsBrowser::OnBackButtonClicked)
+						.IsEnabled(this, &STutorialsBrowser::IsBackButtonEnabled)
+						.ButtonStyle(&FEditorStyle::Get().GetWidgetStyle<FButtonStyle>("Tutorials.Browser.Button"))
+						.Content()
+						[
+							SNew(SImage)
+							.Image(FEditorStyle::GetBrush("Tutorials.Browser.BackButton"))
+							.ColorAndOpacity(this, &STutorialsBrowser::GetBackButtonColor)
+						]
+					]
+					+SHorizontalBox::Slot()
+					.FillWidth(1.0f)
+					.Padding(10.0f)
+					[
+						SNew(STextBlock)
+						.Text(LOCTEXT("TutorialsTitle", "Unreal Editor Tutorials"))
+						.TextStyle( &FEditorStyle::Get().GetWidgetStyle<FTextBlockStyle>( "Tutorials.Browser.WelcomeHeader"))
+					]
+					+SHorizontalBox::Slot()
+					.AutoWidth()
+					.VAlign(VAlign_Top)
+					[
+						SNew(SButton)
+						.OnClicked(this, &STutorialsBrowser::OnCloseButtonClicked)
+						.ButtonStyle(&FEditorStyle::Get().GetWidgetStyle<FButtonStyle>("Tutorials.Browser.Button"))
+						.ContentPadding(0.0f)
+						[
+							SNew(SImage)
+							.Image(FEditorStyle::GetBrush("Symbols.X"))
+							.ColorAndOpacity(FLinearColor(0.0f, 0.0f, 0.0f, 1.0f))
+						]
+					]
+				]
+				+SVerticalBox::Slot()
+				.AutoHeight()
+				[
+					SNew(SSearchBox)
+					.OnTextChanged(this, &STutorialsBrowser::OnSearchTextChanged)
+				]
+				+SVerticalBox::Slot()
+				.AutoHeight()
+				[
+					SNew(SBox)
+					.HeightOverride(400.0f)
+					.WidthOverride(400.0f)
+					[
+						SNew(SScrollBox)
+						+SScrollBox::Slot()
+						[
+							SAssignNew(TutorialList, SListView<TSharedPtr<ITutorialListEntry>>)
+							.ItemHeight(128.0f)
+							.ListItemsSource(&FilteredEntries)
+							.OnGenerateRow(this, &STutorialsBrowser::OnGenerateTutorialRow)
+							.SelectionMode(ESelectionMode::None)
+						]
+					]
+				]
+			]
+		]
+	];
+
+	ReloadTutorials();
+}
+
+void STutorialsBrowser::SetFilter(const FString& InFilter)
+{
+	CategoryFilter = InFilter;
+	ReloadTutorials();
+}
+
+TSharedRef<ITableRow> STutorialsBrowser::OnGenerateTutorialRow(TSharedPtr<ITutorialListEntry> InItem, const TSharedRef<STableViewBase>& OwnerTable) const
+{
+	return InItem->OnGenerateTutorialRow(OwnerTable);
+}
+
+TSharedPtr<FTutorialListEntry_Category> STutorialsBrowser::RebuildCategories()
+{
+	TArray<TSharedPtr<FTutorialListEntry_Category>> Categories;
+
+	// add root category
+	TSharedPtr<FTutorialListEntry_Category> RootCategory = MakeShareable(new FTutorialListEntry_Category(FOnCategorySelected::CreateSP(this, &STutorialsBrowser::OnCategorySelected)));
+	Categories.Add(RootCategory);
+
+	// rebuild categories
+	for(const auto& TutorialCategory : GetDefault<UTutorialSettings>()->Categories)
+	{
+		Categories.Add(MakeShareable(new FTutorialListEntry_Category(TutorialCategory, FOnCategorySelected::CreateSP(this, &STutorialsBrowser::OnCategorySelected), TAttribute<FText>::Create(TAttribute<FText>::FGetter::CreateSP(this, &STutorialsBrowser::GetSearchText)))));
+	}
+
+	for(const auto& TutorialCategory : GetDefault<UEditorTutorialSettings>()->Categories)
+	{
+		Categories.Add(MakeShareable(new FTutorialListEntry_Category(TutorialCategory, FOnCategorySelected::CreateSP(this, &STutorialsBrowser::OnCategorySelected), TAttribute<FText>::Create(TAttribute<FText>::FGetter::CreateSP(this, &STutorialsBrowser::GetSearchText)))));
+	}
+
+	for(auto& Category : Categories)
+	{
+		// Figure out which base category this category belongs in
+		TSharedPtr<FTutorialListEntry_Category> ParentCategory = RootCategory;
+		const FString& CategoryPath = Category->Category.Identifier;
+
+		// We're expecting the category string to be in the "A.B.C" format.  We'll split up the string here and form
+		// a proper hierarchy in the UI
+		TArray< FString > SplitCategories;
+		CategoryPath.ParseIntoArray( &SplitCategories, TEXT( "." ), true /* bCullEmpty */ );
+
+		FString CurrentCategoryPath;
+
+		// Make sure all of the categories exist
+		for(const auto& SplitCategory : SplitCategories)
+		{
+			// Locate this category at the level we're at in the hierarchy
+			TSharedPtr<FTutorialListEntry_Category> FoundCategory = NULL;
+			TArray< TSharedPtr<ITutorialListEntry> >& TestCategoryList = ParentCategory.IsValid() ? ParentCategory->SubCategories : RootCategory->SubCategories;
+			for(auto& TestCategory : TestCategoryList)
+			{
+				if( StaticCastSharedPtr<FTutorialListEntry_Category>(TestCategory)->CategoryName == SplitCategory )
+				{
+					// Found it!
+					FoundCategory = StaticCastSharedPtr<FTutorialListEntry_Category>(TestCategory);
+					break;
+				}
+			}
+
+			if(!CurrentCategoryPath.IsEmpty())
+			{
+				CurrentCategoryPath += TEXT(".");
+			}
+
+			CurrentCategoryPath += SplitCategory;
+
+			if( !FoundCategory.IsValid() )
+			{
+				// OK, this is a new category name for us, so add it now!
+				if(CategoryPath == CurrentCategoryPath)
+				{
+					FoundCategory = Category;
+				}
+				else
+				{
+					FTutorialCategory InterveningCategory;
+					InterveningCategory.Identifier = CurrentCategoryPath;
+					FoundCategory = MakeShareable(new FTutorialListEntry_Category(InterveningCategory, FOnCategorySelected::CreateSP(this, &STutorialsBrowser::OnCategorySelected), TAttribute<FText>::Create(TAttribute<FText>::FGetter::CreateSP(this, &STutorialsBrowser::GetSearchText))));
+				}
+
+				FoundCategory->ParentCategory = ParentCategory;
+				TestCategoryList.Add( FoundCategory );
+			}
+
+			// Descend the hierarchy for the next category
+			ParentCategory = FoundCategory;
+		}
+	}
+
+	return RootCategory;
+}
+
+void STutorialsBrowser::RebuildTutorials(TSharedPtr<FTutorialListEntry_Category> InRootCategory)
+{
+	TArray<TSharedPtr<FTutorialListEntry_Tutorial>> Tutorials;
+
+	// rebuild tutorials
+	for(const auto& TutorialReference : GetDefault<UTutorialSettings>()->Tutorials)
+	{
+		TSubclassOf<UEditorTutorial> TutorialClass = LoadClass<UEditorTutorial>(NULL, *TutorialReference.AssetLongPathname, NULL, LOAD_None, NULL);
+		if(TutorialClass != nullptr)
+		{
+			Tutorials.Add(MakeShareable(new FTutorialListEntry_Tutorial(TutorialClass->GetDefaultObject<UEditorTutorial>(), FOnTutorialSelected::CreateSP(this, &STutorialsBrowser::OnTutorialSelected), TAttribute<FText>::Create(TAttribute<FText>::FGetter::CreateSP(this, &STutorialsBrowser::GetSearchText)))));
+		}
+	}
+
+	for(const auto& TutorialReference : GetDefault<UEditorTutorialSettings>()->Tutorials)
+	{
+		TSubclassOf<UEditorTutorial> TutorialClass = LoadClass<UEditorTutorial>(NULL, *TutorialReference.AssetLongPathname, NULL, LOAD_None, NULL);
+		if(TutorialClass != nullptr)
+		{
+			Tutorials.Add(MakeShareable(new FTutorialListEntry_Tutorial(TutorialClass->GetDefaultObject<UEditorTutorial>(), FOnTutorialSelected::CreateSP(this, &STutorialsBrowser::OnTutorialSelected), TAttribute<FText>::Create(TAttribute<FText>::FGetter::CreateSP(this, &STutorialsBrowser::GetSearchText)))));
+		}
+	}
+
+	// add tutorials to categories
+	for(const auto& Tutorial : Tutorials)
+	{
+		// Figure out which base category this tutorial belongs in
+		TSharedPtr<FTutorialListEntry_Category> CategoryForTutorial = InRootCategory;
+		const FString& CategoryPath = Tutorial->Tutorial->Category;
+
+		// We're expecting the category string to be in the "A.B.C" format.  We'll split up the string here and form
+		// a proper hierarchy in the UI
+		TArray< FString > SplitCategories;
+		CategoryPath.ParseIntoArray( &SplitCategories, TEXT( "." ), true /* bCullEmpty */ );
+
+		FString CurrentCategoryPath;
+
+		// Make sure all of the categories exist
+		for(const auto& SplitCategory : SplitCategories)
+		{
+			// Locate this category at the level we're at in the hierarchy
+			TSharedPtr<FTutorialListEntry_Category> FoundCategory = NULL;
+			TArray< TSharedPtr<ITutorialListEntry> >& TestCategoryList = CategoryForTutorial.IsValid() ? CategoryForTutorial->SubCategories : InRootCategory->SubCategories;
+			for(auto& TestCategory : TestCategoryList)
+			{
+				if( StaticCastSharedPtr<FTutorialListEntry_Category>(TestCategory)->CategoryName == SplitCategory )
+				{
+					// Found it!
+					FoundCategory = StaticCastSharedPtr<FTutorialListEntry_Category>(TestCategory);
+					break;
+				}
+			}
+
+			if(!CurrentCategoryPath.IsEmpty())
+			{
+				CurrentCategoryPath += TEXT(".");
+			}
+
+			CurrentCategoryPath += SplitCategory;
+
+			if( !FoundCategory.IsValid() )
+			{
+				// OK, this is a new category name for us, so add it now!
+				FTutorialCategory InterveningCategory;
+				InterveningCategory.Identifier = CurrentCategoryPath;
+
+				FoundCategory = MakeShareable(new FTutorialListEntry_Category(InterveningCategory, FOnCategorySelected::CreateSP(this, &STutorialsBrowser::OnCategorySelected), TAttribute<FText>::Create(TAttribute<FText>::FGetter::CreateSP(this, &STutorialsBrowser::GetSearchText))));
+				FoundCategory->ParentCategory = CategoryForTutorial;
+				TestCategoryList.Add( FoundCategory );
+			}
+
+			// Descend the hierarchy for the next category
+			CategoryForTutorial = FoundCategory;
+		}
+
+		Tutorial->ParentCategory = CategoryForTutorial;
+		CategoryForTutorial->AddTutorial( Tutorial );
+	}
+}
+
+void STutorialsBrowser::ReloadTutorials()
+{
+	TSharedPtr<FTutorialListEntry_Category> RootCategory = RebuildCategories();
+	RebuildTutorials(RootCategory);
+	RootEntry = RootCategory;
+
+	// now filter & arrange available tutorials
+	FilterTutorials();
+}
+
+FReply STutorialsBrowser::OnCloseButtonClicked()
+{
+	OnClosed.ExecuteIfBound();
+	return FReply::Handled();
+}
+
+FReply STutorialsBrowser::OnBackButtonClicked()
+{
+	TSharedPtr<FTutorialListEntry_Category> CurrentCategory = FindCategory_Recursive(RootEntry);
+	if(CurrentCategory.IsValid() && CurrentCategory->ParentCategory.IsValid())
+	{
+		TSharedPtr<FTutorialListEntry_Category> PinnedParentCategory = StaticCastSharedPtr<FTutorialListEntry_Category>(CurrentCategory->ParentCategory.Pin());
+		if(PinnedParentCategory.IsValid())
+		{
+			NavigationFilter = PinnedParentCategory->Category.Identifier;
+			FilterTutorials();
+		}
+	}
+
+	return FReply::Handled();
+}
+
+bool STutorialsBrowser::IsBackButtonEnabled() const
+{
+	if(CurrentCategoryPtr.IsValid())
+	{
+		return CurrentCategoryPtr.Pin()->ParentCategory.IsValid();
+	}
+
+	return false;
+}
+
+FSlateColor STutorialsBrowser::GetBackButtonColor() const
+{
+	return IsBackButtonEnabled() ? FLinearColor(0.0f, 0.0f, 0.0f, 1.0f) : FLinearColor(1.0f, 1.0f, 1.0f, 0.25f);
+}
+
+void STutorialsBrowser::OnTutorialSelected(UEditorTutorial* InTutorial, bool bRestart)
+{
+	OnLaunchTutorial.ExecuteIfBound(InTutorial, bRestart, ParentWindow);
+}
+
+void STutorialsBrowser::OnCategorySelected(const FString& InCategory)
+{
+	NavigationFilter = InCategory;
+	FilterTutorials();
+}
+
+void STutorialsBrowser::FilterTutorials()
+{
+	FilteredEntries.Empty();
+
+	TSharedPtr<FTutorialListEntry_Category> CurrentCategory = FindCategory_Recursive(RootEntry);
+
+	if(CurrentCategory.IsValid())
+	{
+		for(const auto& SubCategory : CurrentCategory->SubCategories)
+		{
+			if(SubCategory->PassesFilter(CategoryFilter, SearchFilter.ToString()))
+			{
+				FilteredEntries.Add(SubCategory);
+			}
+		}
+
+		for(const auto& Tutorial : CurrentCategory->Tutorials)
+		{
+			if(Tutorial->PassesFilter(CategoryFilter, SearchFilter.ToString()))
+			{
+				FilteredEntries.Add(Tutorial);
+			}
+		}
+
+		CurrentCategoryPtr = CurrentCategory;
+	}
+
+	TutorialList->RequestListRefresh();
+}
+
+TSharedPtr<FTutorialListEntry_Category> STutorialsBrowser::FindCategory_Recursive(TSharedPtr<FTutorialListEntry_Category> InCategory) const
+{
+	if(InCategory->Category.Identifier == NavigationFilter)
+	{
+		return InCategory;
+	}
+
+	for(const auto& Category : InCategory->SubCategories)
+	{
+		TSharedPtr<FTutorialListEntry_Category> TestCategory = FindCategory_Recursive(StaticCastSharedPtr<FTutorialListEntry_Category>(Category));
+		if(TestCategory.IsValid())
+		{
+			return TestCategory;
+		}
+	}
+
+	return TSharedPtr<FTutorialListEntry_Category>();
+}
+
+void STutorialsBrowser::OnSearchTextChanged(const FText& InText)
+{
+	SearchFilter = InText;
+	FilterTutorials();
+}
+
+FText STutorialsBrowser::GetSearchText() const
+{
+	return SearchFilter;
+}
+
+#undef LOCTEXT_NAMESPACE
\ No newline at end of file
diff --git a/Engine/Source/Editor/IntroTutorials/Private/STutorialsBrowser.h b/Engine/Source/Editor/IntroTutorials/Private/STutorialsBrowser.h
new file mode 100644
--- /dev/null
+++ b/Engine/Source/Editor/IntroTutorials/Private/STutorialsBrowser.h
@@ -0,0 +1,113 @@
+// Copyright 1998-2014 Epic Games, Inc. All Rights Reserved.
+
+#pragma once
+
+class UEditorTutorial;
+
+DECLARE_DELEGATE_ThreeParams(FOnLaunchTutorial, UEditorTutorial* /** InTutorialToLaunch */, bool /* bRestart */, TWeakPtr<SWindow> /* InFromWindow */)
+
+/** Abstract base class for list entries in the tutorial menu */
+struct ITutorialListEntry
+{
+	/** Generate content for a tree entry */
+	virtual TSharedRef<ITableRow> OnGenerateTutorialRow(const TSharedRef<STableViewBase>& OwnerTable) const = 0;
+
+	/** Whether this entry passes the current filter criteria */
+	virtual bool PassesFilter(const FString& InCategoryFilter, const FString& InFilter) const = 0;
+};
+
+class FTutorialListEntry_Category;
+
+/**
+ * The widget which holds all available tutorials
+ */
+class STutorialsBrowser : public SCompoundWidget
+{
+	SLATE_BEGIN_ARGS( STutorialsBrowser ) {}
+
+	SLATE_ARGUMENT(FSimpleDelegate, OnClosed)
+	SLATE_ARGUMENT(FOnLaunchTutorial, OnLaunchTutorial)
+	SLATE_ARGUMENT(TWeakPtr<SWindow>, ParentWindow)
+
+	SLATE_END_ARGS()
+
+	void Construct(const FArguments& Args);
+
+	/** Set the current filter string. Filters are used to only show the specifed category of tutorials (e.g. only Blueprint tutorials) */
+	void SetFilter(const FString& InFilter);
+
+protected:
+	/** Handle generating a table row in the browser */
+	TSharedRef<ITableRow> OnGenerateTutorialRow(TSharedPtr<ITutorialListEntry> InItem, const TSharedRef<STableViewBase>& OwnerTable) const;
+
+	/** Reload all tutorials that we know about */
+	void ReloadTutorials();
+
+	/** Handle closing the browser */
+	FReply OnCloseButtonClicked();
+
+	/** Handle traversing back up the browser hierarchy */
+	FReply OnBackButtonClicked();
+
+	/** Handle whether the back button can be clicked */
+	bool IsBackButtonEnabled() const;
+
+	/** Handle back button color - used to dim the button when disabled */
+	FSlateColor GetBackButtonColor() const;
+
+	/** Delegate handler fired when a tutorial is selected form the browser */
+	void OnTutorialSelected(UEditorTutorial* InTutorial, bool bRestart);
+
+	/** Delegate handler fired when a category is selected form the browser */
+	void OnCategorySelected(const FString& InCategory);
+
+	/** Filter displayed tutorials - regenerates the displayed items */
+	void FilterTutorials();
+
+	/** Rebuild the displayed categories */
+	TSharedPtr<FTutorialListEntry_Category> RebuildCategories();
+
+	/** Rebuild the displayed tutorials */
+	void RebuildTutorials(TSharedPtr<FTutorialListEntry_Category> InRootCategory);
+
+	/** Recursive helper function used to find a category given the current NavigationFilter */
+	TSharedPtr<FTutorialListEntry_Category> FindCategory_Recursive(TSharedPtr<FTutorialListEntry_Category> InCategory) const;
+
+	/** Handle rebuilding the browser display when the filter text changes */
+	void OnSearchTextChanged(const FText& InText);
+
+	/** Supplies the text to display in the filter box */
+	FText GetSearchText() const;
+
+private:
+
+	/** Root entry of the tutorials tree */
+	TSharedPtr<FTutorialListEntry_Category> RootEntry;
+
+	/** Current filtered entries */
+	TArray<TSharedPtr<ITutorialListEntry>> FilteredEntries;
+
+	/** List of tutorials widget */
+	TSharedPtr<SListView<TSharedPtr<ITutorialListEntry>>> TutorialList;
+
+	/** Delegate fired when the close button is clicked */
+	FSimpleDelegate OnClosed;
+
+	/** Delegate fired when a tutorial is selected */
+	FOnLaunchTutorial OnLaunchTutorial;
+
+	/** Current search string */
+	FText SearchFilter;
+
+	/** Current category navigation string */
+	FString NavigationFilter;
+
+	/** Static category filter for this browser */
+	FString CategoryFilter;
+
+	/** Cached current category */
+	TWeakPtr<FTutorialListEntry_Category> CurrentCategoryPtr;
+
+	/** Parent window of this browser */
+	TWeakPtr<SWindow> ParentWindow;
+};
\ No newline at end of file
diff --git a/Engine/Source/Editor/IntroTutorials/Private/TutorialSettings.cpp b/Engine/Source/Editor/IntroTutorials/Private/TutorialSettings.cpp
new file mode 100644
--- /dev/null
+++ b/Engine/Source/Editor/IntroTutorials/Private/TutorialSettings.cpp
@@ -0,0 +1,10 @@
+// Copyright 1998-2014 Epic Games, Inc. All Rights Reserved.
+
+#include "IntroTutorialsPrivatePCH.h"
+#include "TutorialSettings.h"
+
+UTutorialSettings::UTutorialSettings(const class FPostConstructInitializeProperties& PCIP)
+	: Super(PCIP)
+{
+}
+
diff --git a/Engine/Source/Editor/IntroTutorials/Private/TutorialSettings.h b/Engine/Source/Editor/IntroTutorials/Private/TutorialSettings.h
new file mode 100644
--- /dev/null
+++ b/Engine/Source/Editor/IntroTutorials/Private/TutorialSettings.h
@@ -0,0 +1,25 @@
+// Copyright 1998-2014 Epic Games, Inc. All Rights Reserved.
+
+#pragma once
+
+#include "EditorTutorial.h"
+#include "TutorialSettings.generated.h"
+
+/** Per-project tutorial settings */
+UCLASS(config=Editor)
+class UTutorialSettings : public UObject
+{
+	GENERATED_UCLASS_BODY()
+
+	/** Categories for tutorials */
+	UPROPERTY(Config, EditAnywhere, Category="Tutorials")
+	TArray<FTutorialCategory> Categories;
+
+	/** The tutorials available in the editor */
+	UPROPERTY(Config, EditAnywhere, Category="Tutorials", meta=(MetaClass="EditorTutorial"))
+	TArray<FStringClassReference> Tutorials;
+
+	/** Tutorial to start on project startup */
+	UPROPERTY(Config, EditAnywhere, Category="Tutorials", meta=(MetaClass="EditorTutorial"))
+	FStringClassReference StartupTutorial;
+};
\ No newline at end of file
diff --git a/Engine/Source/Editor/IntroTutorials/Private/TutorialStructCustomization.cpp b/Engine/Source/Editor/IntroTutorials/Private/TutorialStructCustomization.cpp
new file mode 100644
--- /dev/null
+++ b/Engine/Source/Editor/IntroTutorials/Private/TutorialStructCustomization.cpp
@@ -0,0 +1,355 @@
+// Copyright 1998-2014 Epic Games, Inc. All Rights Reserved.
+
+#include "IntroTutorialsPrivatePCH.h"
+#include "TutorialStructCustomization.h"
+#include "EditorTutorial.h"
+
+#define LOCTEXT_NAMESPACE "TutorialStructCustomization"
+
+TSharedRef<IPropertyTypeCustomization> FTutorialContentCustomization::MakeInstance()
+{
+	return MakeShareable( new FTutorialContentCustomization );
+}
+
+void FTutorialContentCustomization::CustomizeHeader( TSharedRef<class IPropertyHandle> InStructPropertyHandle, class FDetailWidgetRow& HeaderRow, IPropertyTypeCustomizationUtils& StructCustomizationUtils )
+{
+	TSharedPtr<IPropertyHandle> TypeProperty = InStructPropertyHandle->GetChildHandle(GET_MEMBER_NAME_CHECKED(FTutorialContent, Type));
+	TSharedPtr<IPropertyHandle> ContentProperty = InStructPropertyHandle->GetChildHandle(GET_MEMBER_NAME_CHECKED(FTutorialContent, Content));
+	TSharedPtr<IPropertyHandle> ExcerptNameProperty = InStructPropertyHandle->GetChildHandle(GET_MEMBER_NAME_CHECKED(FTutorialContent, ExcerptName));
+	TSharedPtr<IPropertyHandle> TextProperty = InStructPropertyHandle->GetChildHandle(GET_MEMBER_NAME_CHECKED(FTutorialContent, Text));
+
+	// Show and hide various widgets based on current content type
+	struct Local
+	{
+		static EVisibility GetContentVisibility(TSharedPtr<IPropertyHandle> InPropertyHandle)
+		{
+			check(InPropertyHandle.IsValid());
+
+			uint8 Value = 0;
+			if(InPropertyHandle->GetValue(Value) == FPropertyAccess::Success)
+			{
+				const ETutorialContent::Type EnumValue = (ETutorialContent::Type)Value;
+				return (EnumValue == ETutorialContent::UDNExcerpt) ? EVisibility::Visible : EVisibility::Collapsed;
+			}
+
+			return EVisibility::Collapsed;
+		}
+
+		static EVisibility GetExcerptNameVisibility(TSharedPtr<IPropertyHandle> InPropertyHandle)
+		{
+			check(InPropertyHandle.IsValid());
+
+			uint8 Value = 0;
+			if(InPropertyHandle->GetValue(Value) == FPropertyAccess::Success)
+			{
+				const ETutorialContent::Type EnumValue = (ETutorialContent::Type)Value;
+				return (EnumValue == ETutorialContent::UDNExcerpt) ? EVisibility::Visible : EVisibility::Collapsed;
+			}
+
+			return EVisibility::Collapsed;
+		}
+
+		static EVisibility GetTextVisibility(TSharedPtr<IPropertyHandle> InPropertyHandle)
+		{
+			check(InPropertyHandle.IsValid());
+
+			uint8 Value = 0;
+			if(InPropertyHandle->GetValue(Value) == FPropertyAccess::Success)
+			{
+				const ETutorialContent::Type EnumValue = (ETutorialContent::Type)Value;
+				return (EnumValue == ETutorialContent::Text) ? EVisibility::Visible : EVisibility::Collapsed;
+			}
+
+			return EVisibility::Collapsed;
+		}
+	};
+
+	HeaderRow
+	.NameContent()
+	[
+		ContentProperty->CreatePropertyNameWidget()
+	]
+	.ValueContent()
+	.MinDesiredWidth(250.0f)
+	[
+		SNew(SVerticalBox)
+		+SVerticalBox::Slot()
+		.AutoHeight()
+		.Padding(0.0f, 2.0f)
+		[
+			TypeProperty->CreatePropertyValueWidget()
+		]
+		+SVerticalBox::Slot()
+		.AutoHeight()
+		.Padding(0.0f, 2.0f)
+		[
+			SNew(SHorizontalBox)
+			.Visibility_Static(&Local::GetContentVisibility, TypeProperty)
+			+SHorizontalBox::Slot()
+			[
+				ContentProperty->CreatePropertyValueWidget()
+			]
+		]
+		+SVerticalBox::Slot()
+		.AutoHeight()
+		.Padding(0.0f, 2.0f)
+		[
+			SNew(SHorizontalBox)
+			.Visibility_Static(&Local::GetExcerptNameVisibility, TypeProperty)
+			+SHorizontalBox::Slot()
+			[
+				ExcerptNameProperty->CreatePropertyValueWidget()
+			]
+		]
+		+SVerticalBox::Slot()
+		.AutoHeight()
+		.Padding(0.0f, 2.0f)
+		[
+			SNew(SHorizontalBox)
+			.Visibility_Static(&Local::GetTextVisibility, TypeProperty)
+			+SHorizontalBox::Slot()
+			[
+				TextProperty->CreatePropertyValueWidget()
+			]
+		]
+	];
+}
+
+void FTutorialContentCustomization::CustomizeChildren( TSharedRef<class IPropertyHandle> InStructPropertyHandle, class IDetailChildrenBuilder& StructBuilder, IPropertyTypeCustomizationUtils& StructCustomizationUtils )
+{
+
+}
+
+/** 'Tooltip' window to indicate what is currently being picked and how to abort the picker (Esc) */
+class SWidgetPickerFloatingWindow : public SCompoundWidget
+{
+	SLATE_BEGIN_ARGS( SWidgetPickerFloatingWindow ) {}
+
+	SLATE_ARGUMENT(TWeakPtr<SWindow>, ParentWindow)
+
+	SLATE_END_ARGS()
+
+	void Construct(const FArguments& InArgs, TSharedRef<class IPropertyHandle> InStructPropertyHandle)
+	{
+		StructPropertyHandle = InStructPropertyHandle;
+		ParentWindow = InArgs._ParentWindow;
+
+		ChildSlot
+		[
+			SNew(SToolTip)
+			.Text(this, &SWidgetPickerFloatingWindow::GetPickerStatusText)
+		];
+	}
+
+	FName GetPickedWidgetName()
+	{
+		return PickedWidgetName;
+	}
+
+private:
+	virtual void Tick( const FGeometry& AllottedGeometry, const double InCurrentTime, const float InDeltaTime ) override
+	{
+		PickedWidgetName = NAME_None;
+
+		FWidgetPath Path = FSlateApplication::Get().LocateWindowUnderMouse(FSlateApplication::Get().GetCursorPos(), FSlateApplication::Get().GetInteractiveTopLevelWindows(), true);
+
+		for(int32 PathIndex = Path.Widgets.Num() - 1; PathIndex >= 0; PathIndex--)
+		{
+			TSharedRef<SWidget> PathWidget = Path.Widgets(PathIndex).Widget;
+			if(PathWidget->GetTag() != NAME_None)
+			{
+				PickedWidgetName = PathWidget->GetTag();
+				break;
+			}
+		}
+
+		// kind of a hack, but we need to maintain keyboard focus otherwise we wont get our keypress to 'pick'
+		FSlateApplication::Get().SetKeyboardFocus(SharedThis(this), EKeyboardFocusCause::SetDirectly);
+		if(ParentWindow.IsValid())
+		{
+			// also kind of a hack, but this is the only way at the moment to get a 'cursor decorator' without using the drag-drop code path
+			ParentWindow.Pin()->MoveWindowTo(FSlateApplication::Get().GetCursorPos() + FSlateApplication::Get().GetCursorSize());
+		}
+	}
+
+	FText GetPickerStatusText() const
+	{
+		return FText::Format(LOCTEXT("TootipHint", "{0} (Esc to pick)"), FText::FromName(PickedWidgetName));
+	}
+
+	virtual FReply OnKeyDown(const FGeometry& MyGeometry, const FKeyboardEvent& InKeyboardEvent) override
+	{
+		if(InKeyboardEvent.GetKey() == EKeys::Escape)
+		{
+			TSharedPtr<IPropertyHandle> WrapperIdentifierProperty = StructPropertyHandle->GetChildHandle(GET_MEMBER_NAME_CHECKED(FTutorialContentAnchor, WrapperIdentifier));
+			WrapperIdentifierProperty->SetValue(PickedWidgetName);
+
+			TSharedPtr<IPropertyHandle> TypeProperty = StructPropertyHandle->GetChildHandle(GET_MEMBER_NAME_CHECKED(FTutorialContentAnchor, Type));
+			TypeProperty->SetValue((uint8)ETutorialAnchorIdentifier::NamedWidget);
+			
+			PickedWidgetName = NAME_None;
+
+			if(ParentWindow.IsValid())
+			{
+				FSlateApplication::Get().RequestDestroyWindow(ParentWindow.Pin().ToSharedRef());
+				ParentWindow.Reset();
+			}
+
+			return FReply::Handled();
+		}
+
+		return FReply::Unhandled();
+	}
+
+	/** We need to support keyboard focus to process the 'Esc' key */
+	virtual bool SupportsKeyboardFocus() const override
+	{
+		return true;
+	}
+
+private:
+	/** Handle to the property struct */
+	TSharedPtr<class IPropertyHandle> StructPropertyHandle;
+
+	/** Handle to the window that contains this widget */
+	TWeakPtr<SWindow> ParentWindow;
+
+	/** The widget name we are picking */
+	FName PickedWidgetName;
+};
+
+/** Widget used to launch a 'picking' session */
+class SWidgetPicker : public SCompoundWidget
+{
+public:
+	SLATE_BEGIN_ARGS( SWidgetPicker ) {}
+
+	SLATE_END_ARGS()
+
+	~SWidgetPicker()
+	{
+		// kill the picker window as well if this widget is going away - that way we dont get dangling refs to the property
+		if(PickerWindow.IsValid() && FSlateApplication::IsInitialized())
+		{
+			FSlateApplication::Get().RequestDestroyWindow(PickerWindow.Pin().ToSharedRef());
+			PickerWindow.Reset();
+			PickerWidget.Reset();
+		}
+	}
+
+	void Construct(const FArguments& InArgs, TSharedRef<class IPropertyHandle> InStructPropertyHandle, IPropertyTypeCustomizationUtils* InStructCustomizationUtils)
+	{
+		StructPropertyHandle = InStructPropertyHandle;
+
+		ChildSlot
+		[
+			SNew(SHorizontalBox)
+			+SHorizontalBox::Slot()
+			.FillWidth(1.0f)
+			.VAlign(VAlign_Center)
+			[
+				SNew(STextBlock)
+				.Text(this, &SWidgetPicker::HandlePickerStatusText)
+				.Font(InStructCustomizationUtils->GetRegularFont())
+			]
+			+SHorizontalBox::Slot()
+			.AutoWidth()
+			.VAlign(VAlign_Center)
+			[
+				SNew(SButton)
+				.ButtonStyle( FEditorStyle::Get(), "HoverHintOnly" )
+				.OnClicked( this, &SWidgetPicker::OnClicked )
+				.ContentPadding(4.0f)
+				.ForegroundColor( FSlateColor::UseForeground() )
+				.IsFocusable(false)
+				[
+					SNew( SImage )
+					.Image( FEditorStyle::GetBrush("PropertyWindow.Button_PickActorInteractive") )
+					.ColorAndOpacity( FSlateColor::UseForeground() )
+				]
+			]
+		];
+	}
+
+private:
+	FReply OnClicked()
+	{
+		// launch a picker window
+		if(!PickerWindow.IsValid())
+		{
+			TSharedRef<SWindow> NewWindow = SWindow::MakeCursorDecorator();
+			NewWindow->MoveWindowTo(FSlateApplication::Get().GetCursorPos());
+			PickerWindow = NewWindow;
+
+			NewWindow->SetContent(
+				SAssignNew(PickerWidget, SWidgetPickerFloatingWindow, StructPropertyHandle.ToSharedRef())
+				.ParentWindow(NewWindow)
+				);
+
+			TSharedPtr<SWindow> RootWindow = FGlobalTabmanager::Get()->GetRootWindow();
+			check(RootWindow.IsValid());
+			FSlateApplication::Get().AddWindowAsNativeChild(NewWindow, RootWindow.ToSharedRef());
+
+			FIntroTutorials& IntroTutorials = FModuleManager::Get().GetModuleChecked<FIntroTutorials>("IntroTutorials");
+			IntroTutorials.OnIsPicking().BindSP(this, &SWidgetPicker::OnIsPicking);
+		}
+
+		return FReply::Handled();
+	}
+
+	FText HandlePickerStatusText() const
+	{
+		FName WidgetName = NAME_None;
+		TSharedPtr<IPropertyHandle> WrapperIdentifierProperty = StructPropertyHandle->GetChildHandle(GET_MEMBER_NAME_CHECKED(FTutorialContentAnchor, WrapperIdentifier));
+		WrapperIdentifierProperty->GetValue(WidgetName);
+		return FText::FromName(WidgetName);
+	}
+
+	bool OnIsPicking(FName& OutWidgetNameToHighlight) const
+	{
+		if(PickerWidget.IsValid())
+		{
+			OutWidgetNameToHighlight = PickerWidget.Pin()->GetPickedWidgetName();
+			return true;
+		}
+		
+		return false;
+	}
+
+private:
+	/** Picker window widget */
+	TWeakPtr<SWidgetPickerFloatingWindow> PickerWidget;
+
+	/** Picker window */
+	TWeakPtr<SWindow> PickerWindow;
+
+	/** Handle to the struct we are customizing */
+	TSharedPtr<class IPropertyHandle> StructPropertyHandle;
+};
+
+TSharedRef<IPropertyTypeCustomization> FTutorialContentAnchorCustomization::MakeInstance()
+{
+	return MakeShareable( new FTutorialContentAnchorCustomization );
+}
+
+void FTutorialContentAnchorCustomization::CustomizeHeader( TSharedRef<class IPropertyHandle> InStructPropertyHandle, class FDetailWidgetRow& HeaderRow, IPropertyTypeCustomizationUtils& StructCustomizationUtils )
+{
+	HeaderRow
+	.NameContent()
+	[
+		InStructPropertyHandle->CreatePropertyNameWidget()
+	]
+	.ValueContent()
+	.MinDesiredWidth(250.0f)
+	.MaxDesiredWidth(500.0f)
+	[
+		SNew(SWidgetPicker, InStructPropertyHandle, &StructCustomizationUtils)
+	];
+}
+
+void FTutorialContentAnchorCustomization::CustomizeChildren( TSharedRef<class IPropertyHandle> InStructPropertyHandle, class IDetailChildrenBuilder& StructBuilder, IPropertyTypeCustomizationUtils& StructCustomizationUtils )
+{
+
+}
+
+#undef LOCTEXT_NAMESPACE
\ No newline at end of file
diff --git a/Engine/Source/Editor/IntroTutorials/Private/TutorialStructCustomization.h b/Engine/Source/Editor/IntroTutorials/Private/TutorialStructCustomization.h
new file mode 100644
--- /dev/null
+++ b/Engine/Source/Editor/IntroTutorials/Private/TutorialStructCustomization.h
@@ -0,0 +1,25 @@
+// Copyright 1998-2014 Epic Games, Inc. All Rights Reserved.
+
+#pragma once
+
+/** Struct customization for FTutorialContent */
+class FTutorialContentCustomization : public IPropertyTypeCustomization
+{
+public:
+	static TSharedRef<IPropertyTypeCustomization> MakeInstance();
+
+	/** IPropertyTypeCustomization interface */
+	virtual void CustomizeHeader( TSharedRef<class IPropertyHandle> InStructPropertyHandle, class FDetailWidgetRow& HeaderRow, IPropertyTypeCustomizationUtils& StructCustomizationUtils ) override;
+	virtual void CustomizeChildren( TSharedRef<class IPropertyHandle> InStructPropertyHandle, class IDetailChildrenBuilder& StructBuilder, IPropertyTypeCustomizationUtils& StructCustomizationUtils ) override;
+};
+
+/** Struct customization for FTutorialContentAnchor */
+class FTutorialContentAnchorCustomization : public IPropertyTypeCustomization
+{
+public:
+	static TSharedRef<IPropertyTypeCustomization> MakeInstance();
+
+	/** IPropertyTypeCustomization interface */
+	virtual void CustomizeHeader( TSharedRef<class IPropertyHandle> InStructPropertyHandle, class FDetailWidgetRow& HeaderRow, IPropertyTypeCustomizationUtils& StructCustomizationUtils ) override;
+	virtual void CustomizeChildren( TSharedRef<class IPropertyHandle> InStructPropertyHandle, class IDetailChildrenBuilder& StructBuilder, IPropertyTypeCustomizationUtils& StructCustomizationUtils ) override;
+};
diff --git a/Engine/Source/Editor/IntroTutorials/Public/EditorTutorial.h b/Engine/Source/Editor/IntroTutorials/Public/EditorTutorial.h
new file mode 100644
--- /dev/null
+++ b/Engine/Source/Editor/IntroTutorials/Public/EditorTutorial.h
@@ -0,0 +1,215 @@
+// Copyright 1998-2014 Epic Games, Inc. All Rights Reserved.
+
+#pragma once
+
+#include "EditorTutorial.generated.h"
+
+/** The type of tutorial content to display */
+UENUM()
+namespace ETutorialContent
+{
+	enum Type
+	{
+		/** Blank - displays no content */
+		None,
+
+		/** Text content */
+		Text,
+
+		/** Content from a UDN excerpt */
+		UDNExcerpt,
+	};
+}
+
+/** The type of tutorial content to display */
+UENUM()
+namespace ETutorialAnchorIdentifier
+{
+	enum Type
+	{
+		/** No anchor */
+		None,
+
+		/** Uses a tutorial wrapper widget */
+		NamedWidget,
+
+		/** An asset accessible via the content browser */
+		Asset,
+	};
+}
+
+/** Category description */
+USTRUCT()
+struct INTROTUTORIALS_API FTutorialCategory
+{
+	GENERATED_USTRUCT_BODY()
+
+	/** Period-separated category name, e.g. "Editor Quickstart.Level Editor" */
+	UPROPERTY(EditAnywhere, Category="Content")
+	FString Identifier;
+
+	/** Title of the category */
+	UPROPERTY(EditAnywhere, Category="Content")
+	FText Title;
+
+	/** Localized text to use to describe this category */
+	UPROPERTY(EditAnywhere, Category="Content", meta=(MultiLine=true))
+	FText Description;
+
+	/** Icon for this tutorial, used when presented to the user */
+	UPROPERTY(EditAnywhere, Category="Tutorial")
+	FString Icon;
+};
+
+/** Content wrapper */
+USTRUCT()
+struct INTROTUTORIALS_API FTutorialContent
+{
+	GENERATED_USTRUCT_BODY()
+
+	/** The type of this content */
+	UPROPERTY(EditAnywhere, Category="Content")
+	TEnumAsByte<ETutorialContent::Type> Type;
+
+	/** Content reference string, path etc. */
+	UPROPERTY(EditAnywhere, Category="Content")
+	FString Content;
+
+	/** Excerpt name for UDN excerpt */
+	UPROPERTY(EditAnywhere, Category="Content")
+	FString ExcerptName;
+
+	/** Localized text to use with this content */
+	UPROPERTY(EditAnywhere, Category="Content", meta=(MultiLine=true))
+	FText Text;
+};
+
+/** A way of identifying something to be highlighted by a tutorial */
+USTRUCT()
+struct INTROTUTORIALS_API FTutorialContentAnchor
+{
+	GENERATED_USTRUCT_BODY()
+	
+	FTutorialContentAnchor()
+	{
+		 Type = ETutorialAnchorIdentifier::None;
+	}
+
+	UPROPERTY(EditAnywhere, Category="Anchor")
+	TEnumAsByte<ETutorialAnchorIdentifier::Type> Type;
+
+	/** If widget is in a wrapper widget, this is the wrapper widget name */
+	UPROPERTY(EditAnywhere, Category="Anchor")
+	FName WrapperIdentifier;
+
+	/** If reference is an asset, we use this to resolve it */
+	UPROPERTY(EditAnywhere, Category="Anchor")
+	FStringAssetReference Asset;
+};
+
+/** Content that is displayed relative to a widget */
+USTRUCT()
+struct INTROTUTORIALS_API FTutorialWidgetContent
+{
+	GENERATED_USTRUCT_BODY()
+
+	FTutorialWidgetContent()
+	{
+		ContentWidth = 350.0f;
+		HorizontalAlignment = HAlign_Center;
+		VerticalAlignment = VAlign_Bottom;
+	}
+
+	/** Anchor for content widget to highlight */
+	UPROPERTY(EditAnywhere, Category="Widget")
+	FTutorialContentAnchor WidgetAnchor;
+
+	UPROPERTY(EditAnywhere, Category="Widget")
+	TEnumAsByte<EHorizontalAlignment> HorizontalAlignment;
+
+	UPROPERTY(EditAnywhere, Category="Widget")
+	TEnumAsByte<EVerticalAlignment> VerticalAlignment;
+
+	/** Custom offset from widget */
+	UPROPERTY(EditAnywhere, Category="Widget")
+	FVector2D Offset;
+
+	/** Content width - text will be wrapped at this point */
+	UPROPERTY(EditAnywhere, Category="Widget", meta=(UIMin="10.0", UIMax="600.0"))
+	float ContentWidth;
+
+	/** Content to associate with widget */
+	UPROPERTY(EditAnywhere, EditInline, Category="Widget")
+	FTutorialContent Content;
+};
+
+/** A single tutorial stage, containing the optional main content & a number of widgets with content attached */
+USTRUCT()
+struct INTROTUTORIALS_API FTutorialStage
+{
+	GENERATED_USTRUCT_BODY()
+
+	UPROPERTY(EditAnywhere, Category="Stage")
+	FName Name;
+
+	UPROPERTY(EditAnywhere, EditInline, Category="Stage")
+	FTutorialContent Content;
+
+	UPROPERTY(EditAnywhere, EditInline, Category="Stage")
+	TArray<FTutorialWidgetContent> WidgetContent;
+};
+
+/** An asset used to build a stage-by-stage tutorial in the editor */
+UCLASS(Blueprintable, EditInlineNew)
+class INTROTUTORIALS_API UEditorTutorial : public UObject
+{
+	GENERATED_UCLASS_BODY()
+
+	/** Title of this tutorial, used when presented to the user */
+	UPROPERTY(EditAnywhere, Category="Tutorial")
+	FText Title;
+
+	/** Icon for this tutorial, used when presented to the user */
+	UPROPERTY(EditAnywhere, Category="Tutorial")
+	FString Icon;
+
+	/** Category of this tutorial, used to organize tutorials when presented to the user */
+	UPROPERTY(EditAnywhere, Category="Tutorial")
+	FString Category;
+
+	/** Content to be displayed for this tutorial when presented to the user in summary */
+	UPROPERTY(EditAnywhere, Category="Tutorial")
+	FTutorialContent SummaryContent;
+
+	/** The various stages of this tutorial */
+	UPROPERTY(EditAnywhere, EditInline, Category="Stages")
+	TArray<FTutorialStage> Stages;
+
+	/** Tutorial to optionally chain onto after this tutorial completes */
+	UPROPERTY(EditAnywhere, Category="Tutorial", meta=(MetaClass="EditorTutorial"))
+	FStringClassReference NextTutorial;
+
+	/** A standalone tutorial displays no navigation buttons and each content widget has a close button */
+	UPROPERTY(EditAnywhere, Category="Tutorial")
+	bool bIsStandalone;
+
+	/** Event fired when a tutorial stage begins */
+	UFUNCTION(BlueprintImplementableEvent, Category="Tutorial")
+	void OnTutorialStageStarted(FName StageName) const;
+
+	/** Event fired when a tutorial stage ends */
+	UFUNCTION(BlueprintImplementableEvent, Category="Tutorial")
+	void OnTutorialStageCompleted(FName StageName) const;
+
+	/** Advance to the next stage of a tutorial */
+	UFUNCTION(BlueprintCallable, Category="Tutorial")
+	static void GotoNextTutorialStage();
+
+	/** Advance to the previous stage of a tutorial */
+	UFUNCTION(BlueprintCallable, Category="Tutorial")
+	static void GotoPreviousTutorialStage();
+
+	/** Begin a tutorial. Note that this will end the current tutorial that is in progress, if any */
+	UFUNCTION(BlueprintCallable, Category="Tutorial")
+	static void BeginTutorial(UEditorTutorial* TutorialToStart);
+};
\ No newline at end of file
diff --git a/Engine/Source/Editor/IntroTutorials/Public/IIntroTutorials.h b/Engine/Source/Editor/IntroTutorials/Public/IIntroTutorials.h
--- a/Engine/Source/Editor/IntroTutorials/Public/IIntroTutorials.h
+++ b/Engine/Source/Editor/IntroTutorials/Public/IIntroTutorials.h
@@ -4,6 +4,7 @@
 
 #include "ModuleManager.h"
 
+class UEditorTutorial;
 
 /**
  * The public interface to this module.  In most cases, this interface is only public to sibling modules 
@@ -49,5 +50,14 @@ public:
 	 * @param AssetClass		Asset type to unregister the tutorial from
 	 */
 	virtual void UnregisterTutorialForAssetEditor(UClass* AssetClass)=0;
+
+	/**
+	 * Launch a tutorial immediately, bypassing the tutorial browser.
+	 * 
+	 * @param	Tutorial	The tutorial to launch
+	 * @param	bRestart	Whether to restart the tutorial or resume from where we left off last time.
+	 * @param	InNavigationWindow	Optional window to launch the tutorial from - this is where navigation will be displayed.
+	 */
+	virtual void LaunchTutorial(UEditorTutorial* Tutorial, bool bRestart = true, TWeakPtr<SWindow> InNavigationWindow = nullptr) = 0;
 };
 
diff --git a/Engine/Source/Editor/Kismet/Private/SBlueprintEditorSelectedDebugObjectWidget.cpp b/Engine/Source/Editor/Kismet/Private/SBlueprintEditorSelectedDebugObjectWidget.cpp
--- a/Engine/Source/Editor/Kismet/Private/SBlueprintEditorSelectedDebugObjectWidget.cpp
+++ b/Engine/Source/Editor/Kismet/Private/SBlueprintEditorSelectedDebugObjectWidget.cpp
@@ -18,7 +18,6 @@
 #include "PropertyCustomizationHelpers.h"
 #include "IDocumentation.h"
 #include "SLevelOfDetailBranchNode.h"
-#include "STutorialWrapper.h"
 #include "SBlueprintEditorSelectedDebugObjectWidget.h"
 
 #define LOCTEXT_NAMESPACE "KismetToolbar"
diff --git a/Engine/Source/Editor/Kismet/Private/SBlueprintEditorToolbar.cpp b/Engine/Source/Editor/Kismet/Private/SBlueprintEditorToolbar.cpp
--- a/Engine/Source/Editor/Kismet/Private/SBlueprintEditorToolbar.cpp
+++ b/Engine/Source/Editor/Kismet/Private/SBlueprintEditorToolbar.cpp
@@ -19,7 +19,6 @@
 #include "PropertyCustomizationHelpers.h"
 #include "IDocumentation.h"
 #include "SLevelOfDetailBranchNode.h"
-#include "STutorialWrapper.h"
 #include "SBlueprintEditorSelectedDebugObjectWidget.h"
 #include "Engine/LevelScriptBlueprint.h"
 #include "Merge.h"
@@ -815,59 +814,53 @@ TArray< TSharedPtr< SWidget> > FBlueprintEditorToolbar::GenerateToolbarWidgets(c
 		Ret.Add(SNew(SSpacer).Size(FVector2D(4.0f, 1.0f)));
 
 		Ret.Add(
-			SNew(STutorialWrapper, TEXT("DefaultsMode"))
-			[
-				SNew(SModeWidget, FBlueprintEditorApplicationModes::GetLocalizedMode(FBlueprintEditorApplicationModes::BlueprintDefaultsMode), FBlueprintEditorApplicationModes::BlueprintDefaultsMode)
-				.OnGetActiveMode(ActiveModeGetter)
-				.OnSetActiveMode(ActiveModeSetter)
-				.CanBeSelected(BlueprintObj ? FBlueprintEditorUtils::DoesSupportDefaults(BlueprintObj) : false)
-				.ToolTip(IDocumentation::Get()->CreateToolTip(
-					LOCTEXT("BlueprintDefaultsModeButtonTooltip", "Switch to Blueprint Defaults Mode"),
-					NULL,
-					TEXT("Shared/Editors/BlueprintEditor"),
-					TEXT("DefaultsMode")))
-				.IconImage(FEditorStyle::GetBrush("FullBlueprintEditor.SwitchToBlueprintDefaultsMode"))
-				.SmallIconImage(FEditorStyle::GetBrush("FullBlueprintEditor.SwitchToBlueprintDefaultsMode.Small"))
-			]
+			SNew(SModeWidget, FBlueprintEditorApplicationModes::GetLocalizedMode(FBlueprintEditorApplicationModes::BlueprintDefaultsMode), FBlueprintEditorApplicationModes::BlueprintDefaultsMode)
+			.OnGetActiveMode(ActiveModeGetter)
+			.OnSetActiveMode(ActiveModeSetter)
+			.CanBeSelected(BlueprintObj ? FBlueprintEditorUtils::DoesSupportDefaults(BlueprintObj) : false)
+			.ToolTip(IDocumentation::Get()->CreateToolTip(
+				LOCTEXT("BlueprintDefaultsModeButtonTooltip", "Switch to Blueprint Defaults Mode"),
+				NULL,
+				TEXT("Shared/Editors/BlueprintEditor"),
+				TEXT("DefaultsMode")))
+			.IconImage(FEditorStyle::GetBrush("FullBlueprintEditor.SwitchToBlueprintDefaultsMode"))
+			.SmallIconImage(FEditorStyle::GetBrush("FullBlueprintEditor.SwitchToBlueprintDefaultsMode.Small"))
+			.Tag(TEXT("DefaultsMode"))
 		);
 
 		Ret.Add(SNew(SBlueprintModeSeparator));
 
 		Ret.Add(
-			SNew(STutorialWrapper, TEXT("ComponentsMode"))
-			[
-				SNew(SModeWidget, FBlueprintEditorApplicationModes::GetLocalizedMode(FBlueprintEditorApplicationModes::BlueprintComponentsMode), FBlueprintEditorApplicationModes::BlueprintComponentsMode)
-				.OnGetActiveMode(ActiveModeGetter)
-				.OnSetActiveMode(ActiveModeSetter)
-				.CanBeSelected(BlueprintObj ? FBlueprintEditorUtils::DoesSupportComponents(BlueprintObj) : false)
-				.ToolTip(IDocumentation::Get()->CreateToolTip(
-					LOCTEXT("ComponentsModeButtonTooltip", "Switch to Components Mode"),
-					NULL,
-					TEXT("Shared/Editors/BlueprintEditor"),
-					TEXT("ComponentsMode")))
-				.IconImage(FEditorStyle::GetBrush("FullBlueprintEditor.SwitchToComponentsMode"))
-				.SmallIconImage(FEditorStyle::GetBrush("FullBlueprintEditor.SwitchToComponentsMode.Small"))
-			]
+			SNew(SModeWidget, FBlueprintEditorApplicationModes::GetLocalizedMode(FBlueprintEditorApplicationModes::BlueprintComponentsMode), FBlueprintEditorApplicationModes::BlueprintComponentsMode)
+			.OnGetActiveMode(ActiveModeGetter)
+			.OnSetActiveMode(ActiveModeSetter)
+			.CanBeSelected(BlueprintObj ? FBlueprintEditorUtils::DoesSupportComponents(BlueprintObj) : false)
+			.ToolTip(IDocumentation::Get()->CreateToolTip(
+				LOCTEXT("ComponentsModeButtonTooltip", "Switch to Components Mode"),
+				NULL,
+				TEXT("Shared/Editors/BlueprintEditor"),
+				TEXT("ComponentsMode")))
+			.IconImage(FEditorStyle::GetBrush("FullBlueprintEditor.SwitchToComponentsMode"))
+			.SmallIconImage(FEditorStyle::GetBrush("FullBlueprintEditor.SwitchToComponentsMode.Small"))
+			.Tag(TEXT("ComponentsMode"))
 		);
 
 		Ret.Add(SNew(SBlueprintModeSeparator));
 
 		Ret.Add(
-			SNew(STutorialWrapper, TEXT("GraphMode"))
-			[
-				SNew(SModeWidget, FBlueprintEditorApplicationModes::GetLocalizedMode(FBlueprintEditorApplicationModes::StandardBlueprintEditorMode), FBlueprintEditorApplicationModes::StandardBlueprintEditorMode)
-				.OnGetActiveMode(ActiveModeGetter)
-				.OnSetActiveMode(ActiveModeSetter)
-				.CanBeSelected(BlueprintObj != NULL)
-				.ToolTip(IDocumentation::Get()->CreateToolTip(
-					LOCTEXT("GraphModeButtonTooltip", "Switch to Graph Editing Mode"),
-					NULL,
-					TEXT("Shared/Editors/BlueprintEditor"),
-					TEXT("GraphMode")))
-				.ToolTipText(LOCTEXT("GraphModeButtonTooltip", "Switch to Graph Editing Mode"))
-				.IconImage(FEditorStyle::GetBrush("FullBlueprintEditor.SwitchToScriptingMode"))
-				.SmallIconImage(FEditorStyle::GetBrush("FullBlueprintEditor.SwitchToScriptingMode.Small"))
-			]
+			SNew(SModeWidget, FBlueprintEditorApplicationModes::GetLocalizedMode(FBlueprintEditorApplicationModes::StandardBlueprintEditorMode), FBlueprintEditorApplicationModes::StandardBlueprintEditorMode)
+			.OnGetActiveMode(ActiveModeGetter)
+			.OnSetActiveMode(ActiveModeSetter)
+			.CanBeSelected(BlueprintObj != NULL)
+			.ToolTip(IDocumentation::Get()->CreateToolTip(
+				LOCTEXT("GraphModeButtonTooltip", "Switch to Graph Editing Mode"),
+				NULL,
+				TEXT("Shared/Editors/BlueprintEditor"),
+				TEXT("GraphMode")))
+			.ToolTipText(LOCTEXT("GraphModeButtonTooltip", "Switch to Graph Editing Mode"))
+			.IconImage(FEditorStyle::GetBrush("FullBlueprintEditor.SwitchToScriptingMode"))
+			.SmallIconImage(FEditorStyle::GetBrush("FullBlueprintEditor.SwitchToScriptingMode.Small"))
+			.Tag(TEXT("GraphMode"))
 		);
 
 		// Right side padding
diff --git a/Engine/Source/Editor/Kismet/Private/SBlueprintPalette.cpp b/Engine/Source/Editor/Kismet/Private/SBlueprintPalette.cpp
--- a/Engine/Source/Editor/Kismet/Private/SBlueprintPalette.cpp
+++ b/Engine/Source/Editor/Kismet/Private/SBlueprintPalette.cpp
@@ -20,7 +20,6 @@
 #include "SBlueprintLibraryPalette.h"
 #include "SBlueprintFavoritesPalette.h"
 #include "BlueprintPaletteFavorites.h"
-#include "STutorialWrapper.h"
 #include "ClassIconFinder.h"
 #include "AnimationStateMachineGraph.h"
 #include "AnimationStateMachineSchema.h"
@@ -1515,29 +1514,23 @@ void SBlueprintPalette::Construct(const FArguments& InArgs, TWeakPtr<FBlueprintE
 	{
 		this->ChildSlot
 		[
-			SNew( STutorialWrapper, TEXT("FullBlueprintPalette") )
-			[
-				SAssignNew(PaletteSplitter, SSplitter)
-					.Orientation(Orient_Vertical)
-					.OnSplitterFinishedResizing(this, &SBlueprintPalette::OnSplitterResized)
+			SAssignNew(PaletteSplitter, SSplitter)
+				.Orientation(Orient_Vertical)
+				.OnSplitterFinishedResizing(this, &SBlueprintPalette::OnSplitterResized)
+				.Tag(TEXT("FullBlueprintPalette"))
 
-				+ SSplitter::Slot()
-				.Value(FavoritesHeightRatio)
-				[
-					SAssignNew( FavoritesWrapper, STutorialWrapper, TEXT("BlueprintPaletteFavorites") )
-					[
-						SNew(SBlueprintFavoritesPalette, InBlueprintEditor)
-					]
-				]
+			+ SSplitter::Slot()
+			.Value(FavoritesHeightRatio)
+			[
+				SNew(SBlueprintFavoritesPalette, InBlueprintEditor)
+				.Tag(TEXT("BlueprintPaletteFavorites"))
+			]
 
-				+ SSplitter::Slot()
-				.Value(LibraryHeightRatio)
-				[
-					SAssignNew( LibraryWrapper, STutorialWrapper, TEXT("BlueprintPaletteLibrary") )
-					[
-						SNew(SBlueprintLibraryPalette, InBlueprintEditor)
-					]
-				]
+			+ SSplitter::Slot()
+			.Value(LibraryHeightRatio)
+			[
+				SNew(SBlueprintLibraryPalette, InBlueprintEditor)
+				.Tag(TEXT("BlueprintPaletteLibrary"))
 			]
 		];
 	}	
diff --git a/Engine/Source/Editor/Kismet/Private/SGraphTitleBar.cpp b/Engine/Source/Editor/Kismet/Private/SGraphTitleBar.cpp
--- a/Engine/Source/Editor/Kismet/Private/SGraphTitleBar.cpp
+++ b/Engine/Source/Editor/Kismet/Private/SGraphTitleBar.cpp
@@ -83,82 +83,80 @@ void SGraphTitleBar::Construct( const FArguments& InArgs )
 	
 	this->ChildSlot
 	[
-		SNew( STutorialWrapper, TEXT("EventGraphTitleBar") )
+		SNew(SBorder)
+		.BorderImage( FEditorStyle::GetBrush( TEXT("Graph.TitleBackground") ) )
+		.HAlign(HAlign_Fill)
+		.Tag(TEXT("EventGraphTitleBar"))
 		[
-			SNew(SBorder)
-			.BorderImage( FEditorStyle::GetBrush( TEXT("Graph.TitleBackground") ) )
-			.HAlign(HAlign_Fill)
+			SNew(SVerticalBox)
+			// Title text/icon
+			+SVerticalBox::Slot()
+			.AutoHeight()
 			[
-				SNew(SVerticalBox)
-				// Title text/icon
-				+SVerticalBox::Slot()
-				.AutoHeight()
+				SNew(SHorizontalBox)
+
+				+SHorizontalBox::Slot()
+				.AutoWidth()
+				[
+					InArgs._HistoryNavigationWidget.ToSharedRef()
+				]
+
+				+SHorizontalBox::Slot()
+				.HAlign(HAlign_Center)
+				.FillWidth(1.f)
 				[
 					SNew(SHorizontalBox)
+					+SHorizontalBox::Slot()
+					.AutoWidth()
+					.Padding( 10,5 )
+					.VAlign(VAlign_Center)
+					[
+						SNew(SImage)
+						.Image( this, &SGraphTitleBar::GetTypeGlyph )
+					]
+
+					// show fake 'root' breadcrumb for the title
+					+SHorizontalBox::Slot()
+					.AutoWidth()
+					.VAlign(VAlign_Center)
+					.Padding(BreadcrumbTrailPadding)
+					[
+						SNew(STextBlock)
+						.Text(this, &SGraphTitleBar::GetBlueprintTitle )
+						.TextStyle( FEditorStyle::Get(), TEXT("GraphBreadcrumbButtonText") )
+						.Visibility( this, &SGraphTitleBar::IsGraphBlueprintNameVisible )
+					]
+					+SHorizontalBox::Slot()
+					.AutoWidth()
+					.VAlign(VAlign_Center)
+					[
+						SNew(SImage)
+						.Image( BreadcrumbButtonImage )
+						.Visibility( this, &SGraphTitleBar::IsGraphBlueprintNameVisible )
+					]
 
+					// New style breadcrumb
 					+SHorizontalBox::Slot()
 					.AutoWidth()
+					.VAlign(VAlign_Center)
 					[
-						InArgs._HistoryNavigationWidget.ToSharedRef()
+						SAssignNew(BreadcrumbTrail, SBreadcrumbTrail<UEdGraph*>)
+						.ButtonStyle(FEditorStyle::Get(), "GraphBreadcrumbButton")
+						.TextStyle(FEditorStyle::Get(), "GraphBreadcrumbButtonText")
+						.ButtonContentPadding( BreadcrumbTrailPadding )
+						.DelimiterImage( BreadcrumbButtonImage )
+						.PersistentBreadcrumbs( true )
+						.OnCrumbClicked( this, &SGraphTitleBar::OnBreadcrumbClicked )
 					]
 
 					+SHorizontalBox::Slot()
-					.HAlign(HAlign_Center)
-					.FillWidth(1.f)
+					.AutoWidth()
+					.VAlign(VAlign_Center)
 					[
-						SNew(SHorizontalBox)
-						+SHorizontalBox::Slot()
-						.AutoWidth()
-						.Padding( 10,5 )
-						.VAlign(VAlign_Center)
-						[
-							SNew(SImage)
-							.Image( this, &SGraphTitleBar::GetTypeGlyph )
-						]
-
-						// show fake 'root' breadcrumb for the title
-						+SHorizontalBox::Slot()
-						.AutoWidth()
-						.VAlign(VAlign_Center)
-						.Padding(BreadcrumbTrailPadding)
-						[
-							SNew(STextBlock)
-							.Text(this, &SGraphTitleBar::GetBlueprintTitle )
-							.TextStyle( FEditorStyle::Get(), TEXT("GraphBreadcrumbButtonText") )
-							.Visibility( this, &SGraphTitleBar::IsGraphBlueprintNameVisible )
-						]
-						+SHorizontalBox::Slot()
-						.AutoWidth()
-						.VAlign(VAlign_Center)
-						[
-							SNew(SImage)
-							.Image( BreadcrumbButtonImage )
-							.Visibility( this, &SGraphTitleBar::IsGraphBlueprintNameVisible )
-						]
-
-						// New style breadcrumb
-						+SHorizontalBox::Slot()
-						.AutoWidth()
-						.VAlign(VAlign_Center)
-						[
-							SAssignNew(BreadcrumbTrail, SBreadcrumbTrail<UEdGraph*>)
-							.ButtonStyle(FEditorStyle::Get(), "GraphBreadcrumbButton")
-							.TextStyle(FEditorStyle::Get(), "GraphBreadcrumbButtonText")
-							.ButtonContentPadding( BreadcrumbTrailPadding )
-							.DelimiterImage( BreadcrumbButtonImage )
-							.PersistentBreadcrumbs( true )
-							.OnCrumbClicked( this, &SGraphTitleBar::OnBreadcrumbClicked )
-						]
-
-						+SHorizontalBox::Slot()
-						.AutoWidth()
-						.VAlign(VAlign_Center)
-						[
-							SNew(STextBlock)
-							.Font( FSlateFontInfo(FPaths::EngineContentDir() / TEXT("Slate/Fonts/Roboto-Regular.ttf"), 14 ) )
-							.ColorAndOpacity( FLinearColor(1,1,1,0.5) )
-							.Text( this, &SGraphTitleBar::GetTitleExtra )
-						]
+						SNew(STextBlock)
+						.Font( FSlateFontInfo(FPaths::EngineContentDir() / TEXT("Slate/Fonts/Roboto-Regular.ttf"), 14 ) )
+						.ColorAndOpacity( FLinearColor(1,1,1,0.5) )
+						.Text( this, &SGraphTitleBar::GetTitleExtra )
 					]
 				]
 			]
diff --git a/Engine/Source/Editor/Kismet/Private/SKismetInspector.cpp b/Engine/Source/Editor/Kismet/Private/SKismetInspector.cpp
--- a/Engine/Source/Editor/Kismet/Private/SKismetInspector.cpp
+++ b/Engine/Source/Editor/Kismet/Private/SKismetInspector.cpp
@@ -8,7 +8,6 @@
 #include "Editor/PropertyEditor/Public/PropertyEditorModule.h"
 #include "SKismetInspector.h"
 #include "SKismetLinearExpression.h"
-#include "STutorialWrapper.h"
 
 #include "Editor/PropertyEditor/Public/PropertyEditing.h"
 
@@ -241,17 +240,14 @@ void SKismetInspector::Construct(const FArguments& InArgs)
 	// Create the border that all of the content will get stuffed into
 	ChildSlot
 	[
-		SNew( STutorialWrapper, TEXT("BlueprintInspector") )
+		SNew(SVerticalBox)
+		.Tag( TEXT("BlueprintInspector"))
+		+ SVerticalBox::Slot()
+		.FillHeight(1.0f)
 		[
-			SNew(SVerticalBox)
-
-			+ SVerticalBox::Slot()
-			.FillHeight(1.0f)
-			[
-				SAssignNew( ContextualEditingBorderWidget, SBorder )
-				.Padding(0)
-				.BorderImage( FEditorStyle::GetBrush("NoBorder") )
-			]
+			SAssignNew( ContextualEditingBorderWidget, SBorder )
+			.Padding(0)
+			.BorderImage( FEditorStyle::GetBrush("NoBorder") )
 		]
 	];
 
diff --git a/Engine/Source/Editor/Kismet/Private/SMyBlueprint.cpp b/Engine/Source/Editor/Kismet/Private/SMyBlueprint.cpp
--- a/Engine/Source/Editor/Kismet/Private/SMyBlueprint.cpp
+++ b/Engine/Source/Editor/Kismet/Private/SMyBlueprint.cpp
@@ -5,7 +5,6 @@
 #include "BlueprintUtilities.h"
 #include "ScopedTransaction.h"
 #include "GraphEditor.h"
-#include "STutorialWrapper.h"
 
 #include "BlueprintEditor.h"
 
@@ -204,54 +203,52 @@ void SMyBlueprint::Construct(const FArguments& InArgs, TWeakPtr<FBlueprintEditor
 	// now piece together all the content for this widget
 	ChildSlot
 	[
-		SNew( STutorialWrapper, TEXT("MyBlueprintPanel") )
+		SNew(SBorder)
+		.Padding(4.0f)
+		.BorderImage(FEditorStyle::GetBrush("ToolPanel.GroupBorder"))
+		.Tag(TEXT("MyBlueprintPanel"))
 		[
-			SNew(SBorder)
-			.Padding(4.0f)
-			.BorderImage(FEditorStyle::GetBrush("ToolPanel.GroupBorder"))
-			[
-				SNew(SVerticalBox)
+			SNew(SVerticalBox)
 
-				+ SVerticalBox::Slot()
-				.AutoHeight()
-				[
-					ToolbarBuilder.MakeWidget()
-				]
+			+ SVerticalBox::Slot()
+			.AutoHeight()
+			[
+				ToolbarBuilder.MakeWidget()
+			]
 				
-				+ SVerticalBox::Slot()
-				.AutoHeight()
-				[
-					FilterBox.ToSharedRef()
-				]
+			+ SVerticalBox::Slot()
+			.AutoHeight()
+			[
+				FilterBox.ToSharedRef()
+			]
 				
-				+ SVerticalBox::Slot()
-				.FillHeight(1.0f)
-				[
-					SAssignNew(ActionMenuContainer, SSplitter)
-					.Orientation(Orient_Vertical)
+			+ SVerticalBox::Slot()
+			.FillHeight(1.0f)
+			[
+				SAssignNew(ActionMenuContainer, SSplitter)
+				.Orientation(Orient_Vertical)
 
-					+ SSplitter::Slot()
-					[
-						MyBlueprintActionMenu
-					]
+				+ SSplitter::Slot()
+				[
+					MyBlueprintActionMenu
 				]
+			]
 
-				+ SVerticalBox::Slot()
-				.AutoHeight()
-				.Padding(FMargin(3.0f, 2.0f, 0.0f, 0.0f))
+			+ SVerticalBox::Slot()
+			.AutoHeight()
+			.Padding(FMargin(3.0f, 2.0f, 0.0f, 0.0f))
+			[
+				SNew(SCheckBox)
+				.IsChecked(this, &SMyBlueprint::OnUserVarsCheckState)
+				.OnCheckStateChanged(this, &SMyBlueprint::OnUserVarsCheckStateChanged)
 				[
-					SNew(SCheckBox)
-					.IsChecked(this, &SMyBlueprint::OnUserVarsCheckState)
-					.OnCheckStateChanged(this, &SMyBlueprint::OnUserVarsCheckStateChanged)
-					[
-						SNew(STextBlock)
-						.Text(LOCTEXT("ShowInheritedVariables", "Show inherited variables"))
-						.ToolTip(IDocumentation::Get()->CreateToolTip(
-							LOCTEXT("ShowInheritedVariablesTooltip", "Should inherited variables from parent classes and blueprints be shown in the tree?"),
-							NULL,
-							TEXT("Shared/Editors/BlueprintEditor"),
-							TEXT("MyBlueprint_ShowInheritedVariables")))
-					]
+					SNew(STextBlock)
+					.Text(LOCTEXT("ShowInheritedVariables", "Show inherited variables"))
+					.ToolTip(IDocumentation::Get()->CreateToolTip(
+						LOCTEXT("ShowInheritedVariablesTooltip", "Should inherited variables from parent classes and blueprints be shown in the tree?"),
+						NULL,
+						TEXT("Shared/Editors/BlueprintEditor"),
+						TEXT("MyBlueprint_ShowInheritedVariables")))
 				]
 			]
 		]
diff --git a/Engine/Source/Editor/Kismet/Private/SSCSEditor.cpp b/Engine/Source/Editor/Kismet/Private/SSCSEditor.cpp
--- a/Engine/Source/Editor/Kismet/Private/SSCSEditor.cpp
+++ b/Engine/Source/Editor/Kismet/Private/SSCSEditor.cpp
@@ -10,7 +10,6 @@
 #include "SKismetInspector.h"
 #include "SSCSEditorViewport.h"
 #include "SComponentClassCombo.h"
-#include "STutorialWrapper.h"
 
 #include "AssetSelection.h"
 #include "Editor/SceneOutliner/Private/SSocketChooser.h"
@@ -2009,7 +2008,8 @@ void SSCSEditor::Construct( const FArguments& InArgs, TSharedPtr<FBlueprintEdito
 
 	this->ChildSlot
 	[
-		SNew( STutorialWrapper, TEXT("ComponentsPanel") )
+		SNew( SBox )
+		.Tag(TEXT("ComponentsPanel"))
 		[
 			Contents.ToSharedRef()
 		]
diff --git a/Engine/Source/Editor/LevelEditor/Private/SLevelEditor.cpp b/Engine/Source/Editor/LevelEditor/Private/SLevelEditor.cpp
--- a/Engine/Source/Editor/LevelEditor/Private/SLevelEditor.cpp
+++ b/Engine/Source/Editor/LevelEditor/Private/SLevelEditor.cpp
@@ -34,7 +34,6 @@
 #include "Editor/StatsViewer/Public/StatsViewerModule.h"
 #include "Editor/UMGEditor/Public/UMGEditorModule.h"
 #include "EditorModes.h"
-#include "STutorialWrapper.h"
 #include "IDocumentation.h"
 #include "NewsFeed.h"
 
@@ -187,7 +186,8 @@ void SLevelEditor::Initialize( const TSharedRef<SDockTab>& OwnerTab, const TShar
 
 			+SOverlay::Slot()
 			[
-				SNew( STutorialWrapper, TEXT("MainMenu") )
+				SNew( SBox )
+				.Tag(TEXT("MainMenu"))
 				[
 					Widget1
 				]
@@ -198,10 +198,8 @@ void SLevelEditor::Initialize( const TSharedRef<SDockTab>& OwnerTab, const TShar
 			+SOverlay::Slot()
 			.HAlign( HAlign_Right )
 			[
-				SNew( STutorialWrapper, TEXT("PerformanceTools") )
-				[
-					SAssignNew( NotificationBarBox, SHorizontalBox )
-				]
+				SAssignNew( NotificationBarBox, SHorizontalBox )
+				.Tag(TEXT("PerformanceTools"))
 			]
 #endif
 		]
@@ -214,13 +212,11 @@ void SLevelEditor::Initialize( const TSharedRef<SDockTab>& OwnerTab, const TShar
 	];
 	
 // For OS X we need to put it into the window's title bar since there's no per-window menu bar
-#if PLATFORM_MAC
-	TSharedRef<SWidget> TutorialWidget = SNew( STutorialWrapper, TEXT("PerformanceTools") )
-	[
+#if PLATFORM_MAC	
+	OwnerTab->SetRightContent(
 		SAssignNew( NotificationBarBox, SHorizontalBox )
-	];
-	
-	OwnerTab->SetRightContent(TutorialWidget);
+		.Tag(TEXT("PerformanceTools"))
+		);
 #endif
 
 	ConstructNotificationBar();
@@ -553,7 +549,8 @@ static TSharedRef<SDockTab> SummonDetailsPanel( FName TabIdentifier )
 		.Label( Label )
 		.ToolTip( IDocumentation::Get()->CreateToolTip( Label, nullptr, "Shared/LevelEditor", "DetailsTab" ) )
 		[
-			SNew( STutorialWrapper, TEXT("ActorDetails") )
+			SNew( SBox )
+			.Tag(TEXT("ActorDetails"))
 			[
 				DetailsView
 			]
@@ -586,17 +583,14 @@ TSharedRef<SDockTab> SLevelEditor::SpawnLevelEditorTab( const FSpawnTabArgs& Arg
 			.ShouldAutosize(true)
 			.Icon( FEditorStyle::GetBrush("ToolBar.Icon") )
 			[
-				SNew( STutorialWrapper, TEXT("LevelToolbar") )
+				SNew(SHorizontalBox)
+				.Tag(TEXT("LevelToolbar"))
+				+ SHorizontalBox::Slot()
+				.FillWidth(1)
+				.VAlign(VAlign_Bottom)
+				.HAlign(HAlign_Left)
 				[
-					SNew(SHorizontalBox)
-
-					+ SHorizontalBox::Slot()
-					.FillWidth(1)
-					.VAlign(VAlign_Bottom)
-					.HAlign(HAlign_Left)
-					[
-						FLevelEditorToolBar::MakeLevelEditorToolBar( LevelEditorCommands.ToSharedRef(), SharedThis(this) )
-					]
+					FLevelEditorToolBar::MakeLevelEditorToolBar( LevelEditorCommands.ToSharedRef(), SharedThis(this) )
 				]
 			];
 
@@ -619,7 +613,8 @@ TSharedRef<SDockTab> SLevelEditor::SpawnLevelEditorTab( const FSpawnTabArgs& Arg
 			.Icon( FEditorStyle::GetBrush( "LevelEditor.Tabs.Modes" ) )
 			.Label( NSLOCTEXT( "LevelEditor", "ToolsTabTitle", "Modes" ) )
 			[
-				SNew( STutorialWrapper, TEXT("ToolsPanel") )
+				SNew( SBox )
+				.Tag(TEXT("ToolsPanel"))
 				[
 					NewToolBox
 				]
@@ -669,16 +664,14 @@ TSharedRef<SDockTab> SLevelEditor::SpawnLevelEditorTab( const FSpawnTabArgs& Arg
 			.ToolTip( IDocumentation::Get()->CreateToolTip( Label, nullptr, "Shared/LevelEditor", "SceneOutlinerTab" ) )
 			.ContentPadding( 5 )
 			[
-				SNew( STutorialWrapper, TEXT("SceneOutliner") )
+				SNew(SBorder)
+				.Padding(4)
+				.BorderImage( FEditorStyle::GetBrush("ToolPanel.GroupBorder") )
+				.Tag(TEXT("SceneOutliner"))
 				[
-					SNew(SBorder)
-					.Padding(4)
-					.BorderImage( FEditorStyle::GetBrush("ToolPanel.GroupBorder") )
-					[
-						SceneOutlinerModule.CreateSceneOutliner(
-							InitOptions,
-							FOnActorPicked() /* Not used for outliner when in browsing mode */ )
-					]
+					SceneOutlinerModule.CreateSceneOutliner(
+						InitOptions,
+						FOnActorPicked() /* Not used for outliner when in browsing mode */ )
 				]
 			];
 	}
diff --git a/Engine/Source/Editor/MaterialEditor/Private/MaterialEditor.cpp b/Engine/Source/Editor/MaterialEditor/Private/MaterialEditor.cpp
--- a/Engine/Source/Editor/MaterialEditor/Private/MaterialEditor.cpp
+++ b/Engine/Source/Editor/MaterialEditor/Private/MaterialEditor.cpp
@@ -41,7 +41,6 @@
 #include "PreviewScene.h"
 #include "ScopedTransaction.h"
 #include "BusyCursor.h"
-#include "STutorialWrapper.h"
 
 #include "Editor/PropertyEditor/Public/PropertyEditorModule.h"
 #include "Editor/PropertyEditor/Public/IDetailsView.h"
@@ -2475,7 +2474,8 @@ TSharedRef<SDockTab> FMaterialEditor::SpawnTab_Palette(const FSpawnTabArgs& Args
 		.Icon(FEditorStyle::GetBrush("Kismet.Tabs.Palette"))
 		.Label(LOCTEXT("MaterialPaletteTitle", "Palette"))
 		[
-			SNew( STutorialWrapper, TEXT("MaterialPalette") )
+			SNew( SBox )
+			.Tag(TEXT("MaterialPalette"))
 			[
 				Palette.ToSharedRef()
 			]
@@ -2492,7 +2492,8 @@ TSharedRef<SDockTab> FMaterialEditor::SpawnTab_Stats(const FSpawnTabArgs& Args)
 		.Icon(FEditorStyle::GetBrush("Kismet.Tabs.CompilerResults"))
 		.Label(LOCTEXT("MaterialStatsTitle", "Stats"))
 		[
-			SNew( STutorialWrapper, TEXT("MaterialStats") )
+			SNew( SBox )
+			.Tag(TEXT("MaterialStats"))
 			[
 				Stats.ToSharedRef()
 			]
@@ -2509,7 +2510,8 @@ TSharedRef<SDockTab> FMaterialEditor::SpawnTab_Find(const FSpawnTabArgs& Args)
 		.Icon(FEditorStyle::GetBrush("Kismet.Tabs.FindResults"))
 		.Label(LOCTEXT("MaterialFindTitle", "Find Results"))
 		[
-			SNew(STutorialWrapper, TEXT("MaterialFind"))
+			SNew(SBox)
+			.Tag(TEXT("MaterialFind"))
 			[
 				FindResults.ToSharedRef()
 			]
diff --git a/Engine/Source/Editor/PlacementMode/Private/SPlacementModeTools.cpp b/Engine/Source/Editor/PlacementMode/Private/SPlacementModeTools.cpp
--- a/Engine/Source/Editor/PlacementMode/Private/SPlacementModeTools.cpp
+++ b/Engine/Source/Editor/PlacementMode/Private/SPlacementModeTools.cpp
@@ -10,8 +10,6 @@
 #include "KismetEditorUtilities.h"
 #include "EditorClassUtils.h"
 
-#include "STutorialWrapper.h"
-
 /**
  * These are the tab indexes, if the tabs are reorganized you need to adjust the
  * enum accordingly.
@@ -322,7 +320,8 @@ TSharedRef< SWidget > SPlacementModeTools::CreateStandardPanel()
 		.Padding( 0, 3, 0, 0 )
 		.AutoHeight()
 		[
-			SNew( STutorialWrapper, TEXT("PMRecentlyPlaced") )
+			SNew( SBox)
+			.Tag(TEXT("PMRecentlyPlaced"))
 			[
 				CreatePlacementGroupTab( (int32)EPlacementTab::RecentlyPlaced, NSLOCTEXT( "PlacementMode", "RecentlyPlaced", "Recently Placed" ), true )
 			]
@@ -331,7 +330,8 @@ TSharedRef< SWidget > SPlacementModeTools::CreateStandardPanel()
 		+ SVerticalBox::Slot()
 		.AutoHeight()
 		[
-			SNew( STutorialWrapper, TEXT("PMGeometry") )
+			SNew( SBox )
+			.Tag(TEXT("PMGeometry"))
 			[
 				CreatePlacementGroupTab( (int32)EPlacementTab::Geometry, NSLOCTEXT( "PlacementMode", "Geometry", "Geometry" ), false )
 			]
@@ -340,7 +340,8 @@ TSharedRef< SWidget > SPlacementModeTools::CreateStandardPanel()
 		+ SVerticalBox::Slot()
 		.AutoHeight()
 		[
-			SNew( STutorialWrapper, TEXT("PMLights") )
+			SNew( SBox )
+			.Tag(TEXT("PMLights"))
 			[
 				CreatePlacementGroupTab( (int32)EPlacementTab::Lights, NSLOCTEXT( "PlacementMode", "Lights", "Lights" ), false )
 			]
@@ -349,7 +350,8 @@ TSharedRef< SWidget > SPlacementModeTools::CreateStandardPanel()
 		+ SVerticalBox::Slot()
 		.AutoHeight()
 		[
-			SNew( STutorialWrapper, TEXT("PMVisual") )
+			SNew( SBox )
+			.Tag(TEXT("PMVisual"))
 			[
 				CreatePlacementGroupTab( (int32)EPlacementTab::Visual, NSLOCTEXT( "PlacementMode", "Visual", "Visual" ), false )
 			]
@@ -358,7 +360,8 @@ TSharedRef< SWidget > SPlacementModeTools::CreateStandardPanel()
 		+ SVerticalBox::Slot()
 		.AutoHeight()
 		[
-			SNew( STutorialWrapper, TEXT("PMBasic") )
+			SNew( SBox )
+			.Tag(TEXT("PMBasic"))
 			[
 				CreatePlacementGroupTab( (int32)EPlacementTab::Basic, NSLOCTEXT( "PlacementMode", "Basic", "Basic" ), false )
 			]
@@ -367,7 +370,8 @@ TSharedRef< SWidget > SPlacementModeTools::CreateStandardPanel()
 		+ SVerticalBox::Slot()
 		.AutoHeight()
 		[
-			SNew( STutorialWrapper, TEXT("PMVolumes") )
+			SNew( SBox )
+			.Tag(TEXT("PMVolumes"))
 			[
 				CreatePlacementGroupTab( (int32)EPlacementTab::Volumes, NSLOCTEXT( "PlacementMode", "Volumes", "Volumes" ), false )
 			]
@@ -376,7 +380,8 @@ TSharedRef< SWidget > SPlacementModeTools::CreateStandardPanel()
 		+ SVerticalBox::Slot()
 		.AutoHeight()
 		[
-			SNew( STutorialWrapper, TEXT("PMAllClasses") )
+			SNew( SBox )
+			.Tag(TEXT("PMAllClasses"))
 			[
 				CreatePlacementGroupTab( (int32)EPlacementTab::AllClasses, NSLOCTEXT( "PlacementMode", "AllClasses", "All Classes" ), true )
 			]
diff --git a/Engine/Source/Editor/PropertyEditor/Private/UserInterface/PropertyEditor/SPropertyEditorText.cpp b/Engine/Source/Editor/PropertyEditor/Private/UserInterface/PropertyEditor/SPropertyEditorText.cpp
--- a/Engine/Source/Editor/PropertyEditor/Private/UserInterface/PropertyEditor/SPropertyEditorText.cpp
+++ b/Engine/Source/Editor/PropertyEditor/Private/UserInterface/PropertyEditor/SPropertyEditorText.cpp
@@ -12,17 +12,37 @@ void SPropertyEditorText::Construct( const FArguments& InArgs, const TSharedRef<
 {
 	PropertyEditor = InPropertyEditor;
 
-	ChildSlot
-	[
-		SAssignNew( PrimaryWidget, SEditableTextBox )
-		.Text( InPropertyEditor, &FPropertyEditor::GetValueAsText )
-		.Font( InArgs._Font )
-		.SelectAllTextWhenFocused( true )
-		.ClearKeyboardFocusOnCommit(false)
-		.OnTextCommitted( this, &SPropertyEditorText::OnTextCommitted )
-		.SelectAllTextOnCommit( true )
-		.IsReadOnly(this, &SPropertyEditorText::IsReadOnly)
-	];
+	const UProperty* Property = InPropertyEditor->GetProperty();
+	bIsMultiLine = Property->GetBoolMetaData("MultiLine");
+	if(bIsMultiLine)
+	{
+		ChildSlot
+		[
+			SAssignNew( PrimaryWidget, SMultiLineEditableTextBox)
+			.Text( InPropertyEditor, &FPropertyEditor::GetValueAsText )
+			.Font( InArgs._Font )
+			.SelectAllTextWhenFocused( false )
+			.ClearKeyboardFocusOnCommit(false)
+			.OnTextCommitted( this, &SPropertyEditorText::OnTextCommitted )
+			.SelectAllTextOnCommit( false )
+			.IsReadOnly(this, &SPropertyEditorText::IsReadOnly)
+			.AutoWrapText(true)
+		];
+	}
+	else
+	{
+		ChildSlot
+		[
+			SAssignNew( PrimaryWidget, SEditableTextBox )
+			.Text( InPropertyEditor, &FPropertyEditor::GetValueAsText )
+			.Font( InArgs._Font )
+			.SelectAllTextWhenFocused( true )
+			.ClearKeyboardFocusOnCommit(false)
+			.OnTextCommitted( this, &SPropertyEditorText::OnTextCommitted )
+			.SelectAllTextOnCommit( true )
+			.IsReadOnly(this, &SPropertyEditorText::IsReadOnly)
+		];
+	}
 
 	if( InPropertyEditor->PropertyIsA( UObjectPropertyBase::StaticClass() ) )
 	{
@@ -35,7 +55,15 @@ void SPropertyEditorText::Construct( const FArguments& InArgs, const TSharedRef<
 
 void SPropertyEditorText::GetDesiredWidth( float& OutMinDesiredWidth, float& OutMaxDesiredWidth )
 {
-	OutMinDesiredWidth = 125.0f;
+	if(bIsMultiLine)
+	{
+		OutMinDesiredWidth = 250.0f;
+	}
+	else
+	{
+		OutMinDesiredWidth = 125.0f;
+	}
+	
 	OutMaxDesiredWidth = 600.0f;
 }
 
diff --git a/Engine/Source/Editor/PropertyEditor/Private/UserInterface/PropertyEditor/SPropertyEditorText.h b/Engine/Source/Editor/PropertyEditor/Private/UserInterface/PropertyEditor/SPropertyEditorText.h
--- a/Engine/Source/Editor/PropertyEditor/Private/UserInterface/PropertyEditor/SPropertyEditorText.h
+++ b/Engine/Source/Editor/PropertyEditor/Private/UserInterface/PropertyEditor/SPropertyEditorText.h
@@ -38,4 +38,7 @@ private:
 	TSharedPtr< class FPropertyEditor > PropertyEditor;
 
 	TSharedPtr< class SWidget > PrimaryWidget;
+
+	/** Cached flag as we would like multi-line text widgets to be slightly larger */
+	bool bIsMultiLine;
 };
\ No newline at end of file
diff --git a/Engine/Source/Editor/UMGEditor/Private/TabFactory/DesignerTabSummoner.cpp b/Engine/Source/Editor/UMGEditor/Private/TabFactory/DesignerTabSummoner.cpp
--- a/Engine/Source/Editor/UMGEditor/Private/TabFactory/DesignerTabSummoner.cpp
+++ b/Engine/Source/Editor/UMGEditor/Private/TabFactory/DesignerTabSummoner.cpp
@@ -28,10 +28,8 @@ TSharedRef<SWidget> FDesignerTabSummoner::CreateTabBody(const FWorkflowTabSpawnI
 		+ SHorizontalBox::Slot()
 		.FillWidth(1)
 		[
-			SNew(STutorialWrapper, TEXT("Designer"))
-			[
-				SNew(SDesignerView, BlueprintEditor.Pin())
-			]
+			SNew(SDesignerView, BlueprintEditor.Pin())
+			.Tag(TEXT("Designer"))
 		];
 }
 
diff --git a/Engine/Source/Editor/UMGEditor/Private/TabFactory/DetailsTabSummoner.cpp b/Engine/Source/Editor/UMGEditor/Private/TabFactory/DetailsTabSummoner.cpp
--- a/Engine/Source/Editor/UMGEditor/Private/TabFactory/DetailsTabSummoner.cpp
+++ b/Engine/Source/Editor/UMGEditor/Private/TabFactory/DetailsTabSummoner.cpp
@@ -26,10 +26,8 @@ TSharedRef<SWidget> FDetailsTabSummoner::CreateTabBody(const FWorkflowTabSpawnIn
 {
 	TSharedPtr<FWidgetBlueprintEditor> BlueprintEditorPtr = StaticCastSharedPtr<FWidgetBlueprintEditor>(BlueprintEditor.Pin());
 
-	return SNew(STutorialWrapper, TEXT("Details"))
-		[
-			SNew(SWidgetDetailsView, BlueprintEditorPtr)
-		];
+	return SNew(SWidgetDetailsView, BlueprintEditorPtr)
+			.Tag(TEXT("Details"));
 }
 
 #undef LOCTEXT_NAMESPACE 
diff --git a/Engine/Source/Editor/UMGEditor/Private/TabFactory/HierarchyTabSummoner.cpp b/Engine/Source/Editor/UMGEditor/Private/TabFactory/HierarchyTabSummoner.cpp
--- a/Engine/Source/Editor/UMGEditor/Private/TabFactory/HierarchyTabSummoner.cpp
+++ b/Engine/Source/Editor/UMGEditor/Private/TabFactory/HierarchyTabSummoner.cpp
@@ -26,10 +26,8 @@ TSharedRef<SWidget> FHierarchyTabSummoner::CreateTabBody(const FWorkflowTabSpawn
 {
 	TSharedPtr<FWidgetBlueprintEditor> BlueprintEditorPtr = StaticCastSharedPtr<FWidgetBlueprintEditor>(BlueprintEditor.Pin());
 
-	return SNew(STutorialWrapper, TEXT("Hierarchy"))
-		[
-			SNew(SHierarchyView, BlueprintEditorPtr, BlueprintEditorPtr->GetBlueprintObj()->SimpleConstructionScript)
-		];
+	return SNew(SHierarchyView, BlueprintEditorPtr, BlueprintEditorPtr->GetBlueprintObj()->SimpleConstructionScript)
+			.Tag(TEXT("Hierarchy"));
 }
 
 #undef LOCTEXT_NAMESPACE 
diff --git a/Engine/Source/Editor/UMGEditor/Private/TabFactory/PaletteTabSummoner.cpp b/Engine/Source/Editor/UMGEditor/Private/TabFactory/PaletteTabSummoner.cpp
--- a/Engine/Source/Editor/UMGEditor/Private/TabFactory/PaletteTabSummoner.cpp
+++ b/Engine/Source/Editor/UMGEditor/Private/TabFactory/PaletteTabSummoner.cpp
@@ -26,10 +26,8 @@ TSharedRef<SWidget> FPaletteTabSummoner::CreateTabBody(const FWorkflowTabSpawnIn
 {
 	TSharedPtr<FWidgetBlueprintEditor> BlueprintEditorPtr = StaticCastSharedPtr<FWidgetBlueprintEditor>(BlueprintEditor.Pin());
 
-	return SNew(STutorialWrapper, TEXT("Palette"))
-		[
-			SNew(SPaletteView, BlueprintEditorPtr)
-		];
+	return SNew(SPaletteView, BlueprintEditorPtr)
+			.Tag(TEXT("Palette"));
 }
 
 #undef LOCTEXT_NAMESPACE 
diff --git a/Engine/Source/Editor/UMGEditor/Private/TabFactory/SequencerTabSummoner.cpp b/Engine/Source/Editor/UMGEditor/Private/TabFactory/SequencerTabSummoner.cpp
--- a/Engine/Source/Editor/UMGEditor/Private/TabFactory/SequencerTabSummoner.cpp
+++ b/Engine/Source/Editor/UMGEditor/Private/TabFactory/SequencerTabSummoner.cpp
@@ -27,7 +27,8 @@ TSharedRef<SWidget> FSequencerTabSummoner::CreateTabBody(const FWorkflowTabSpawn
 {
 	TSharedPtr<FWidgetBlueprintEditor> BlueprintEditorPinned = BlueprintEditor.Pin();
 
-	return SNew(STutorialWrapper, TEXT("Sequencer"))
+	return SNew(SBox)
+		.Tag(TEXT("Sequencer"))
 		[
 			BlueprintEditorPinned->GetSequencer()->GetSequencerWidget()
 		];
diff --git a/Engine/Source/Editor/UMGEditor/Private/WidgetBlueprintEditorToolbar.cpp b/Engine/Source/Editor/UMGEditor/Private/WidgetBlueprintEditorToolbar.cpp
--- a/Engine/Source/Editor/UMGEditor/Private/WidgetBlueprintEditorToolbar.cpp
+++ b/Engine/Source/Editor/UMGEditor/Private/WidgetBlueprintEditorToolbar.cpp
@@ -73,39 +73,35 @@ void FWidgetBlueprintEditorToolbar::FillWidgetBlueprintEditorModesToolbar(FToolB
 		BlueprintEditorPtr->AddToolbarWidget(SNew(SSpacer).Size(FVector2D(4.0f, 1.0f)));
 
 		BlueprintEditorPtr->AddToolbarWidget(
-			SNew( STutorialWrapper, TEXT("DefaultsMode") )
-			[
-				SNew(SModeWidget, FWidgetBlueprintApplicationModes::GetLocalizedMode(FWidgetBlueprintApplicationModes::DesignerMode), FWidgetBlueprintApplicationModes::DesignerMode)
-				.OnGetActiveMode(GetActiveMode)
-				.OnSetActiveMode(SetActiveMode)
-				.ToolTip(IDocumentation::Get()->CreateToolTip(
-					LOCTEXT("BlueprintDefaultsModeButtonTooltip", "Switch to Blueprint Defaults Mode"),
-					NULL,
-					TEXT("Shared/Editors/BlueprintEditor"),
-					TEXT("DefaultsMode")))
-				.IconImage(FEditorStyle::GetBrush("FullBlueprintEditor.SwitchToBlueprintDefaultsMode"))
-				.SmallIconImage(FEditorStyle::GetBrush("FullBlueprintEditor.SwitchToBlueprintDefaultsMode.Small"))
-			]
+			SNew(SModeWidget, FWidgetBlueprintApplicationModes::GetLocalizedMode(FWidgetBlueprintApplicationModes::DesignerMode), FWidgetBlueprintApplicationModes::DesignerMode)
+			.OnGetActiveMode(GetActiveMode)
+			.OnSetActiveMode(SetActiveMode)
+			.ToolTip(IDocumentation::Get()->CreateToolTip(
+				LOCTEXT("BlueprintDefaultsModeButtonTooltip", "Switch to Blueprint Defaults Mode"),
+				NULL,
+				TEXT("Shared/Editors/BlueprintEditor"),
+				TEXT("DefaultsMode")))
+			.IconImage(FEditorStyle::GetBrush("FullBlueprintEditor.SwitchToBlueprintDefaultsMode"))
+			.SmallIconImage(FEditorStyle::GetBrush("FullBlueprintEditor.SwitchToBlueprintDefaultsMode.Small"))
+			.Tag(TEXT("DefaultsMode"))
 		);
 
 		BlueprintEditorPtr->AddToolbarWidget(SNew(SBlueprintModeSeparator));
 
 		BlueprintEditorPtr->AddToolbarWidget(
-			SNew( STutorialWrapper, TEXT("GraphMode") )
-			[
-				SNew(SModeWidget, FWidgetBlueprintApplicationModes::GetLocalizedMode(FWidgetBlueprintApplicationModes::GraphMode), FWidgetBlueprintApplicationModes::GraphMode)
-				.OnGetActiveMode(GetActiveMode)
-				.OnSetActiveMode(SetActiveMode)
-				.CanBeSelected(BlueprintEditorPtr.Get(), &FBlueprintEditor::IsEditingSingleBlueprint)
-				.ToolTip(IDocumentation::Get()->CreateToolTip(
-					LOCTEXT("GraphModeButtonTooltip", "Switch to Graph Editing Mode"),
-					NULL,
-					TEXT("Shared/Editors/BlueprintEditor"),
-					TEXT("GraphMode")))
-				.ToolTipText(LOCTEXT("GraphModeButtonTooltip", "Switch to Graph Editing Mode"))
-				.IconImage(FEditorStyle::GetBrush("FullBlueprintEditor.SwitchToScriptingMode"))
-				.SmallIconImage(FEditorStyle::GetBrush("FullBlueprintEditor.SwitchToScriptingMode.Small"))
-			]
+			SNew(SModeWidget, FWidgetBlueprintApplicationModes::GetLocalizedMode(FWidgetBlueprintApplicationModes::GraphMode), FWidgetBlueprintApplicationModes::GraphMode)
+			.OnGetActiveMode(GetActiveMode)
+			.OnSetActiveMode(SetActiveMode)
+			.CanBeSelected(BlueprintEditorPtr.Get(), &FBlueprintEditor::IsEditingSingleBlueprint)
+			.ToolTip(IDocumentation::Get()->CreateToolTip(
+				LOCTEXT("GraphModeButtonTooltip", "Switch to Graph Editing Mode"),
+				NULL,
+				TEXT("Shared/Editors/BlueprintEditor"),
+				TEXT("GraphMode")))
+			.ToolTipText(LOCTEXT("GraphModeButtonTooltip", "Switch to Graph Editing Mode"))
+			.IconImage(FEditorStyle::GetBrush("FullBlueprintEditor.SwitchToScriptingMode"))
+			.SmallIconImage(FEditorStyle::GetBrush("FullBlueprintEditor.SwitchToScriptingMode.Small"))
+			.Tag(TEXT("GraphMode"))
 		);
 		
 		// Right side padding
diff --git a/Engine/Source/Editor/UnrealEd/Private/GraphEditor.cpp b/Engine/Source/Editor/UnrealEd/Private/GraphEditor.cpp
--- a/Engine/Source/Editor/UnrealEd/Private/GraphEditor.cpp
+++ b/Engine/Source/Editor/UnrealEd/Private/GraphEditor.cpp
@@ -5,7 +5,6 @@
 #include "GraphEditor.h"
 #include "Editor/GraphEditor/Public/GraphEditorModule.h"
 #include "BlueprintUtilities.h"
-#include "STutorialWrapper.h"
 
 // List of all active GraphEditor wrappers
 TArray< TWeakPtr<SGraphEditor> > SGraphEditor::AllInstances;
@@ -33,7 +32,8 @@ void SGraphEditor::ConstructImplementation( const FArguments& InArgs )
 
 	this->ChildSlot
 	[
-		SNew( STutorialWrapper, TEXT("GraphEditorPanel") )
+		SNew( SBox )
+		.Tag(TEXT("GraphEditorPanel"))
 		[
 			Implementation.ToSharedRef()
 		]
diff --git a/Engine/Source/Editor/UnrealEd/Private/SEditorViewport.cpp b/Engine/Source/Editor/UnrealEd/Private/SEditorViewport.cpp
--- a/Engine/Source/Editor/UnrealEd/Private/SEditorViewport.cpp
+++ b/Engine/Source/Editor/UnrealEd/Private/SEditorViewport.cpp
@@ -33,22 +33,20 @@ void SEditorViewport::Construct( const FArguments& InArgs )
 
 	ChildSlot
 	[
-		SNew( STutorialWrapper, TEXT("EditorViewports") )
+		SAssignNew( ViewportWidget, SViewport )
+		.ShowEffectWhenDisabled( false )
+		.EnableGammaCorrection( false ) // Scene rendering handles this
+		.Tag(TEXT("EditorViewports"))
 		[
-			SAssignNew( ViewportWidget, SViewport )
-			.ShowEffectWhenDisabled( false )
-			.EnableGammaCorrection( false ) // Scene rendering handles this
+			SAssignNew( ViewportOverlay, SOverlay )
+			+SOverlay::Slot()
 			[
-				SAssignNew( ViewportOverlay, SOverlay )
-				+SOverlay::Slot()
-				[
-					SNew( SBorder )
-					.BorderImage( this, &SEditorViewport::OnGetViewportBorderBrush )
-					.BorderBackgroundColor( this, &SEditorViewport::OnGetViewportBorderColorAndOpacity )
-					.Visibility( this, &SEditorViewport::OnGetViewportContentVisibility )
-					.Padding(0.0f)
-					.ShowEffectWhenDisabled( false )
-				]
+				SNew( SBorder )
+				.BorderImage( this, &SEditorViewport::OnGetViewportBorderBrush )
+				.BorderBackgroundColor( this, &SEditorViewport::OnGetViewportBorderColorAndOpacity )
+				.Visibility( this, &SEditorViewport::OnGetViewportContentVisibility )
+				.Padding(0.0f)
+				.ShowEffectWhenDisabled( false )
 			]
 		]
 	];
diff --git a/Engine/Source/Runtime/Slate/Private/Framework/MultiBox/MultiBoxBuilder.cpp b/Engine/Source/Runtime/Slate/Private/Framework/MultiBox/MultiBoxBuilder.cpp
--- a/Engine/Source/Runtime/Slate/Private/Framework/MultiBox/MultiBoxBuilder.cpp
+++ b/Engine/Source/Runtime/Slate/Private/Framework/MultiBox/MultiBoxBuilder.cpp
@@ -12,11 +12,11 @@
 #include "SButtonRowBlock.h"
 #include "SWidgetBlock.h"
 #include "SGroupMarkerBlock.h"
-#include "STutorialWrapper.h"
 
-FMultiBoxBuilder::FMultiBoxBuilder( const EMultiBoxType::Type InType, FMultiBoxCustomization InCustomization, const bool bInShouldCloseWindowAfterMenuSelection, const TSharedPtr< const FUICommandList >& InCommandList, TSharedPtr<FExtender> InExtender )
+FMultiBoxBuilder::FMultiBoxBuilder( const EMultiBoxType::Type InType, FMultiBoxCustomization InCustomization, const bool bInShouldCloseWindowAfterMenuSelection, const TSharedPtr< const FUICommandList >& InCommandList, TSharedPtr<FExtender> InExtender, FName InTutorialHighlightName )
 	: MultiBox( FMultiBox::Create( InType, InCustomization, bInShouldCloseWindowAfterMenuSelection ) )
 	, CommandListStack()
+	, TutorialHighlightName(InTutorialHighlightName)
 {
 	CommandListStack.Push( InCommandList );
 	ExtenderStack.Push(InExtender);
@@ -90,14 +90,39 @@ TSharedRef< class FMultiBox > FMultiBoxBuilder::GetMultiBox()
 	return MultiBox;
 }
 
-FBaseMenuBuilder::FBaseMenuBuilder( const EMultiBoxType::Type InType, const bool bInShouldCloseWindowAfterMenuSelection, TSharedPtr< const FUICommandList > InCommandList, bool bInCloseSelfOnly, TSharedPtr<FExtender> InExtender, const ISlateStyle* InStyleSet )
-	: FMultiBoxBuilder( InType, FMultiBoxCustomization::None, bInShouldCloseWindowAfterMenuSelection, InCommandList, InExtender )
+/** Helper function to generate unique widget-identifying names given various bits of information */
+static FName GenerateTutorialIdentfierName(FName InContainerName, FName InElementName, const TSharedPtr< const FUICommandInfo > InCommand, int32 InIndex)
+{
+	FString BaseName;
+	if(InContainerName != NAME_None)
+	{
+		BaseName = InContainerName.ToString() + TEXT(".");
+	}
+
+	if(InElementName != NAME_None)
+	{
+		return FName(*(BaseName + InElementName.ToString()));
+	}
+	else if(InCommand.IsValid() && InCommand->GetCommandName() != NAME_None)
+	{
+		return FName(*(BaseName + InCommand->GetCommandName().ToString()));
+	}
+	else
+	{
+		// default to index if no other info is available
+		const FString IndexedName = FString::Printf(TEXT("MultiboxWidget%d"), InIndex);
+		return FName(*(BaseName + IndexedName));
+	}
+}
+
+FBaseMenuBuilder::FBaseMenuBuilder( const EMultiBoxType::Type InType, const bool bInShouldCloseWindowAfterMenuSelection, TSharedPtr< const FUICommandList > InCommandList, bool bInCloseSelfOnly, TSharedPtr<FExtender> InExtender, const ISlateStyle* InStyleSet, FName InTutorialHighlightName )
+	: FMultiBoxBuilder( InType, FMultiBoxCustomization::None, bInShouldCloseWindowAfterMenuSelection, InCommandList, InExtender, InTutorialHighlightName )
 	, bCloseSelfOnly( bInCloseSelfOnly )
 {
 	MultiBox->SetStyle(InStyleSet, "Menu");
 }
 
-void FBaseMenuBuilder::AddMenuEntry( const TSharedPtr< const FUICommandInfo > InCommand, FName InExtensionHook, const TAttribute<FText>& InLabelOverride, const TAttribute<FText>& InToolTipOverride, const FSlateIcon& InIconOverride )
+void FBaseMenuBuilder::AddMenuEntry( const TSharedPtr< const FUICommandInfo > InCommand, FName InExtensionHook, const TAttribute<FText>& InLabelOverride, const TAttribute<FText>& InToolTipOverride, const FSlateIcon& InIconOverride, FName InTutorialHighlightName )
 {
 	ApplySectionBeginning();
 
@@ -106,30 +131,33 @@ void FBaseMenuBuilder::AddMenuEntry( const TSharedPtr< const FUICommandInfo > In
 	// The command must be valid
 	check( InCommand.IsValid() );
 	TSharedRef< FMenuEntryBlock > NewMenuEntryBlock( new FMenuEntryBlock( InExtensionHook, InCommand, CommandListStack.Last(), InLabelOverride, InToolTipOverride, InIconOverride, bCloseSelfOnly ) );
+	NewMenuEntryBlock->SetTutorialHighlightName(GenerateTutorialIdentfierName(TutorialHighlightName, InTutorialHighlightName, InCommand, MultiBox->GetBlocks().Num()));
 	MultiBox->AddMultiBlock( NewMenuEntryBlock );
-	
+
 	ApplyHook(InExtensionHook, EExtensionHook::After);
 }
 
-void FBaseMenuBuilder::AddMenuEntry( const TAttribute<FText>& InLabel, const TAttribute<FText>& InToolTip, const FSlateIcon& InIcon, const FUIAction& InAction, FName InExtensionHook, const EUserInterfaceActionType::Type UserInterfaceActionType )
+void FBaseMenuBuilder::AddMenuEntry( const TAttribute<FText>& InLabel, const TAttribute<FText>& InToolTip, const FSlateIcon& InIcon, const FUIAction& InAction, FName InExtensionHook, const EUserInterfaceActionType::Type UserInterfaceActionType, FName InTutorialHighlightName )
 {
 	ApplySectionBeginning();
 
 	ApplyHook(InExtensionHook, EExtensionHook::Before);
 	
 	TSharedRef< FMenuEntryBlock > NewMenuEntryBlock( new FMenuEntryBlock( InExtensionHook, InLabel, InToolTip, InIcon, InAction, UserInterfaceActionType, bCloseSelfOnly ) );
+	NewMenuEntryBlock->SetTutorialHighlightName(GenerateTutorialIdentfierName(TutorialHighlightName, InTutorialHighlightName, nullptr, MultiBox->GetBlocks().Num()));
 	MultiBox->AddMultiBlock( NewMenuEntryBlock );
 	
 	ApplyHook(InExtensionHook, EExtensionHook::After);
 }
 
-void FBaseMenuBuilder::AddMenuEntry( const FUIAction& UIAction, const TSharedRef< SWidget > Contents, const FName& InExtensionHook, const TAttribute<FText>& InToolTip, const EUserInterfaceActionType::Type UserInterfaceActionType )
+void FBaseMenuBuilder::AddMenuEntry( const FUIAction& UIAction, const TSharedRef< SWidget > Contents, const FName& InExtensionHook, const TAttribute<FText>& InToolTip, const EUserInterfaceActionType::Type UserInterfaceActionType, FName InTutorialHighlightName )
 {
 	ApplySectionBeginning();
 
 	ApplyHook(InExtensionHook, EExtensionHook::Before);
 
 	TSharedRef< FMenuEntryBlock > NewMenuEntryBlock( new FMenuEntryBlock( InExtensionHook, UIAction, Contents, InToolTip, UserInterfaceActionType, bCloseSelfOnly ) );
+	NewMenuEntryBlock->SetTutorialHighlightName(GenerateTutorialIdentfierName(TutorialHighlightName, InTutorialHighlightName, nullptr, MultiBox->GetBlocks().Num()));
 	MultiBox->AddMultiBlock( NewMenuEntryBlock );
 
 	ApplyHook(InExtensionHook, EExtensionHook::After);
@@ -290,7 +318,8 @@ void FMenuBarBuilder::AddPullDownMenu( const FText& InMenuLabel, const FText& In
 	// Pulldown menus always close all menus not just themselves
 	const bool bCloseSelfOnly = false;
 	TSharedRef< FMenuEntryBlock > NewMenuEntryBlock( new FMenuEntryBlock( InExtensionHook, InMenuLabel, InToolTip, InPullDownMenu, ExtenderStack.Top(), bIsSubMenu, bOpenSubMenuOnClick, CommandListStack.Last(), bCloseSelfOnly ) );
-	NewMenuEntryBlock->SetTutorialHightlightName(InTutorialHighlightName);
+	NewMenuEntryBlock->SetTutorialHighlightName(GenerateTutorialIdentfierName(TutorialHighlightName, InTutorialHighlightName, nullptr, MultiBox->GetBlocks().Num()));
+
 	MultiBox->AddMultiBlock( NewMenuEntryBlock );
 	
 	ApplyHook(InExtensionHook, EExtensionHook::After);
@@ -306,7 +335,7 @@ void FMenuBarBuilder::ApplyHook(FName InExtensionHook, EExtensionHook::Position
 	}
 }
 
-void FToolBarBuilder::AddToolBarButton( const TSharedPtr< const FUICommandInfo > InCommand, FName InExtensionHook, const TAttribute<FText>& InLabelOverride, const TAttribute<FText>& InToolTipOverride, const TAttribute<FSlateIcon>& InIconOverride, FName TutorialHighlightName)
+void FToolBarBuilder::AddToolBarButton( const TSharedPtr< const FUICommandInfo > InCommand, FName InExtensionHook, const TAttribute<FText>& InLabelOverride, const TAttribute<FText>& InToolTipOverride, const TAttribute<FSlateIcon>& InIconOverride, FName InTutorialHighlightName)
 {
 	ApplySectionBeginning();
 
@@ -321,14 +350,14 @@ void FToolBarBuilder::AddToolBarButton( const TSharedPtr< const FUICommandInfo >
 
 	NewToolBarButtonBlock->SetIsFocusable(bIsFocusable);
 	NewToolBarButtonBlock->SetForceSmallIcons(bForceSmallIcons);
-	NewToolBarButtonBlock->SetTutorialHightlightName(TutorialHighlightName);
+	NewToolBarButtonBlock->SetTutorialHighlightName(GenerateTutorialIdentfierName(TutorialHighlightName, InTutorialHighlightName, InCommand, MultiBox->GetBlocks().Num()));
 
 	MultiBox->AddMultiBlock( NewToolBarButtonBlock );
 
 	ApplyHook(InExtensionHook, EExtensionHook::After);
 }
 
-void FToolBarBuilder::AddToolBarButton( const FUIAction& InAction, FName InExtensionHook, const TAttribute<FText>& InLabelOverride, const TAttribute<FText>& InToolTipOverride, const TAttribute<FSlateIcon>& InIconOverride, const EUserInterfaceActionType::Type UserInterfaceActionType, FName TutorialHighlightName )
+void FToolBarBuilder::AddToolBarButton( const FUIAction& InAction, FName InExtensionHook, const TAttribute<FText>& InLabelOverride, const TAttribute<FText>& InToolTipOverride, const TAttribute<FSlateIcon>& InIconOverride, const EUserInterfaceActionType::Type UserInterfaceActionType, FName InTutorialHighlightName )
 {
 	ApplySectionBeginning();
 
@@ -343,14 +372,14 @@ void FToolBarBuilder::AddToolBarButton( const FUIAction& InAction, FName InExten
 
 	NewToolBarButtonBlock->SetIsFocusable(bIsFocusable);
 	NewToolBarButtonBlock->SetForceSmallIcons(bForceSmallIcons);
-	NewToolBarButtonBlock->SetTutorialHightlightName(TutorialHighlightName);
+	NewToolBarButtonBlock->SetTutorialHighlightName(GenerateTutorialIdentfierName(TutorialHighlightName, InTutorialHighlightName, nullptr, MultiBox->GetBlocks().Num()));
 
 	MultiBox->AddMultiBlock( NewToolBarButtonBlock );
 
 	ApplyHook(InExtensionHook, EExtensionHook::After);
 }
 
-void FToolBarBuilder::AddComboButton( const FUIAction& InAction, const FOnGetContent& InMenuContentGenerator, const TAttribute<FText>& InLabelOverride, const TAttribute<FText>& InToolTipOverride, const TAttribute<FSlateIcon>& InIconOverride, bool bInSimpleComboBox )
+void FToolBarBuilder::AddComboButton( const FUIAction& InAction, const FOnGetContent& InMenuContentGenerator, const TAttribute<FText>& InLabelOverride, const TAttribute<FText>& InToolTipOverride, const TAttribute<FSlateIcon>& InIconOverride, bool bInSimpleComboBox, FName InTutorialHighlightName )
 {
 	ApplySectionBeginning();
 
@@ -362,6 +391,7 @@ void FToolBarBuilder::AddComboButton( const FUIAction& InAction, const FOnGetCon
 	}
 
 	NewToolBarComboButtonBlock->SetForceSmallIcons(bForceSmallIcons);
+	NewToolBarComboButtonBlock->SetTutorialHighlightName(GenerateTutorialIdentfierName(TutorialHighlightName, InTutorialHighlightName, nullptr, MultiBox->GetBlocks().Num()));
 
 	MultiBox->AddMultiBlock( NewToolBarComboButtonBlock );
 }
@@ -371,16 +401,15 @@ void FToolBarBuilder::AddWidget( TSharedRef<SWidget> InWidget, FName TutorialHig
 	ApplySectionBeginning();
 
 	// If tutorial name specified, wrap in tutorial wrapper
-	if (TutorialHighlightName != NAME_None)
-	{
-		TSharedRef<SWidget> ChildWidget = InWidget;
-
-		InWidget = 
-			SNew( STutorialWrapper, TutorialHighlightName )
-			[
-				ChildWidget
-			];
-	}
+	const FName WrapperName = GenerateTutorialIdentfierName(TutorialHighlightName, NAME_None, nullptr, MultiBox->GetBlocks().Num());
+
+	TSharedRef<SWidget> ChildWidget = InWidget;
+	InWidget = 
+		SNew( SBox )
+		.Tag(TutorialHighlightName)
+		[
+			ChildWidget
+		];
 	
 	TSharedRef< FWidgetBlock > NewWidgetBlock( new FWidgetBlock( InWidget, FText::GetEmpty(), true ) );
 	MultiBox->AddMultiBlock( NewWidgetBlock );
diff --git a/Engine/Source/Runtime/Slate/Private/Framework/MultiBox/SMenuEntryBlock.cpp b/Engine/Source/Runtime/Slate/Private/Framework/MultiBox/SMenuEntryBlock.cpp
--- a/Engine/Source/Runtime/Slate/Private/Framework/MultiBox/SMenuEntryBlock.cpp
+++ b/Engine/Source/Runtime/Slate/Private/Framework/MultiBox/SMenuEntryBlock.cpp
@@ -756,14 +756,15 @@ void SMenuEntryBlock::BuildMultiBlockWidget(const ISlateStyle* StyleSet, const F
 		ChildSlot[ BuildMenuBarWidget( BuildParams ) ];
 	}
 
-	// Insert named STutorialWrapper if desired
+	// Insert named widget if desired
 	FName TutorialName = MenuEntryBlock->GetTutorialHighlightName();
 	if(TutorialName != NAME_None)
 	{
 		TSharedRef<SWidget> ChildWidget = ChildSlot.GetWidget();
 		ChildSlot
 		[
-			SNew( STutorialWrapper, TutorialName )
+			SNew( SBox )
+			.Tag(TutorialName)
 			[
 				ChildWidget
 			]
diff --git a/Engine/Source/Runtime/Slate/Private/Framework/MultiBox/SToolBarButtonBlock.cpp b/Engine/Source/Runtime/Slate/Private/Framework/MultiBox/SToolBarButtonBlock.cpp
--- a/Engine/Source/Runtime/Slate/Private/Framework/MultiBox/SToolBarButtonBlock.cpp
+++ b/Engine/Source/Runtime/Slate/Private/Framework/MultiBox/SToolBarButtonBlock.cpp
@@ -154,6 +154,8 @@ void SToolBarButtonBlock::BuildMultiBlockWidget(const ISlateStyle* StyleSet, con
 	TSharedRef< SWidget > ButtonContent =
 
 		SNew(SHorizontalBox)
+		.Tag(TutorialHighlightName)
+
 		+ SHorizontalBox::Slot()
 		.FillWidth(1)
 		.VAlign(VAlign_Center)
@@ -185,18 +187,6 @@ void SToolBarButtonBlock::BuildMultiBlockWidget(const ISlateStyle* StyleSet, con
 			]
 		];
 
-	// Add a tutorial highlight wrapper if desired
-	if(TutorialHighlightName != NAME_None)
-	{
-		TSharedRef<SWidget> Wrapper = 
-			SNew(STutorialWrapper, TutorialHighlightName)
-			[
-				ButtonContent
-			];
-
-		ButtonContent = Wrapper;
-	}
-
 	EMultiBlockLocation::Type BlockLocation = GetMultiBlockLocation();
 	
 	// What type of UI should we create for this block?
diff --git a/Engine/Source/Runtime/Slate/Private/Widgets/Tutorials/STutorialWrapper.cpp b/Engine/Source/Runtime/Slate/Private/Widgets/Tutorials/STutorialWrapper.cpp
deleted file mode 100644
diff --git a/Engine/Source/Runtime/Slate/Public/Framework/MultiBox/MultiBox.h b/Engine/Source/Runtime/Slate/Public/Framework/MultiBox/MultiBox.h
--- a/Engine/Source/Runtime/Slate/Public/Framework/MultiBox/MultiBox.h
+++ b/Engine/Source/Runtime/Slate/Public/Framework/MultiBox/MultiBox.h
@@ -99,7 +99,7 @@ public:
 	virtual bool IsGroupEndBlock()		const	{ return false; };
 
 	/** Set the tutorial highlight name for this menu entry */
-	void SetTutorialHightlightName(FName InTutorialName)
+	void SetTutorialHighlightName(FName InTutorialName)
 	{
 		TutorialHighlightName = InTutorialName;
 	}
diff --git a/Engine/Source/Runtime/Slate/Public/Framework/MultiBox/MultiBoxBuilder.h b/Engine/Source/Runtime/Slate/Public/Framework/MultiBox/MultiBoxBuilder.h
--- a/Engine/Source/Runtime/Slate/Public/Framework/MultiBox/MultiBoxBuilder.h
+++ b/Engine/Source/Runtime/Slate/Public/Framework/MultiBox/MultiBoxBuilder.h
@@ -27,8 +27,9 @@ public:
 	 * @param	InType	Type of MultiBox
 	 * @param	bInShouldCloseWindowAfterMenuSelection	Sets whether or not the window that contains this multibox should be destroyed after the user clicks on a menu item in this box
 	 * @param	InCommandList	The action list that maps command infos to delegates that should be called for each command associated with a multiblock widget.  This can be modified after the MultiBox is created by calling the PushCommandList() and PopCommandList() methods.
+	 * @param	InTutorialHighlightName	Optional name to identify this widget and highlight during tutorials
 	 */
-	FMultiBoxBuilder( const EMultiBoxType::Type InType, FMultiBoxCustomization InCustomization, const bool bInShouldCloseWindowAfterMenuSelection, const TSharedPtr< const FUICommandList >& InCommandList, TSharedPtr<FExtender> InExtender = TSharedPtr<FExtender>() );
+	FMultiBoxBuilder( const EMultiBoxType::Type InType, FMultiBoxCustomization InCustomization, const bool bInShouldCloseWindowAfterMenuSelection, const TSharedPtr< const FUICommandList >& InCommandList, TSharedPtr<FExtender> InExtender = TSharedPtr<FExtender>(), FName InTutorialHighlightName = NAME_None );
 
 
 	/**
@@ -122,6 +123,9 @@ protected:
 
 	/** The extender stack holding all the possible extensions for this menu builder */
 	TArray< TSharedPtr<class FExtender> > ExtenderStack;
+
+	/** Name to identify this widget and highlight during tutorials */
+	FName TutorialHighlightName;
 };
 
 
@@ -141,8 +145,9 @@ public:
 	 * @param	bInShouldCloseWindowAfterMenuSelection	Sets whether or not the window that contains this multibox should be destroyed after the user clicks on a menu item in this box
 	 * @param	InCommandList	The action list that maps command infos to delegates that should be called for each command associated with a multiblock widget
 	 * @param	bInCloseSelfOnly	True if clicking on a menu entry closes itself only and its children but not the entire stack 
+	 * @param	InTutorialHighlightName	Optional name to identify this widget and highlight during tutorials
 	 */
-	FBaseMenuBuilder( const EMultiBoxType::Type InType, const bool bInShouldCloseWindowAfterMenuSelection, TSharedPtr< const FUICommandList > InCommandList, bool bInCloseSelfOnly, TSharedPtr<FExtender> InExtender = TSharedPtr<FExtender>(), const ISlateStyle* InStyleSet = &FCoreStyle::Get() );
+	FBaseMenuBuilder( const EMultiBoxType::Type InType, const bool bInShouldCloseWindowAfterMenuSelection, TSharedPtr< const FUICommandList > InCommandList, bool bInCloseSelfOnly, TSharedPtr<FExtender> InExtender = TSharedPtr<FExtender>(), const ISlateStyle* InStyleSet = &FCoreStyle::Get(), FName InTutorialHighlightName = NAME_None );
 
 	/**
 	 * Adds a menu entry
@@ -151,8 +156,9 @@ public:
 	 * @param	InLabelOverride		Optional label override.  If omitted, then the action's label will be used instead.
 	 * @param	InToolTipOverride	Optional tool tip override.	 If omitted, then the action's label will be used instead.
 	 * @param	InIconOverride		Optional name of the slate brush to use for the tool bar image.  If omitted, then the action's icon will be used instead.
+	 * @param	InTutorialHighlightName	Optional name to identify this widget and highlight during tutorials
 	 */
-	void AddMenuEntry( const TSharedPtr< const FUICommandInfo > InCommand, FName InExtensionHook = NAME_None, const TAttribute<FText>& InLabelOverride = TAttribute<FText>(), const TAttribute<FText>& InToolTipOverride = TAttribute<FText>(), const FSlateIcon& InIconOverride = FSlateIcon() );
+	void AddMenuEntry( const TSharedPtr< const FUICommandInfo > InCommand, FName InExtensionHook = NAME_None, const TAttribute<FText>& InLabelOverride = TAttribute<FText>(), const TAttribute<FText>& InToolTipOverride = TAttribute<FText>(), const FSlateIcon& InIconOverride = FSlateIcon(), FName InTutorialHighlightName = NAME_None );
 
 	/**
 	 * Adds a menu entry without the use of a command
@@ -162,10 +168,11 @@ public:
 	 * @param	InIcon		The icon to use		
 	 * @param	UIAction	Actions to execute on this menu item.
 	 * @param	UserInterfaceActionType	Type of interface action
+	 * @param	InTutorialHighlightName	Optional name to identify this widget and highlight during tutorials
 	 */
-	void AddMenuEntry( const TAttribute<FText>& InLabel, const TAttribute<FText>& InToolTip, const FSlateIcon& InIcon, const FUIAction& UIAction, FName InExtensionHook = NAME_None, const EUserInterfaceActionType::Type UserInterfaceActionType = EUserInterfaceActionType::Button );
+	void AddMenuEntry( const TAttribute<FText>& InLabel, const TAttribute<FText>& InToolTip, const FSlateIcon& InIcon, const FUIAction& UIAction, FName InExtensionHook = NAME_None, const EUserInterfaceActionType::Type UserInterfaceActionType = EUserInterfaceActionType::Button, FName InTutorialHighlightName = NAME_None );
 
-	void AddMenuEntry( const FUIAction& UIAction, const TSharedRef< SWidget > Contents, const FName& InExtensionHook = NAME_None, const TAttribute<FText>& InToolTip = TAttribute<FText>(), const EUserInterfaceActionType::Type UserInterfaceActionType = EUserInterfaceActionType::Button );
+	void AddMenuEntry( const FUIAction& UIAction, const TSharedRef< SWidget > Contents, const FName& InExtensionHook = NAME_None, const TAttribute<FText>& InToolTip = TAttribute<FText>(), const EUserInterfaceActionType::Type UserInterfaceActionType = EUserInterfaceActionType::Button, FName InTutorialHighlightName = NAME_None );
 
 protected:
 	/** True if clicking on a menu entry closes itself only and its children and not the entire stack */
@@ -347,8 +354,9 @@ public:
 	 * @param	InLabelOverride			Optional label override.  If omitted, then the action's label will be used instead.
 	 * @param	InToolTipOverride		Optional tool tip override.	 If omitted, then the action's label will be used instead.
 	 * @param	InIconOverride			Optional name of the slate brush to use for the tool bar image.  If omitted, then the action's icon will be used instead.
+	 * @param	InTutorialHighlightName	Name to identify this widget and highlight during tutorials
 	 */
-	void AddToolBarButton( const TSharedPtr< const FUICommandInfo > InCommand, FName InExtensionHook = NAME_None, const TAttribute<FText>& InLabelOverride = TAttribute<FText>(), const TAttribute<FText>& InToolTipOverride = TAttribute<FText>(), const TAttribute<FSlateIcon>& InIconOverride = TAttribute<FSlateIcon>(), FName TutorialHighlightName = NAME_None );
+	void AddToolBarButton( const TSharedPtr< const FUICommandInfo > InCommand, FName InExtensionHook = NAME_None, const TAttribute<FText>& InLabelOverride = TAttribute<FText>(), const TAttribute<FText>& InToolTipOverride = TAttribute<FText>(), const TAttribute<FSlateIcon>& InIconOverride = TAttribute<FSlateIcon>(), FName InTutorialHighlightName = NAME_None );
 	
 	/**
 	 * Adds a tool bar button
@@ -358,8 +366,9 @@ public:
 	 * @param	InToolTip	Tool tip used when hovering over the menu entry
 	 * @param	InIcon		The icon to use		
 	 * @param	UserInterfaceActionType	Type of interface action
+	 * @param	InTutorialHighlightName	Name to identify this widget and highlight during tutorials
 	 */
-	void AddToolBarButton( const FUIAction& InAction, FName InExtensionHook = NAME_None, const TAttribute<FText>& InLabelOverride = TAttribute<FText>(), const TAttribute<FText>& InToolTipOverride = TAttribute<FText>(), const TAttribute<FSlateIcon>& InIconOverride = TAttribute<FSlateIcon>(), const EUserInterfaceActionType::Type UserInterfaceActionType = EUserInterfaceActionType::Button, FName TutorialHighlightName = NAME_None );
+	void AddToolBarButton( const FUIAction& InAction, FName InExtensionHook = NAME_None, const TAttribute<FText>& InLabelOverride = TAttribute<FText>(), const TAttribute<FText>& InToolTipOverride = TAttribute<FText>(), const TAttribute<FSlateIcon>& InIconOverride = TAttribute<FSlateIcon>(), const EUserInterfaceActionType::Type UserInterfaceActionType = EUserInterfaceActionType::Button, FName InTutorialHighlightName = NAME_None );
 
 	/**
 	 * Adds a combo button
@@ -370,16 +379,17 @@ public:
 	 * @param	InToolTipOverride			Optional tool tip override.	 If omitted, then the action's label will be used instead.
 	 * @param	InIconOverride				Optional icon to use for the tool bar image.  If omitted, then the action's icon will be used instead.
 	 * @param	bInSimpleComboBox			If true, the icon and label won't be displayed
+	 * @param	InTutorialHighlightName		Name to identify this widget and highlight during tutorials
 	 */
-	void AddComboButton( const FUIAction& InAction, const FOnGetContent& InMenuContentGenerator, const TAttribute<FText>& InLabelOverride = TAttribute<FText>(), const TAttribute<FText>& InToolTipOverride = TAttribute<FText>(), const TAttribute<FSlateIcon>& InIconOverride = TAttribute<FSlateIcon>(), bool bInSimpleComboBox = false );
+	void AddComboButton( const FUIAction& InAction, const FOnGetContent& InMenuContentGenerator, const TAttribute<FText>& InLabelOverride = TAttribute<FText>(), const TAttribute<FText>& InToolTipOverride = TAttribute<FText>(), const TAttribute<FSlateIcon>& InIconOverride = TAttribute<FSlateIcon>(), bool bInSimpleComboBox = false, FName InTutorialHighlightName = NAME_None );
 
 	/**
 	 * Adds any widget to the toolbar
 	 * 
 	 * @param	InWidget				The widget that should be shown in the toolbar
-	 * @param	TutorialHighlightName	Name to identify this widget and highlight during tutorials
+	 * @param	InTutorialHighlightName	Name to identify this widget and highlight during tutorials
 	 */
-	void AddWidget( TSharedRef<SWidget> InWidget, FName TutorialHighlightName = NAME_None );
+	void AddWidget( TSharedRef<SWidget> InWidget, FName InTutorialHighlightName = NAME_None );
 	
 	/**
 	 * Adds a toolbar separator
diff --git a/Engine/Source/Runtime/Slate/Public/Slate.h b/Engine/Source/Runtime/Slate/Public/Slate.h
--- a/Engine/Source/Runtime/Slate/Public/Slate.h
+++ b/Engine/Source/Runtime/Slate/Public/Slate.h
@@ -156,7 +156,6 @@ DECLARE_CYCLE_STAT_EXTERN( TEXT("Slate Misc Time"), STAT_SlateMiscTime, STATGROU
 #include "SDPIScaler.h"
 #include "SInlineEditableTextBlock.h"
 #include "SVirtualKeyboardEntry.h"
-#include "STutorialWrapper.h"
 #include "ScrollyZoomy.h"
 #include "SSafeZone.h"
 #include "MarqueeRect.h"
diff --git a/Engine/Source/Runtime/Slate/Public/Widgets/Tutorials/STutorialWrapper.h b/Engine/Source/Runtime/Slate/Public/Widgets/Tutorials/STutorialWrapper.h
deleted file mode 100644
diff --git a/Engine/Source/Runtime/SlateCore/Private/Widgets/SWidget.cpp b/Engine/Source/Runtime/SlateCore/Private/Widgets/SWidget.cpp
--- a/Engine/Source/Runtime/SlateCore/Private/Widgets/SWidget.cpp
+++ b/Engine/Source/Runtime/SlateCore/Private/Widgets/SWidget.cpp
@@ -53,6 +53,7 @@ void SWidget::Construct(
 	Cursor = InCursor;
 	EnabledState = InEnabledState;
 	Visibility = InVisibility;
+	Tag = InTag;
 }
 
 
@@ -446,6 +447,12 @@ int32 SWidget::GetCreatedInLineNumber() const
 }
 
 
+FName SWidget::GetTag() const
+{
+	return this->Tag;
+}
+
+
 FSlateColor SWidget::GetForegroundColor() const
 {
 	static FSlateColor NoColor = FSlateColor::UseForeground();
diff --git a/Engine/Source/Runtime/SlateCore/Private/Widgets/SWindow.cpp b/Engine/Source/Runtime/SlateCore/Private/Widgets/SWindow.cpp
--- a/Engine/Source/Runtime/SlateCore/Private/Widgets/SWindow.cpp
+++ b/Engine/Source/Runtime/SlateCore/Private/Widgets/SWindow.cpp
@@ -891,6 +891,11 @@ TSharedRef<const SWidget> SWindow::GetContent() const
 	}
 }
 
+bool SWindow::HasOverlay() const
+{
+	return WindowOverlay.IsValid();
+}
+
 SOverlay::FOverlaySlot& SWindow::AddOverlaySlot( const int32 ZOrder )
 {
 	if(!WindowOverlay.IsValid())
diff --git a/Engine/Source/Runtime/SlateCore/Public/Widgets/SWidget.h b/Engine/Source/Runtime/SlateCore/Public/Widgets/SWidget.h
--- a/Engine/Source/Runtime/SlateCore/Public/Widgets/SWidget.h
+++ b/Engine/Source/Runtime/SlateCore/Public/Widgets/SWidget.h
@@ -610,6 +610,9 @@ public:
 	/** @return The line number of the widgets location */
 	virtual int32 GetCreatedInLineNumber() const;
 
+	/** @return The name this widget was tagged with */
+	virtual FName GetTag() const;
+
 	/** @return the Foreground color that this widget sets; unset options if the widget does not set a foreground color */
 	virtual FSlateColor GetForegroundColor() const;
 
@@ -710,6 +713,9 @@ protected:
 	/** Line number on which this widget was created */
 	int32 CreatedOnLine;
 
+	/** Tag for this widget */
+	FName Tag;
+
 	/** The cursor to show when the mouse is hovering over this widget. */
 	TAttribute< TOptional<EMouseCursor::Type> > Cursor;
 
diff --git a/Engine/Source/Runtime/SlateCore/Public/Widgets/SWindow.h b/Engine/Source/Runtime/SlateCore/Public/Widgets/SWindow.h
--- a/Engine/Source/Runtime/SlateCore/Public/Widgets/SWindow.h
+++ b/Engine/Source/Runtime/SlateCore/Public/Widgets/SWindow.h
@@ -328,6 +328,13 @@ public:
 	 */
 	TSharedRef<const SWidget> GetContent() const;
 
+	/**
+	 * Check whether we have a full window overlay, used to draw content over the entire window.
+	 * 
+	 * @return true if the window has an overlay
+	 */
+	bool HasOverlay() const;
+
 	/**
 	 * Adds content to draw on top of the entire window
 	 *
