Editor - Smaller allowable thumbnail size in the Content Browser.  We now also adjust the font size at the lower end of the thumbnail size to make a little more text fit on the lines.

diff --git a/Engine/Source/Editor/ContentBrowser/Private/AssetViewWidgets.cpp b/Engine/Source/Editor/ContentBrowser/Private/AssetViewWidgets.cpp
--- a/Engine/Source/Editor/ContentBrowser/Private/AssetViewWidgets.cpp
+++ b/Engine/Source/Editor/ContentBrowser/Private/AssetViewWidgets.cpp
@@ -1265,7 +1265,7 @@ void SAssetTileItem::Construct( const FArguments& InArgs )
 			.FillHeight(1.f)
 			[
 				SAssignNew(InlineRenameWidget, SInlineEditableTextBlock)
-					.Font(FEditorStyle::GetFontStyle("ContentBrowser.AssetTileViewNameFont"))
+					.Font( this, &SAssetTileItem::GetThumbnailFont )
 					.Text( GetNameText() )
 					.OnBeginTextEdit(this, &SAssetTileItem::HandleBeginNameChange)
 					.OnTextCommitted(this, &SAssetTileItem::HandleNameCommitted)
@@ -1310,6 +1310,21 @@ FOptionalSize SAssetTileItem::GetSCCImageSize() const
 	return GetThumbnailBoxSize().Get() * 0.2;
 }
 
+FSlateFontInfo SAssetTileItem::GetThumbnailFont() const
+{
+	FOptionalSize ThumbSize = GetThumbnailBoxSize();
+	if ( ThumbSize.IsSet() )
+	{
+		float Size = ThumbSize.Get();
+		if ( Size < 85 )
+		{
+			return FEditorStyle::GetFontStyle("ContentBrowser.AssetTileViewNameFontSmall");
+		}
+	}
+
+	return FEditorStyle::GetFontStyle("ContentBrowser.AssetTileViewNameFont");
+}
+
 float SAssetTileItem::GetNameTextWrapWidth() const
 {
 	return ItemWidth.Get() - ThumbnailPadding * 2.f;
diff --git a/Engine/Source/Editor/ContentBrowser/Private/AssetViewWidgets.h b/Engine/Source/Editor/ContentBrowser/Private/AssetViewWidgets.h
--- a/Engine/Source/Editor/ContentBrowser/Private/AssetViewWidgets.h
+++ b/Engine/Source/Editor/ContentBrowser/Private/AssetViewWidgets.h
@@ -459,6 +459,9 @@ private:
 	/** Returns the size of the thumbnail box widget */
 	FOptionalSize GetThumbnailBoxSize() const;
 
+	/** Returns the font to use for the thumbnail label */
+	FSlateFontInfo GetThumbnailFont() const;
+
 	/** Returns the size of the source control state box widget */
 	FOptionalSize GetSCCImageSize() const;
 
diff --git a/Engine/Source/Editor/ContentBrowser/Private/SAssetView.cpp b/Engine/Source/Editor/ContentBrowser/Private/SAssetView.cpp
--- a/Engine/Source/Editor/ContentBrowser/Private/SAssetView.cpp
+++ b/Engine/Source/Editor/ContentBrowser/Private/SAssetView.cpp
@@ -121,7 +121,7 @@ void SAssetView::Construct( const FArguments& InArgs )
 		ThumbnailScaleSliderValue = FMath::Clamp<float>(ThumbnailScaleSliderValue.Get(), 0.0f, 1.0f);
 	}
 
-	MinThumbnailScale = 0.6f * ThumbnailScaleRangeScalar;
+	MinThumbnailScale = 0.4f * ThumbnailScaleRangeScalar;
 	MaxThumbnailScale = 2.0f * ThumbnailScaleRangeScalar;
 
 	bCanShowClasses = InArgs._CanShowClasses;
diff --git a/Engine/Source/Editor/EditorStyle/Private/SlateEditorStyle.cpp b/Engine/Source/Editor/EditorStyle/Private/SlateEditorStyle.cpp
--- a/Engine/Source/Editor/EditorStyle/Private/SlateEditorStyle.cpp
+++ b/Engine/Source/Editor/EditorStyle/Private/SlateEditorStyle.cpp
@@ -4487,6 +4487,7 @@ void FSlateEditorStyle::FStyle::SetupContentBrowserStyle()
 
 		// Tile view
 		Set( "ContentBrowser.AssetTileViewNameFont", TTF_CORE_FONT( "Fonts/Roboto-Regular", 10 ) );
+		Set( "ContentBrowser.AssetTileViewNameFontSmall", TTF_CORE_FONT( "Fonts/Roboto-Regular", 8 ) );
 		Set( "ContentBrowser.AssetTileViewNameFontDirty", TTF_CORE_FONT( "Fonts/Roboto-Bold", 10 ) );
 		Set( "ContentBrowser.AssetListView.TableRow", FTableRowStyle()
 			.SetEvenRowBackgroundBrush( FSlateNoResource() )
