Change Full documentation -> View Documentation
Add documentation anchor icon next to View Documentation

diff --git a/Engine/Source/Editor/ContentBrowser/Private/AssetContextMenu.cpp b/Engine/Source/Editor/ContentBrowser/Private/AssetContextMenu.cpp
--- a/Engine/Source/Editor/ContentBrowser/Private/AssetContextMenu.cpp
+++ b/Engine/Source/Editor/ContentBrowser/Private/AssetContextMenu.cpp
@@ -375,9 +375,9 @@ bool FAssetContextMenu::AddDocumentationMenuOptions(FMenuBuilder& MenuBuilder)
 			MenuBuilder.BeginSection( "ActorDocumentation", LOCTEXT("ActorDocsHeading", "Documentation") );
 			{
 					MenuBuilder.AddMenuEntry(
-						LOCTEXT("GoToDocsForActor", "Full documentation"),
+						LOCTEXT("GoToDocsForActor", "View Documentation"),
 						LOCTEXT("GoToDocsForActor_ToolTip", "Click to open documentation for this actor"),
-						FSlateIcon(),
+						FSlateIcon(FEditorStyle::GetStyleSetName(), "HelpIcon.Hovered" ),
 						FUIAction( FExecuteAction::CreateSP( this, &FAssetContextMenu::ExecuteGoToDocsForAsset, SelectedClass ) )
 						);
 			}
diff --git a/Engine/Source/Editor/LevelEditor/Private/LevelEditorContextMenu.cpp b/Engine/Source/Editor/LevelEditor/Private/LevelEditorContextMenu.cpp
--- a/Engine/Source/Editor/LevelEditor/Private/LevelEditorContextMenu.cpp
+++ b/Engine/Source/Editor/LevelEditor/Private/LevelEditorContextMenu.cpp
@@ -263,8 +263,9 @@ void FLevelEditorContextMenu::FillMenu( FMenuBuilder& MenuBuilder, TWeakPtr<SLev
 			{
 				MenuBuilder.AddMenuEntry( FLevelEditorCommands::Get().GoToDocsForActor,
 					NAME_None, 
-					LOCTEXT("GoToDocsForActor", "Full documentation"),
-					LOCTEXT("GoToDocsForActor_ToolTip", "Click to open documentation for this actor"));
+					LOCTEXT("GoToDocsForActor", "View Documentation"),
+					LOCTEXT("GoToDocsForActor_ToolTip", "Click to open documentation for this actor"),
+					FSlateIcon(FEditorStyle::GetStyleSetName(), "HelpIcon.Hovered" ));
 			}
 			MenuBuilder.EndSection();
 		}
