Add rich tool tip on asset type to content browser new asset menu

diff --git a/Engine/Source/Editor/ContentBrowser/Private/NewAssetContextMenu.cpp b/Engine/Source/Editor/ContentBrowser/Private/NewAssetContextMenu.cpp
--- a/Engine/Source/Editor/ContentBrowser/Private/NewAssetContextMenu.cpp
+++ b/Engine/Source/Editor/ContentBrowser/Private/NewAssetContextMenu.cpp
@@ -108,7 +108,8 @@ public:
 			]
 		];
 
-		SetToolTipText( Factory->GetToolTip() );
+		
+		SetToolTip(IDocumentation::Get()->CreateToolTip(Factory->GetToolTip(), nullptr, FEditorClassUtils::GetDocumentationPage(Factory->SupportedClass), FEditorClassUtils::GetDocumentationExcerpt(Factory->SupportedClass)));
 	}
 
 private:
