#UE4 Nonunity fix

diff --git a/Engine/Source/Editor/ContentBrowser/Private/NewAssetContextMenu.cpp b/Engine/Source/Editor/ContentBrowser/Private/NewAssetContextMenu.cpp
--- a/Engine/Source/Editor/ContentBrowser/Private/NewAssetContextMenu.cpp
+++ b/Engine/Source/Editor/ContentBrowser/Private/NewAssetContextMenu.cpp
@@ -2,6 +2,7 @@
 
 #include "ContentBrowserPCH.h"
 #include "NewAssetContextMenu.h"
+#include "IDocumentation.h"
 
 #define LOCTEXT_NAMESPACE "ContentBrowser"
 
