Fix non-unity

diff --git a/Engine/Source/Editor/ContentBrowser/Private/NewAssetContextMenu.cpp b/Engine/Source/Editor/ContentBrowser/Private/NewAssetContextMenu.cpp
--- a/Engine/Source/Editor/ContentBrowser/Private/NewAssetContextMenu.cpp
+++ b/Engine/Source/Editor/ContentBrowser/Private/NewAssetContextMenu.cpp
@@ -3,6 +3,7 @@
 #include "ContentBrowserPCH.h"
 #include "NewAssetContextMenu.h"
 #include "IDocumentation.h"
+#include "EditorClassUtils.h"
 
 #define LOCTEXT_NAMESPACE "ContentBrowser"
 
