Blueprint documentation links link to both the class it inherits from and Blueprint system documentation
Documentation links for blueprints are split to go specific graph type (Macros, Interfaces, Function Library)
Factories can specify specific pages or excerpts that are different than the default ones for their Supported Class (in the same way they can override the tooltip text)

diff --git a/Engine/Source/Editor/ContentBrowser/Private/AssetContextMenu.cpp b/Engine/Source/Editor/ContentBrowser/Private/AssetContextMenu.cpp
--- a/Engine/Source/Editor/ContentBrowser/Private/AssetContextMenu.cpp
+++ b/Engine/Source/Editor/ContentBrowser/Private/AssetContextMenu.cpp
@@ -379,11 +379,11 @@ bool FAssetContextMenu::AddDocumentationMenuOptions(FMenuBuilder& MenuBuilder)
 
 				const FString CodeFileName = FPaths::GetCleanFilename( *ClassHeaderPath );
 
-				MenuBuilder.BeginSection( "ActorCode", LOCTEXT("ActorCodeHeading", "C++") );
+				MenuBuilder.BeginSection( "AssetCode", LOCTEXT("AssetCodeHeading", "C++") );
 				{
 					MenuBuilder.AddMenuEntry(
-						FText::Format( LOCTEXT("GoToCodeForActor", "Open {0}"), FText::FromString( CodeFileName ) ),
-						FText::Format( LOCTEXT("GoToCodeForActor_ToolTip", "Opens the header file for this actor ({0}) in a code editing program"), FText::FromString( CodeFileName ) ),
+						FText::Format( LOCTEXT("GoToCodeForAsset", "Open {0}"), FText::FromString( CodeFileName ) ),
+						FText::Format( LOCTEXT("GoToCodeForAsset_ToolTip", "Opens the header file for this asset ({0}) in a code editing program"), FText::FromString( CodeFileName ) ),
 						FSlateIcon(),
 						FUIAction( FExecuteAction::CreateSP( this, &FAssetContextMenu::ExecuteGoToCodeForAsset, SelectedClass ) )
 						);
@@ -397,14 +397,54 @@ bool FAssetContextMenu::AddDocumentationMenuOptions(FMenuBuilder& MenuBuilder)
 		{
 			bAddedOption = true;
 
-			MenuBuilder.BeginSection( "ActorDocumentation", LOCTEXT("ActorDocsHeading", "Documentation") );
+			MenuBuilder.BeginSection( "AssetDocumentation", LOCTEXT("AseetDocsHeading", "Documentation") );
 			{
-					MenuBuilder.AddMenuEntry(
-						LOCTEXT("GoToDocsForActor", "View Documentation"),
-						LOCTEXT("GoToDocsForActor_ToolTip", "Click to open documentation for this actor"),
-						FSlateIcon(FEditorStyle::GetStyleSetName(), "HelpIcon.Hovered" ),
-						FUIAction( FExecuteAction::CreateSP( this, &FAssetContextMenu::ExecuteGoToDocsForAsset, SelectedClass ) )
-						);
+					if (bIsBlueprint)
+					{
+						MenuBuilder.AddMenuEntry(
+							FText::Format( LOCTEXT("GoToDocsForAssetWithClass", "View Documentation - {0}"), SelectedClass->GetDisplayNameText() ),
+							FText::Format( LOCTEXT("GoToDocsForAssetWithClass_ToolTip", "Click to open documentation for {0}"), SelectedClass->GetDisplayNameText() ),
+							FSlateIcon(FEditorStyle::GetStyleSetName(), "HelpIcon.Hovered" ),
+							FUIAction( FExecuteAction::CreateSP( this, &FAssetContextMenu::ExecuteGoToDocsForAsset, SelectedClass ) )
+							);
+
+						FString ExcerptSection;
+
+						UEnum* BlueprintTypeEnum = FindObject<UEnum>(nullptr, TEXT("EBlueprintType"), true);
+						FString* EnumString = SelectedAssets[0].TagsAndValues.Find(GET_MEMBER_NAME_CHECKED(UBlueprint,BlueprintType));
+						EBlueprintType BlueprintType = (EnumString ? (EBlueprintType)BlueprintTypeEnum->FindEnumIndex(**EnumString) : BPTYPE_Normal);
+
+						switch (BlueprintType)
+						{
+						case BPTYPE_FunctionLibrary:
+							ExcerptSection = TEXT("UBlueprint_FunctionLibrary");
+							break;
+						case BPTYPE_Interface:
+							ExcerptSection = TEXT("UBlueprint_Interface");
+							break;
+						case BPTYPE_MacroLibrary:
+							ExcerptSection = TEXT("UBlueprint_Macro");
+							break;
+						default:
+							ExcerptSection = TEXT("UBlueprint");
+						}
+
+						MenuBuilder.AddMenuEntry(
+							LOCTEXT("GoToDocsForBlueprint", "View Documentation - Blueprint"),
+							LOCTEXT("GoToDocsForBlueprint_ToolTip", "Click to open documentation on blueprints"),
+							FSlateIcon(FEditorStyle::GetStyleSetName(), "HelpIcon.Hovered" ),
+							FUIAction( FExecuteAction::CreateSP( this, &FAssetContextMenu::ExecuteGoToDocsForAsset, UBlueprint::StaticClass(), ExcerptSection ) )
+							);
+					}
+					else
+					{
+						MenuBuilder.AddMenuEntry(
+							LOCTEXT("GoToDocsForAsset", "View Documentation"),
+							LOCTEXT("GoToDocsForAsset_ToolTip", "Click to open documentation"),
+							FSlateIcon(FEditorStyle::GetStyleSetName(), "HelpIcon.Hovered" ),
+							FUIAction( FExecuteAction::CreateSP( this, &FAssetContextMenu::ExecuteGoToDocsForAsset, SelectedClass ) )
+							);
+					}
 			}
 			MenuBuilder.EndSection();
 		}
@@ -1019,10 +1059,15 @@ void FAssetContextMenu::ExecuteGoToCodeForAsset(UClass* SelectedClass)
 }
 
 void FAssetContextMenu::ExecuteGoToDocsForAsset(UClass* SelectedClass)
+{
+	ExecuteGoToDocsForAsset(SelectedClass, FString());
+}
+
+void FAssetContextMenu::ExecuteGoToDocsForAsset(UClass* SelectedClass, const FString ExcerptSection)
 {
 	if (SelectedClass)
 	{
-		FString DocumentationLink = FEditorClassUtils::GetDocumentationLink(SelectedClass);
+		FString DocumentationLink = FEditorClassUtils::GetDocumentationLink(SelectedClass, ExcerptSection);
 		if (!DocumentationLink.IsEmpty())
 		{
 			IDocumentation::Get()->Open( DocumentationLink );
diff --git a/Engine/Source/Editor/ContentBrowser/Private/AssetContextMenu.h b/Engine/Source/Editor/ContentBrowser/Private/AssetContextMenu.h
--- a/Engine/Source/Editor/ContentBrowser/Private/AssetContextMenu.h
+++ b/Engine/Source/Editor/ContentBrowser/Private/AssetContextMenu.h
@@ -115,6 +115,9 @@ private:
 	/** Handler for GoToAssetDocs */
 	void ExecuteGoToDocsForAsset(UClass* SelectedClass);
 
+	/** Handler for GoToAssetDocs */
+	void ExecuteGoToDocsForAsset(UClass* SelectedClass, const FString ExcerptSection);
+
 	/** Handler for CopyReference */
 	void ExecuteCopyReference();
 
diff --git a/Engine/Source/Editor/ContentBrowser/Private/NewAssetContextMenu.cpp b/Engine/Source/Editor/ContentBrowser/Private/NewAssetContextMenu.cpp
--- a/Engine/Source/Editor/ContentBrowser/Private/NewAssetContextMenu.cpp
+++ b/Engine/Source/Editor/ContentBrowser/Private/NewAssetContextMenu.cpp
@@ -111,7 +111,7 @@ public:
 		];
 
 		
-		SetToolTip(IDocumentation::Get()->CreateToolTip(Factory->GetToolTip(), nullptr, FEditorClassUtils::GetDocumentationPage(Factory->SupportedClass), FEditorClassUtils::GetDocumentationExcerpt(Factory->SupportedClass)));
+		SetToolTip(IDocumentation::Get()->CreateToolTip(Factory->GetToolTip(), nullptr, Factory->GetToolTipDocumentationPage(), Factory->GetToolTipDocumentationExcerpt()));
 	}
 
 private:
diff --git a/Engine/Source/Editor/UnrealEd/Classes/Factories/BlueprintFunctionLibraryFactory.h b/Engine/Source/Editor/UnrealEd/Classes/Factories/BlueprintFunctionLibraryFactory.h
--- a/Engine/Source/Editor/UnrealEd/Classes/Factories/BlueprintFunctionLibraryFactory.h
+++ b/Engine/Source/Editor/UnrealEd/Classes/Factories/BlueprintFunctionLibraryFactory.h
@@ -17,6 +17,7 @@ class UBlueprintFunctionLibraryFactory : public UBlueprintFactory
 	virtual FName GetNewAssetThumbnailOverride() const override;
 	virtual uint32 GetMenuCategories() const override;
 	virtual FText GetToolTip() const override;
+	virtual FString GetToolTipDocumentationExcerpt() const override;
 	virtual UObject* FactoryCreateNew(UClass* Class, UObject* InParent, FName Name, EObjectFlags Flags, UObject* Context, FFeedbackContext* Warn, FName CallingContext) override;
 	virtual bool ConfigureProperties() override;
 	virtual FString GetDefaultNewAssetName() const override;
diff --git a/Engine/Source/Editor/UnrealEd/Classes/Factories/BlueprintInterfaceFactory.h b/Engine/Source/Editor/UnrealEd/Classes/Factories/BlueprintInterfaceFactory.h
--- a/Engine/Source/Editor/UnrealEd/Classes/Factories/BlueprintInterfaceFactory.h
+++ b/Engine/Source/Editor/UnrealEd/Classes/Factories/BlueprintInterfaceFactory.h
@@ -17,6 +17,7 @@ class UBlueprintInterfaceFactory : public UFactory
 	virtual FName GetNewAssetThumbnailOverride() const override;
 	virtual uint32 GetMenuCategories() const override;
 	virtual FText GetToolTip() const override;
+	virtual FString GetToolTipDocumentationExcerpt() const override;
 	virtual UObject* FactoryCreateNew(UClass* Class, UObject* InParent, FName Name, EObjectFlags Flags, UObject* Context, FFeedbackContext* Warn, FName CallingContext) override;
 	virtual UObject* FactoryCreateNew(UClass* Class, UObject* InParent, FName Name, EObjectFlags Flags, UObject* Context, FFeedbackContext* Warn) override;
 	virtual FString GetDefaultNewAssetName() const override;
diff --git a/Engine/Source/Editor/UnrealEd/Classes/Factories/BlueprintMacroFactory.h b/Engine/Source/Editor/UnrealEd/Classes/Factories/BlueprintMacroFactory.h
--- a/Engine/Source/Editor/UnrealEd/Classes/Factories/BlueprintMacroFactory.h
+++ b/Engine/Source/Editor/UnrealEd/Classes/Factories/BlueprintMacroFactory.h
@@ -17,6 +17,7 @@ class UBlueprintMacroFactory : public UBlueprintFactory
 	virtual FName GetNewAssetThumbnailOverride() const override;
 	virtual uint32 GetMenuCategories() const override;
 	virtual FText GetToolTip() const override;
+	virtual FString GetToolTipDocumentationExcerpt() const override;
 	virtual UObject* FactoryCreateNew(UClass* Class, UObject* InParent, FName Name, EObjectFlags Flags, UObject* Context, FFeedbackContext* Warn, FName CallingContext) override;
 	virtual FString GetDefaultNewAssetName() const override;
 	// End of UFactory interface
diff --git a/Engine/Source/Editor/UnrealEd/Classes/Factories/Factory.h b/Engine/Source/Editor/UnrealEd/Classes/Factories/Factory.h
--- a/Engine/Source/Editor/UnrealEd/Classes/Factories/Factory.h
+++ b/Engine/Source/Editor/UnrealEd/Classes/Factories/Factory.h
@@ -83,6 +83,12 @@ public:
 	/** Returns the tooltip text description of this factory */
 	virtual FText GetToolTip() const;
 
+	/** Returns the documentation page that should be use for the rich tool tip for this factory */
+	virtual FString GetToolTipDocumentationPage() const;
+
+	/** Returns the documentation excerpt that should be use for the rich tool tip for this factory */
+	virtual FString GetToolTipDocumentationExcerpt() const;
+
 	/**
 	 * @return		The object class supported by this factory.
 	 */
diff --git a/Engine/Source/Editor/UnrealEd/Private/EditorClassUtils.cpp b/Engine/Source/Editor/UnrealEd/Private/EditorClassUtils.cpp
--- a/Engine/Source/Editor/UnrealEd/Private/EditorClassUtils.cpp
+++ b/Engine/Source/Editor/UnrealEd/Private/EditorClassUtils.cpp
@@ -26,7 +26,7 @@ TSharedRef<SToolTip> FEditorClassUtils::GetTooltip(const UClass* Class, const TA
 	return (Class ? IDocumentation::Get()->CreateToolTip(OverrideText, nullptr, GetDocumentationPage(Class), GetDocumentationExcerpt(Class)) : SNew(SToolTip));
 }
 
-FString FEditorClassUtils::GetDocumentationLink(const UClass* Class)
+FString FEditorClassUtils::GetDocumentationLink(const UClass* Class, const FString& OverrideExcerpt)
 {
 	FString DocumentationLink;
 	const FString ClassDocsPage = GetDocumentationPage(Class);
@@ -36,8 +36,10 @@ FString FEditorClassUtils::GetDocumentationLink(const UClass* Class)
 	{
 		TSharedRef<IDocumentationPage> ClassDocs = Documentation->GetPage(ClassDocsPage, NULL);
 
+		const FString ExcerptSection = (OverrideExcerpt.IsEmpty() ? GetDocumentationExcerpt(Class) : OverrideExcerpt);
+
 		FExcerpt Excerpt;
-		if (ClassDocs->GetExcerpt(GetDocumentationExcerpt(Class), Excerpt))
+		if (ClassDocs->GetExcerpt(ExcerptSection, Excerpt))
 		{
 			FString* FullDocumentationLink = Excerpt.Variables.Find( TEXT("ToolTipFullLink") );
 			if (FullDocumentationLink)
diff --git a/Engine/Source/Editor/UnrealEd/Private/Factories/EditorFactories.cpp b/Engine/Source/Editor/UnrealEd/Private/Factories/EditorFactories.cpp
--- a/Engine/Source/Editor/UnrealEd/Private/Factories/EditorFactories.cpp
+++ b/Engine/Source/Editor/UnrealEd/Private/Factories/EditorFactories.cpp
@@ -6193,6 +6193,11 @@ FText UBlueprintMacroFactory::GetToolTip() const
 	return LOCTEXT("BlueprintMacroLibraryTooltip", "Blueprint Macro Libraries are containers of macros to be used in other blueprints. They cannot contain variables, inherit from other blueprints, or be placed in levels. Changes to macros in a Blueprint Macro Library will not take effect until client blueprints are recompiled.");
 }
 
+FString UBlueprintMacroFactory::GetToolTipDocumentationExcerpt() const
+{
+	return TEXT("UBlueprint_Macro");
+}
+
 UObject* UBlueprintMacroFactory::FactoryCreateNew(UClass* Class, UObject* InParent, FName Name, EObjectFlags Flags, UObject* Context, FFeedbackContext* Warn, FName CallingContext)
 {
 	// Make sure we are trying to factory a blueprint, then create and init one
@@ -6250,6 +6255,11 @@ FText UBlueprintFunctionLibraryFactory::GetToolTip() const
 	return LOCTEXT("BlueprintFunctionLibraryTooltip", "Blueprint Function Libraries are containers of functions to be used in other blueprints. They cannot contain variables, inherit from other blueprints, or be placed in levels. Changes to functions in a Blueprint Function Library will take effect without recompiling the client blueprints.");
 }
 
+FString UBlueprintFunctionLibraryFactory::GetToolTipDocumentationExcerpt() const
+{
+	return TEXT("UBlueprint_FunctionLibrary");
+}
+
 UObject* UBlueprintFunctionLibraryFactory::FactoryCreateNew(UClass* Class, UObject* InParent, FName Name, EObjectFlags Flags, UObject* Context, FFeedbackContext* Warn, FName CallingContext)
 {
 	// Make sure we are trying to factory a blueprint, then create and init one
@@ -6310,6 +6320,11 @@ FText UBlueprintInterfaceFactory::GetToolTip() const
 	return LOCTEXT("BlueprintInterfaceTooltip", "A Blueprint Interface is a collection of one or more functions - name only, no implementation - that can be added to other Blueprints. These other Blueprints are then expected to implement the functions of the Blueprint Interface in a unique manner.");
 }
 
+FString UBlueprintInterfaceFactory::GetToolTipDocumentationExcerpt() const
+{
+	return TEXT("UBlueprint_Interface");
+}
+
 UObject* UBlueprintInterfaceFactory::FactoryCreateNew(UClass* Class, UObject* InParent, FName Name, EObjectFlags Flags, UObject* Context, FFeedbackContext* Warn, FName CallingContext)
 {
 	// Make sure we are trying to factory a blueprint, then create and init one
diff --git a/Engine/Source/Editor/UnrealEd/Private/Factories/Factory.cpp b/Engine/Source/Editor/UnrealEd/Private/Factories/Factory.cpp
--- a/Engine/Source/Editor/UnrealEd/Private/Factories/Factory.cpp
+++ b/Engine/Source/Editor/UnrealEd/Private/Factories/Factory.cpp
@@ -5,6 +5,7 @@
 #include "UnrealEd.h"
 #include "ObjectTools.h"
 #include "AssetToolsModule.h"
+#include "EditorClassUtils.h"
 
 DEFINE_LOG_CATEGORY_STATIC(LogFactory, Log, All);
 
@@ -121,6 +122,16 @@ FText UFactory::GetToolTip() const
 	return GetSupportedClass()->GetToolTipText();
 }
 
+FString UFactory::GetToolTipDocumentationPage() const
+{
+	return FEditorClassUtils::GetDocumentationPage(GetSupportedClass());
+}
+
+FString UFactory::GetToolTipDocumentationExcerpt() const
+{
+	return FEditorClassUtils::GetDocumentationExcerpt(GetSupportedClass());
+}
+
 UClass* UFactory::GetSupportedClass() const
 {
 	return SupportedClass;
diff --git a/Engine/Source/Editor/UnrealEd/Public/EditorClassUtils.h b/Engine/Source/Editor/UnrealEd/Public/EditorClassUtils.h
--- a/Engine/Source/Editor/UnrealEd/Public/EditorClassUtils.h
+++ b/Engine/Source/Editor/UnrealEd/Public/EditorClassUtils.h
@@ -48,7 +48,7 @@ namespace FEditorClassUtils
 	 * @param	Class		Class we want to build a link for
 	 * @return				The path to the documentation for the class
 	 */
-	UNREALED_API FString GetDocumentationLink(const UClass* Class);
+	UNREALED_API FString GetDocumentationLink(const UClass* Class, const FString& OverrideExcerpt = FString());
 
 	/**
 	 * Creates a link widget to the documentation for a given class
