Display Blueprint documentation links even if there is no documentation for the Actor class
Distinguish label for different types on the Blueprint documentation link

diff --git a/Engine/Source/Editor/ContentBrowser/Private/AssetContextMenu.cpp b/Engine/Source/Editor/ContentBrowser/Private/AssetContextMenu.cpp
--- a/Engine/Source/Editor/ContentBrowser/Private/AssetContextMenu.cpp
+++ b/Engine/Source/Editor/ContentBrowser/Private/AssetContextMenu.cpp
@@ -393,7 +393,7 @@ bool FAssetContextMenu::AddDocumentationMenuOptions(FMenuBuilder& MenuBuilder)
 		}
 
 		const FString DocumentationLink = FEditorClassUtils::GetDocumentationLink(SelectedClass);
-		if (!DocumentationLink.IsEmpty())
+		if (bIsBlueprint || !DocumentationLink.IsEmpty())
 		{
 			bAddedOption = true;
 
@@ -401,40 +401,54 @@ bool FAssetContextMenu::AddDocumentationMenuOptions(FMenuBuilder& MenuBuilder)
 			{
 					if (bIsBlueprint)
 					{
-						MenuBuilder.AddMenuEntry(
-							FText::Format( LOCTEXT("GoToDocsForAssetWithClass", "View Documentation - {0}"), SelectedClass->GetDisplayNameText() ),
-							FText::Format( LOCTEXT("GoToDocsForAssetWithClass_ToolTip", "Click to open documentation for {0}"), SelectedClass->GetDisplayNameText() ),
-							FSlateIcon(FEditorStyle::GetStyleSetName(), "HelpIcon.Hovered" ),
-							FUIAction( FExecuteAction::CreateSP( this, &FAssetContextMenu::ExecuteGoToDocsForAsset, SelectedClass ) )
-							);
-
-						FString ExcerptSection;
+						if (!DocumentationLink.IsEmpty())
+						{
+							MenuBuilder.AddMenuEntry(
+								FText::Format( LOCTEXT("GoToDocsForAssetWithClass", "View Documentation - {0}"), SelectedClass->GetDisplayNameText() ),
+								FText::Format( LOCTEXT("GoToDocsForAssetWithClass_ToolTip", "Click to open documentation for {0}"), SelectedClass->GetDisplayNameText() ),
+								FSlateIcon(FEditorStyle::GetStyleSetName(), "HelpIcon.Hovered" ),
+								FUIAction( FExecuteAction::CreateSP( this, &FAssetContextMenu::ExecuteGoToDocsForAsset, SelectedClass ) )
+								);
+						}
 
-						UEnum* BlueprintTypeEnum = FindObject<UEnum>(nullptr, TEXT("EBlueprintType"), true);
+						UEnum* BlueprintTypeEnum = FindObject<UEnum>(ANY_PACKAGE, TEXT("EBlueprintType"), true);
 						FString* EnumString = SelectedAssets[0].TagsAndValues.Find(GET_MEMBER_NAME_CHECKED(UBlueprint,BlueprintType));
 						EBlueprintType BlueprintType = (EnumString ? (EBlueprintType)BlueprintTypeEnum->FindEnumIndex(**EnumString) : BPTYPE_Normal);
 
 						switch (BlueprintType)
 						{
 						case BPTYPE_FunctionLibrary:
-							ExcerptSection = TEXT("UBlueprint_FunctionLibrary");
+							MenuBuilder.AddMenuEntry(
+								LOCTEXT("GoToDocsForMacroBlueprint", "View Documentation - Function Library"),
+								LOCTEXT("GoToDocsForMacroBlueprint_ToolTip", "Click to open documentation on blueprint function libraries"),
+								FSlateIcon(FEditorStyle::GetStyleSetName(), "HelpIcon.Hovered" ),
+								FUIAction( FExecuteAction::CreateSP( this, &FAssetContextMenu::ExecuteGoToDocsForAsset, UBlueprint::StaticClass(), FString(TEXT("UBlueprint_FunctionLibrary")) ) )
+								);
 							break;
 						case BPTYPE_Interface:
-							ExcerptSection = TEXT("UBlueprint_Interface");
+							MenuBuilder.AddMenuEntry(
+								LOCTEXT("GoToDocsForInterfaceBlueprint", "View Documentation - Interface"),
+								LOCTEXT("GoToDocsForInterfaceBlueprint_ToolTip", "Click to open documentation on blueprint interfaces"),
+								FSlateIcon(FEditorStyle::GetStyleSetName(), "HelpIcon.Hovered" ),
+								FUIAction( FExecuteAction::CreateSP( this, &FAssetContextMenu::ExecuteGoToDocsForAsset, UBlueprint::StaticClass(), FString(TEXT("UBlueprint_Interface")) ) )
+								);
 							break;
 						case BPTYPE_MacroLibrary:
-							ExcerptSection = TEXT("UBlueprint_Macro");
+							MenuBuilder.AddMenuEntry(
+								LOCTEXT("GoToDocsForInterfaceBlueprint", "View Documentation - Macro"),
+								LOCTEXT("GoToDocsForInterfaceBlueprint_ToolTip", "Click to open documentation on blueprint macros"),
+								FSlateIcon(FEditorStyle::GetStyleSetName(), "HelpIcon.Hovered" ),
+								FUIAction( FExecuteAction::CreateSP( this, &FAssetContextMenu::ExecuteGoToDocsForAsset, UBlueprint::StaticClass(), FString(TEXT("UBlueprint_Macro")) ) )
+								);
 							break;
 						default:
-							ExcerptSection = TEXT("UBlueprint");
+							MenuBuilder.AddMenuEntry(
+								LOCTEXT("GoToDocsForBlueprint", "View Documentation - Blueprint"),
+								LOCTEXT("GoToDocsForBlueprint_ToolTip", "Click to open documentation on blueprints"),
+								FSlateIcon(FEditorStyle::GetStyleSetName(), "HelpIcon.Hovered" ),
+								FUIAction( FExecuteAction::CreateSP( this, &FAssetContextMenu::ExecuteGoToDocsForAsset, UBlueprint::StaticClass(), FString(TEXT("UBlueprint")) ) )
+								);
 						}
-
-						MenuBuilder.AddMenuEntry(
-							LOCTEXT("GoToDocsForBlueprint", "View Documentation - Blueprint"),
-							LOCTEXT("GoToDocsForBlueprint_ToolTip", "Click to open documentation on blueprints"),
-							FSlateIcon(FEditorStyle::GetStyleSetName(), "HelpIcon.Hovered" ),
-							FUIAction( FExecuteAction::CreateSP( this, &FAssetContextMenu::ExecuteGoToDocsForAsset, UBlueprint::StaticClass(), ExcerptSection ) )
-							);
 					}
 					else
 					{
