Editor: Clarified section headings in the new asset menu

diff --git a/Engine/Source/Editor/ContentBrowser/Private/NewAssetContextMenu.cpp b/Engine/Source/Editor/ContentBrowser/Private/NewAssetContextMenu.cpp
--- a/Engine/Source/Editor/ContentBrowser/Private/NewAssetContextMenu.cpp
+++ b/Engine/Source/Editor/ContentBrowser/Private/NewAssetContextMenu.cpp
@@ -160,13 +160,13 @@ void FNewAssetContextMenu::MakeContextMenu(FMenuBuilder& MenuBuilder, const FStr
 	}
 	MenuBuilder.EndSection();
 
-	MenuBuilder.BeginSection("ContentBrowserNewBasicAsset", LOCTEXT("BasicAssetsMenuHeading", "Basic Assets") );
+	MenuBuilder.BeginSection("ContentBrowserNewBasicAsset", LOCTEXT("BasicAssetsMenuHeading", "Create Basic Asset") );
 	{
 		CreateNewAssetMenuCategory(MenuBuilder, EAssetTypeCategories::Basic, InPath, InOnNewAssetRequested);
 	}
 	MenuBuilder.EndSection(); //ContentBrowserNewBasicAsset
 
-	MenuBuilder.BeginSection("ContentBrowserNewAdvancedAsset", LOCTEXT("AdvancedAssetsMenuHeading", "Other Assets"));
+	MenuBuilder.BeginSection("ContentBrowserNewAdvancedAsset", LOCTEXT("AdvancedAssetsMenuHeading", "Create Advanced Asset"));
 	{
 		{
 			TArray<FFactoryItem> AnimationFactories = FindFactoriesInCategory(EAssetTypeCategories::Animation);
