UE4: Fixed a bug that caused folder names to incorrectly be invalidated during rename or create.

diff --git a/Engine/Source/Editor/ContentBrowser/Private/ContentBrowserUtils.cpp b/Engine/Source/Editor/ContentBrowser/Private/ContentBrowserUtils.cpp
--- a/Engine/Source/Editor/ContentBrowser/Private/ContentBrowserUtils.cpp
+++ b/Engine/Source/Editor/ContentBrowser/Private/ContentBrowserUtils.cpp
@@ -1362,23 +1362,23 @@ bool ContentBrowserUtils::IsValidObjectPathForCreate(const FString& ObjectPath,
 	return true;
 }
 
-bool ContentBrowserUtils::IsValidFolderPathForCreate(const FString& FolderPath, FText& OutErrorMessage)
+bool ContentBrowserUtils::IsValidFolderPathForCreate(const FString& InFolderPath, const FString& NewFolderName, FText& OutErrorMessage)
 {
-	const FString FolderName = FPackageName::ObjectPathToObjectName(FolderPath);
-
-	if (!ContentBrowserUtils::IsValidFolderName(FolderName, OutErrorMessage))
+	if (!ContentBrowserUtils::IsValidFolderName(NewFolderName, OutErrorMessage))
 	{
 		return false;
 	}
 
-	if (ContentBrowserUtils::DoesFolderExist(FolderPath))
+	const FString NewFolderPath = InFolderPath / NewFolderName;
+
+	if (ContentBrowserUtils::DoesFolderExist(NewFolderPath))
 	{
 		OutErrorMessage = LOCTEXT("RenameFolderAlreadyExists", "A folder already exists at this location with this name.");
 		return false;
 	}
 
 	// Make sure we are not creating a folder path that is too long
-	if (FolderPath.Len() > PLATFORM_MAX_FILEPATH_LENGTH - MAX_CLASS_NAME_LENGTH)
+	if (NewFolderPath.Len() > PLATFORM_MAX_FILEPATH_LENGTH - MAX_CLASS_NAME_LENGTH)
 	{
 		// The full path for the folder is too long
 		OutErrorMessage = FText::Format(LOCTEXT("RenameFolderPathTooLong",
diff --git a/Engine/Source/Editor/ContentBrowser/Private/ContentBrowserUtils.h b/Engine/Source/Editor/ContentBrowser/Private/ContentBrowserUtils.h
--- a/Engine/Source/Editor/ContentBrowser/Private/ContentBrowserUtils.h
+++ b/Engine/Source/Editor/ContentBrowser/Private/ContentBrowserUtils.h
@@ -175,6 +175,6 @@ namespace ContentBrowserUtils
 	/** Returns true if the specified path is available for object creation */
 	bool IsValidObjectPathForCreate(const FString& ObjectPath, FText& OutErrorMessage);
 
-	/** Returns true if the specified path is available for folder creation */
-	bool IsValidFolderPathForCreate(const FString& FolderPath, FText& OutErrorMessage);
+	/** Returns true if the specified folder name in the specified path is available for folder creation */
+	bool IsValidFolderPathForCreate(const FString& FolderPath, const FString& NewFolderName, FText& OutErrorMessage);
 }
\ No newline at end of file
diff --git a/Engine/Source/Editor/ContentBrowser/Private/SAssetView.cpp b/Engine/Source/Editor/ContentBrowser/Private/SAssetView.cpp
--- a/Engine/Source/Editor/ContentBrowser/Private/SAssetView.cpp
+++ b/Engine/Source/Editor/ContentBrowser/Private/SAssetView.cpp
@@ -3220,8 +3220,8 @@ bool SAssetView::AssetVerifyRenameCommit(const TSharedPtr<FAssetViewItem>& Item,
 	else
 	{
 		const TSharedPtr<FAssetViewFolder>& ItemAsFolder = StaticCastSharedPtr<FAssetViewFolder>(Item);
-		const FString NewFolderPath = FPaths::GetPath(ItemAsFolder->FolderPath) / NewNameString;
-		return ContentBrowserUtils::IsValidFolderPathForCreate(NewFolderPath, OutErrorMessage);
+		const FString FolderPath = FPaths::GetPath(ItemAsFolder->FolderPath);
+		return ContentBrowserUtils::IsValidFolderPathForCreate(FolderPath, NewNameString, OutErrorMessage);
 	}
 
 	return true;
