Making Chunk assignment context menu options experimental.
TTP#344152

diff --git a/Engine/Source/Editor/ContentBrowser/Private/AssetContextMenu.cpp b/Engine/Source/Editor/ContentBrowser/Private/AssetContextMenu.cpp
--- a/Engine/Source/Editor/ContentBrowser/Private/AssetContextMenu.cpp
+++ b/Engine/Source/Editor/ContentBrowser/Private/AssetContextMenu.cpp
@@ -249,29 +249,32 @@ bool FAssetContextMenu::AddCommonMenuOptions(FMenuBuilder& MenuBuilder)
 			)
 		);
 
-	MenuBuilder.AddMenuEntry(
-		LOCTEXT("AssignAssetChunk", "Assign to Chunk..."),
-		LOCTEXT("AssignAssetChunkTooltip", "Assign this asset to a specific Chunk"),
-		FSlateIcon(),
-		FUIAction(
-			FExecuteAction::CreateSP(this, &FAssetContextMenu::ExecuteAssignChunkID)
-		)
-	);
+	if (GetDefault<UEditorExperimentalSettings>()->bContextMenuChunkAssignments)
+	{
+		MenuBuilder.AddMenuEntry(
+			LOCTEXT("AssignAssetChunk","Assign to Chunk..."),
+			LOCTEXT("AssignAssetChunkTooltip","Assign this asset to a specific Chunk"),
+			FSlateIcon(),
+			FUIAction(
+				FExecuteAction::CreateSP(this,&FAssetContextMenu::ExecuteAssignChunkID)
+			)
+		);
 
-	MenuBuilder.AddSubMenu(
-		LOCTEXT("RemoveAssetFromChunk", "Remove from Chunk..."),
-		LOCTEXT("RemoveAssetFromChunkTooltip", "Removed an asset from a Chunk it's assigned to."),
-		FNewMenuDelegate::CreateRaw(this, &FAssetContextMenu::MakeChunkIDListMenu)
-	);
+		MenuBuilder.AddSubMenu(
+			LOCTEXT("RemoveAssetFromChunk","Remove from Chunk..."),
+			LOCTEXT("RemoveAssetFromChunkTooltip","Removed an asset from a Chunk it's assigned to."),
+			FNewMenuDelegate::CreateRaw(this,&FAssetContextMenu::MakeChunkIDListMenu)
+		);
 
-	MenuBuilder.AddMenuEntry(
-		LOCTEXT("RemoveAllChunkAssignments", "Remove from all Chunks"),
-		LOCTEXT("RemoveAllChunkAssignmentsTooltip", "Removed an asset from all Chunks it's assigned to."),
-		FSlateIcon(),
-		FUIAction(
-			FExecuteAction::CreateSP(this, &FAssetContextMenu::ExecuteRemoveAllChunkID)
-		)
-	);
+		MenuBuilder.AddMenuEntry(
+			LOCTEXT("RemoveAllChunkAssignments","Remove from all Chunks"),
+			LOCTEXT("RemoveAllChunkAssignmentsTooltip","Removed an asset from all Chunks it's assigned to."),
+			FSlateIcon(),
+			FUIAction(
+				FExecuteAction::CreateSP(this,&FAssetContextMenu::ExecuteRemoveAllChunkID)
+			)
+		);
+	}
 
 	if (CanExecuteDiffSelected())
 	{
diff --git a/Engine/Source/Editor/UnrealEd/Classes/Settings/EditorExperimentalSettings.h b/Engine/Source/Editor/UnrealEd/Classes/Settings/EditorExperimentalSettings.h
--- a/Engine/Source/Editor/UnrealEd/Classes/Settings/EditorExperimentalSettings.h
+++ b/Engine/Source/Editor/UnrealEd/Classes/Settings/EditorExperimentalSettings.h
@@ -63,6 +63,10 @@ public:
 	UPROPERTY(EditAnywhere, config, Category=Audio)
 	bool bShowAudioStreamingOptions;
 
+	/** Allows ChunkIDs to be assigned to assets to via the content browser context menu. */
+	UPROPERTY(EditAnywhere,config,Category=UserInterface,meta=(DisplayName="Allow ChunkID Assignments"))
+	bool bContextMenuChunkAssignments;
+
 	/**
 	 * Returns an event delegate that is executed when a setting has changed.
 	 *
