UE4: Asset pickers no longer match an asset's class name if there is a filter for exactly one class.

#coderevie Matt.Kuhlenschmidt

diff --git a/Engine/Source/Editor/ContentBrowser/Private/FrontendFilters.cpp b/Engine/Source/Editor/ContentBrowser/Private/FrontendFilters.cpp
--- a/Engine/Source/Editor/ContentBrowser/Private/FrontendFilters.cpp
+++ b/Engine/Source/Editor/ContentBrowser/Private/FrontendFilters.cpp
@@ -10,16 +10,21 @@
 // FFrontendFilter_Text
 /////////////////////////////////////////
 
-void AssetDataToString(AssetFilterType Asset, OUT TArray< FString >& Array)
+void AssetDataToExportTextString(AssetFilterType Asset, OUT TArray< FString >& Array)
 {
 	Array.Add(Asset.GetExportTextName());
 }
 
+void AssetDataToObjectPathString(AssetFilterType Asset, OUT TArray< FString >& Array)
+{
+	Array.Add(Asset.ObjectPath.ToString());
+}
+
 FFrontendFilter_Text::FFrontendFilter_Text()
 	: FFrontendFilter(nullptr)
-	, TextFilter( TTextFilter<AssetFilterType>::FItemToStringArray::CreateStatic( &AssetDataToString ) )
+	, TextFilter( TTextFilter<AssetFilterType>::FItemToStringArray::CreateStatic( &AssetDataToExportTextString ) )
 {
-	TextFilter.OnChanged().AddRaw(this, &FFrontendFilter_Text::HandleOnChangedEvent);
+	SetIncludeClassName(true);
 }
 
 FFrontendFilter_Text::~FFrontendFilter_Text()
@@ -42,6 +47,20 @@ void FFrontendFilter_Text::SetRawFilterText(const FText& InFilterText)
 	return TextFilter.SetRawFilterText(InFilterText);
 }
 
+void FFrontendFilter_Text::SetIncludeClassName(bool bIncludeClassName)
+{
+	if ( bIncludeClassName )
+	{
+		TextFilter = TTextFilter<AssetFilterType>::FItemToStringArray::CreateStatic( &AssetDataToExportTextString );
+	}
+	else
+	{
+		TextFilter = TTextFilter<AssetFilterType>::FItemToStringArray::CreateStatic( &AssetDataToObjectPathString );
+	}
+
+	TextFilter.OnChanged().AddRaw(this, &FFrontendFilter_Text::HandleOnChangedEvent);
+}
+
 void FFrontendFilter_Text::HandleOnChangedEvent()
 {
 	BroadcastChangedEvent();
diff --git a/Engine/Source/Editor/ContentBrowser/Private/FrontendFilters.h b/Engine/Source/Editor/ContentBrowser/Private/FrontendFilters.h
--- a/Engine/Source/Editor/ContentBrowser/Private/FrontendFilters.h
+++ b/Engine/Source/Editor/ContentBrowser/Private/FrontendFilters.h
@@ -28,6 +28,9 @@ public:
 	/** Set the Text to be used as the Filter's restrictions */
 	void SetRawFilterText(const FText& InFilterText);
 
+	/** If bIncludeClassName is true, the text filter will include an asset's class name in the search */
+	void SetIncludeClassName(bool bIncludeClassName);
+
 private:
 
 	/** Handler for the internal text filter */
diff --git a/Engine/Source/Editor/ContentBrowser/Private/SAssetPicker.cpp b/Engine/Source/Editor/ContentBrowser/Private/SAssetPicker.cpp
--- a/Engine/Source/Editor/ContentBrowser/Private/SAssetPicker.cpp
+++ b/Engine/Source/Editor/ContentBrowser/Private/SAssetPicker.cpp
@@ -55,6 +55,7 @@ void SAssetPicker::Construct( const FArguments& InArgs )
 	if (!InArgs._AssetPickerConfig.bAutohideSearchBar)
 	{
 		TextFilter = MakeShareable( new FFrontendFilter_Text() );
+		TextFilter->SetIncludeClassName(InArgs._AssetPickerConfig.Filter.ClassNames.Num() != 1);
 		FrontendFilters->Add( TextFilter );
 		HighlightText = TAttribute< FText >( this, &SAssetPicker::GetHighlightedText );
 
@@ -300,6 +301,9 @@ void SAssetPicker::SetNewBackendFilter(const FARFilter& NewFilter)
 
 	CurrentBackendFilter = NewFilter;
 	CurrentBackendFilter.PackagePaths.Empty();
+
+	// Update the Text filter too, since now class names may no longer matter
+	TextFilter->SetIncludeClassName(NewFilter.ClassNames.Num() != 1);
 	
 	OnFilterChanged();
 }
