UE4: All path views are now affected by the "Show Engine Content" checkbox (and show developers and show plugins), not just path views found in content browsers.

diff --git a/Engine/Source/Editor/ContentBrowser/Private/SContentBrowser.cpp b/Engine/Source/Editor/ContentBrowser/Private/SContentBrowser.cpp
--- a/Engine/Source/Editor/ContentBrowser/Private/SContentBrowser.cpp
+++ b/Engine/Source/Editor/ContentBrowser/Private/SContentBrowser.cpp
@@ -48,9 +48,6 @@ void SContentBrowser::Construct( const FArguments& InArgs, const FName& InInstan
 	
 	bIsLocked = InArgs._InitiallyLocked;
 
-	// Listen for when view settings are changed
-	UContentBrowserSettings::OnSettingChanged().AddSP(this, &SContentBrowser::HandleSettingChanged);
-
 	HistoryManager.SetOnApplyHistoryData(FOnApplyHistoryData::CreateSP(this, &SContentBrowser::OnApplyHistoryData));
 	HistoryManager.SetOnUpdateHistoryData(FOnUpdateHistoryData::CreateSP(this, &SContentBrowser::OnUpdateHistoryData));
 
@@ -1818,48 +1815,6 @@ void SContentBrowser::OnAssetViewRefreshRequested()
 	AssetViewPtr->RequestSlowFullListRefresh();
 }
 
-void SContentBrowser::HandleSettingChanged(FName PropertyName)
-{
-	if ((PropertyName == "DisplayDevelopersFolder") ||
-		(PropertyName == "DisplayEngineFolder") ||
-		(PropertyName == "DisplayPluginFolders") ||
-		(PropertyName == NAME_None))	// @todo: Needed if PostEditChange was called manually, for now
-	{
-		// If the dev or engine folder is no longer visible but we're inside it...
-		const bool bDisplayDev = GetDefault<UContentBrowserSettings>()->GetDisplayDevelopersFolder();
-		const bool bDisplayEngine = GetDefault<UContentBrowserSettings>()->GetDisplayEngineFolder();
-		const bool bDisplayPlugins = GetDefault<UContentBrowserSettings>()->GetDisplayPluginFolders();
-		if ( !bDisplayDev || !bDisplayEngine || !bDisplayPlugins )
-		{
-		    const FString OldSelectedPath = PathViewPtr->GetSelectedPath();
-		    if ( (!bDisplayDev && ContentBrowserUtils::IsDevelopersFolder( OldSelectedPath )) || (!bDisplayEngine && ContentBrowserUtils::IsEngineFolder( OldSelectedPath )) || (!bDisplayPlugins && ContentBrowserUtils::IsPluginFolder( OldSelectedPath ) ) )
-		    {
-			    // Set the folder back to the root, and refresh the contents
-			    TArray<FString> SelectedPaths;
-			    SelectedPaths.Add(TEXT("/Game"));
-			    PathViewPtr->SetSelectedPaths(SelectedPaths);
-			    SourcesChanged(SelectedPaths, TArray<FCollectionNameType>());
-		    }
-		}
-
-		// Update our path view so that it can include/exclude the dev folder
-		PathViewPtr->Populate();
-
-		// If the dev or engine folder has become visible and we're inside it...
-		if ( bDisplayDev || bDisplayEngine || bDisplayPlugins )
-		{
-			const FString NewSelectedPath = PathViewPtr->GetSelectedPath();
-			if ( (bDisplayDev && ContentBrowserUtils::IsDevelopersFolder( NewSelectedPath )) || (bDisplayEngine && ContentBrowserUtils::IsEngineFolder( NewSelectedPath )) || (bDisplayPlugins && ContentBrowserUtils::IsPluginFolder( NewSelectedPath ) ) )
-			{
-				// Refresh the contents
-				TArray<FString> SelectedPaths;
-				SelectedPaths.Add(NewSelectedPath);
-				SourcesChanged(SelectedPaths, TArray<FCollectionNameType>());
-			}
-		}
-	}
-}
-
 void SContentBrowser::HandleCollectionRemoved(const FCollectionNameType& Collection)
 {
 	AssetViewPtr->SetSourcesData(FSourcesData());
diff --git a/Engine/Source/Editor/ContentBrowser/Private/SContentBrowser.h b/Engine/Source/Editor/ContentBrowser/Private/SContentBrowser.h
--- a/Engine/Source/Editor/ContentBrowser/Private/SContentBrowser.h
+++ b/Engine/Source/Editor/ContentBrowser/Private/SContentBrowser.h
@@ -281,9 +281,6 @@ private:
 	/** Handler for when the asset context menu requests to refresh the asset view */
 	void OnAssetViewRefreshRequested();
 
-	/** Delegate called when an editor setting is changed */
-	void HandleSettingChanged(FName PropertyName);
-
 	/** Handles an on collection destroyed event */
 	void HandleCollectionRemoved(const FCollectionNameType& Collection);
 
diff --git a/Engine/Source/Editor/ContentBrowser/Private/SPathView.cpp b/Engine/Source/Editor/ContentBrowser/Private/SPathView.cpp
--- a/Engine/Source/Editor/ContentBrowser/Private/SPathView.cpp
+++ b/Engine/Source/Editor/ContentBrowser/Private/SPathView.cpp
@@ -36,6 +36,9 @@ void SPathView::Construct( const FArguments& InArgs )
 	bAllowClassesFolder = InArgs._AllowClassesFolder;
 	PreventTreeItemChangedDelegateCount = 0;
 
+	// Listen for when view settings are changed
+	UContentBrowserSettings::OnSettingChanged().AddSP(this, &SPathView::HandleSettingChanged);
+
 	//Setup the SearchBox filter
 	SearchBoxFolderFilter = MakeShareable( new FolderTextFilter( FolderTextFilter::FItemToStringArray::CreateSP( this, &SPathView::PopulateFolderSearchStrings ) ) );
 	SearchBoxFolderFilter->OnChanged().AddSP( this, &SPathView::FilterUpdated );
@@ -477,7 +480,7 @@ void SPathView::SyncToAssets( const TArray<FAssetData>& AssetDataList, const boo
 		else
 		{
 			// Explicit sync so just clear the selection
-		TreeViewPtr->ClearSelection();
+			TreeViewPtr->ClearSelection();
 		}
 
 		// SyncTreeItems should now only contain items which aren't already shown explicitly or implicitly (as a child)
@@ -491,10 +494,10 @@ void SPathView::SyncToAssets( const TArray<FAssetData>& AssetDataList, const boo
 		if (SyncTreeItems.Num() > 0)
 		{
 			// Scroll the first item into view if applicable
-		TreeViewPtr->RequestScrollIntoView(SyncTreeItems[0]);
+			TreeViewPtr->RequestScrollIntoView(SyncTreeItems[0]);
+		}
 	}
 }
-}
 
 TSharedPtr<FTreeItem> SPathView::FindItemRecursive(const FString& Path) const
 {
@@ -811,17 +814,17 @@ void SPathView::TreeSelectionChanged( TSharedPtr< FTreeItem > TreeItem, ESelectI
 
 		if ( OnPathSelected.IsBound() )
 		{
-		if ( TreeItem.IsValid() )
-		{
-			OnPathSelected.Execute(TreeItem->FolderPath);
-		}
-		else
-		{
-			OnPathSelected.Execute(TEXT(""));
+			if ( TreeItem.IsValid() )
+			{
+				OnPathSelected.Execute(TreeItem->FolderPath);
+			}
+			else
+			{
+				OnPathSelected.Execute(TEXT(""));
+			}
 		}
 	}
 }
-}
 
 void SPathView::TreeExpansionChanged( TSharedPtr< FTreeItem > TreeItem, bool bIsExpanded )
 {
@@ -1321,12 +1324,56 @@ void SPathView::OnAssetRegistrySearchCompleted()
 	PendingInitialPaths.Empty();
 }
 
-
 void SPathView::OnContentPathMounted( const FString& ContentPath )
 {
 	// A new content path has appeared, so we should refresh out root set of paths
 	bNeedsRepopulate = true;
 }
 
+void SPathView::HandleSettingChanged(FName PropertyName)
+{
+	if ((PropertyName == "DisplayDevelopersFolder") ||
+		(PropertyName == "DisplayEngineFolder") ||
+		(PropertyName == "DisplayPluginFolders") ||
+		(PropertyName == NAME_None))	// @todo: Needed if PostEditChange was called manually, for now
+	{
+		// If the dev or engine folder is no longer visible but we're inside it...
+		const bool bDisplayDev = GetDefault<UContentBrowserSettings>()->GetDisplayDevelopersFolder();
+		const bool bDisplayEngine = GetDefault<UContentBrowserSettings>()->GetDisplayEngineFolder();
+		const bool bDisplayPlugins = GetDefault<UContentBrowserSettings>()->GetDisplayPluginFolders();
+		if (!bDisplayDev || !bDisplayEngine || !bDisplayPlugins)
+		{
+			const FString OldSelectedPath = GetSelectedPath();
+			if ((!bDisplayDev && ContentBrowserUtils::IsDevelopersFolder(OldSelectedPath)) || (!bDisplayEngine && ContentBrowserUtils::IsEngineFolder(OldSelectedPath)) || (!bDisplayPlugins && ContentBrowserUtils::IsPluginFolder(OldSelectedPath)))
+			{
+				// Set the folder back to the root, and refresh the contents
+				TSharedPtr<FTreeItem> GameRoot = FindItemRecursive(TEXT("/Game"));
+				if ( GameRoot.IsValid() )
+				{
+					TreeViewPtr->SetSelection(GameRoot);
+				}
+				else
+				{
+					TreeViewPtr->ClearSelection();
+				}
+			}
+		}
+
+		// Update our path view so that it can include/exclude the dev folder
+		Populate();
+
+		// If the dev or engine folder has become visible and we're inside it...
+		if (bDisplayDev || bDisplayEngine || bDisplayPlugins)
+		{
+			const FString NewSelectedPath = GetSelectedPath();
+			if ((bDisplayDev && ContentBrowserUtils::IsDevelopersFolder(NewSelectedPath)) || (bDisplayEngine && ContentBrowserUtils::IsEngineFolder(NewSelectedPath)) || (bDisplayPlugins && ContentBrowserUtils::IsPluginFolder(NewSelectedPath)))
+			{
+				// Refresh the contents
+				OnPathSelected.ExecuteIfBound(NewSelectedPath);
+			}
+		}
+	}
+}
+
 
 #undef LOCTEXT_NAMESPACE
diff --git a/Engine/Source/Editor/ContentBrowser/Private/SPathView.h b/Engine/Source/Editor/ContentBrowser/Private/SPathView.h
--- a/Engine/Source/Editor/ContentBrowser/Private/SPathView.h
+++ b/Engine/Source/Editor/ContentBrowser/Private/SPathView.h
@@ -222,6 +222,9 @@ private:
 	/** Called from an engine core event when a new content path has been added, so that we can refresh our root set of paths */
 	void OnContentPathMounted( const FString& ContentPath );
 
+	/** Delegate called when an editor setting is changed */
+	void HandleSettingChanged(FName PropertyName);
+
 private:
 
 	/** A helper class to manage PreventTreeItemChangedDelegateCount by incrementing it when constructed (on the stack) and decrementing when destroyed */
