UE4: SAssetPickers can now save their view settings. The SAssetDialog saves it's asset picker view settings.

diff --git a/Engine/Source/Editor/ContentBrowser/Private/SAssetDialog.cpp b/Engine/Source/Editor/ContentBrowser/Private/SAssetDialog.cpp
--- a/Engine/Source/Editor/ContentBrowser/Private/SAssetDialog.cpp
+++ b/Engine/Source/Editor/ContentBrowser/Private/SAssetDialog.cpp
@@ -30,6 +30,7 @@ void SAssetDialog::Construct(const FArguments& InArgs, const FSharedAssetDialogC
 	AssetPickerConfig.OnAssetsActivated = FOnAssetsActivated::CreateSP(this, &SAssetDialog::OnAssetsActivated);
 	AssetPickerConfig.SetFilterDelegates.Add(&SetFilterDelegate);
 	AssetPickerConfig.GetCurrentSelectionDelegates.Add(&GetCurrentSelectionDelegate);
+	AssetPickerConfig.SaveSettingsName = TEXT("AssetDialog");
 	TSharedPtr<SWidget> AssetPicker = FContentBrowserSingleton::Get().CreateAssetPicker(AssetPickerConfig);
 
 	SetCurrentlySelectedPath(DefaultPath);
diff --git a/Engine/Source/Editor/ContentBrowser/Private/SAssetPicker.cpp b/Engine/Source/Editor/ContentBrowser/Private/SAssetPicker.cpp
--- a/Engine/Source/Editor/ContentBrowser/Private/SAssetPicker.cpp
+++ b/Engine/Source/Editor/ContentBrowser/Private/SAssetPicker.cpp
@@ -4,6 +4,11 @@
 
 #define LOCTEXT_NAMESPACE "ContentBrowser"
 
+SAssetPicker::~SAssetPicker()
+{
+	SaveSettings();
+}
+
 void SAssetPicker::Construct( const FArguments& InArgs )
 {
 	TSharedPtr<AssetFilterCollectionType> FrontendFilters = MakeShareable(new AssetFilterCollectionType());
@@ -16,6 +21,7 @@ void SAssetPicker::Construct( const FArguments& InArgs )
 	OnAssetEnterPressed = InArgs._AssetPickerConfig.OnAssetEnterPressed;
 	bPendingFocusNextFrame = InArgs._AssetPickerConfig.bFocusSearchBoxWhenOpened;
 	DefaultFilterMenuExpansion = InArgs._AssetPickerConfig.DefaultFilterMenuExpansion;
+	SaveSettingsName = InArgs._AssetPickerConfig.SaveSettingsName;
 
 	for (auto DelegateIt = InArgs._AssetPickerConfig.GetCurrentSelectionDelegates.CreateConstIterator(); DelegateIt; ++DelegateIt)
 	{
@@ -216,6 +222,8 @@ void SAssetPicker::Construct( const FArguments& InArgs )
 		.AllowFocusOnSync(false)	// Stop the asset view from stealing focus (we're in control of that)
 	];
 
+	LoadSettings();
+
 	AssetViewPtr->RequestSlowFullListRefresh();
 }
 
@@ -421,4 +429,36 @@ void SAssetPicker::BindCommands()
 		));
 }
 
+void SAssetPicker::LoadSettings()
+{
+	const FString& SettingsString = SaveSettingsName;
+
+	if ( !SettingsString.IsEmpty() )
+	{
+		// Load all our data using the settings string as a key in the user settings ini
+		if (FilterListPtr.IsValid())
+		{
+			FilterListPtr->LoadSettings(GEditorUserSettingsIni, SContentBrowser::SettingsIniSection, SettingsString);
+		}
+		
+		AssetViewPtr->LoadSettings(GEditorUserSettingsIni, SContentBrowser::SettingsIniSection, SettingsString);
+	}
+}
+
+void SAssetPicker::SaveSettings() const
+{
+	const FString& SettingsString = SaveSettingsName;
+
+	if ( !SettingsString.IsEmpty() )
+	{
+		// Save all our data using the settings string as a key in the user settings ini
+		if (FilterListPtr.IsValid())
+		{
+			FilterListPtr->SaveSettings(GEditorUserSettingsIni, SContentBrowser::SettingsIniSection, SettingsString);
+		}
+
+		AssetViewPtr->SaveSettings(GEditorUserSettingsIni, SContentBrowser::SettingsIniSection, SettingsString);
+	}
+}
+
 #undef LOCTEXT_NAMESPACE
\ No newline at end of file
diff --git a/Engine/Source/Editor/ContentBrowser/Private/SAssetPicker.h b/Engine/Source/Editor/ContentBrowser/Private/SAssetPicker.h
--- a/Engine/Source/Editor/ContentBrowser/Private/SAssetPicker.h
+++ b/Engine/Source/Editor/ContentBrowser/Private/SAssetPicker.h
@@ -15,6 +15,8 @@ public:
 
 	SLATE_END_ARGS()
 
+	virtual ~SAssetPicker();
+
 	/** Constructs this widget with InArgs */
 	void Construct( const FArguments& InArgs );
 
@@ -72,6 +74,12 @@ private:
 	/** Bind our UI commands */
 	void BindCommands();
 
+	/** Loads settings for this asset picker if SaveSettingsName was set */
+	void LoadSettings();
+
+	/** Saves settings for this asset picker if SaveSettingsName was set */
+	void SaveSettings() const;
+
 private:
 	/** The asset view widget */
 	TSharedPtr<SAssetView> AssetViewPtr;
@@ -112,4 +120,7 @@ private:
 
 	/** UICommand list, holds list of actions for processing */
 	TSharedPtr< FUICommandList > Commands;
+
+	/** If set, view settings will be saved and loaded for the asset view using this name in ini files */
+	FString SaveSettingsName;
 };
\ No newline at end of file
