UE4: Missing file from previous checkin

diff --git a/Engine/Source/Editor/ContentBrowser/Public/IContentBrowserSingleton.h b/Engine/Source/Editor/ContentBrowser/Public/IContentBrowserSingleton.h
--- a/Engine/Source/Editor/ContentBrowser/Public/IContentBrowserSingleton.h
+++ b/Engine/Source/Editor/ContentBrowser/Public/IContentBrowserSingleton.h
@@ -93,6 +93,9 @@ struct FAssetPickerConfig
 	/** The text to show when there are no assets to show */
 	TAttribute< FText > AssetShowWarningText;
 
+	/** If set, view settings will be saved and loaded for the asset view using this name in ini files */
+	FString SaveSettingsName;
+
 	/** If true, the search box will gain focus when the asset picker is created */
 	bool bFocusSearchBoxWhenOpened;
 
