Improved the validation for working out if a path is too long to cook

TTP# 332328 - Content Browser: Need new warnings about long relative path lengths for assets.

https://answers.unrealengine.com/questions/60187/the-path-to-the-asset-is-too-long-for-cooking-im-l.html

The previous logic assumed that the project would always be in the same source tree as UE4, which isn't the case for content creation and Rocket projects.

In the case where the project was outside the UE4 source tree, the validation often failed as the path was never trimmed to remove the relative part of it (between the project and the UE4 binary).

This change gracefully handles both cases, and now performs two sets of validation:
 1) It checks that the content can be cooked based on where the project actually is (for local cooking)
 2) If it's an internal build, it also checks that the content will be cookable via our build machines (forthe build farm cooking)

diff --git a/Engine/Source/Editor/ContentBrowser/Private/ContentBrowserUtils.cpp b/Engine/Source/Editor/ContentBrowser/Private/ContentBrowserUtils.cpp
--- a/Engine/Source/Editor/ContentBrowser/Private/ContentBrowserUtils.cpp
+++ b/Engine/Source/Editor/ContentBrowser/Private/ContentBrowserUtils.cpp
@@ -9,6 +9,7 @@
 #include "ImageUtils.h"
 #include "ISourceControlModule.h"
 #include "MessageLog.h"
+#include "EngineBuildSettings.h"
 
 #define LOCTEXT_NAMESPACE "ContentBrowser"
 
@@ -1309,28 +1310,8 @@ bool ContentBrowserUtils::IsValidObjectPathForCreate(const FString& ObjectPath,
 
 	const FString PackageName = FPackageName::ObjectPathToPackageName(ObjectPath);
 
-	// The following checks are done mostly to prevent / alleviate the problems that "long" paths are causing with the BuildFarm and cooked builds.
-	// The BuildFarm buildmachines use a verbose path to encode extra information to provide more information when things fail, however
-	// this makes the path limitation (260 chars on Windows) a problem. It doubles up the GGameName and does the cooking in another
-	// sub-folder, one of the "saved/sandboxes", with folder duplication.
-
-	// Get the SubPath containing folders without the "game name" folder itself
-	const FString GameNameStr(GGameName);
-	FString SubPath = FPaths::GameDir();
-	FPaths::NormalizeDirectoryName(SubPath);
-	SubPath = SubPath.Replace(*(FString(TEXT("../../../")) + GameNameStr), TEXT(""));
-	FPaths::RemoveDuplicateSlashes(SubPath);
-
-	// Calculate the maximum path length this will generate when doing a cooked build.
-	const int32 MaxProjectedCookingPath = 165;
-	const int32 PathCalcLen = SubPath.Len() + (2 * GameNameStr.Len()) + (PackageName + FPackageName::GetAssetPackageExtension()).Len();
-	if ( PathCalcLen >= MaxProjectedCookingPath )
+	if (!IsValidPackageForCooking(PackageName, OutErrorMessage))
 	{
-		// The projected length of the path for cooking is too long
-		OutErrorMessage = FText::Format( LOCTEXT("AssetCookingPathTooLong", 
-			"The path to the asset is too long for cooking, the maximum is '{0}' characters.\nPlease choose a shorter name for the asset or create it in a shallower folder structure with shorter folder names."),
-			FText::AsNumber(MaxProjectedCookingPath) );
-		// Return false to indicate that the user should enter a new name
 		return false;
 	}
 
@@ -1391,4 +1372,81 @@ bool ContentBrowserUtils::IsValidFolderPathForCreate(const FString& InFolderPath
 	return true;
 }
 
+bool ContentBrowserUtils::IsValidPackageForCooking(const FString& PackageName, FText& OutErrorMessage)
+{
+	// We assume the game name is 20 characters (the maximum allowed) to make sure that content can be ported between projects
+	// 260 characters is the limit on Windows, which is the shortest max path of any platforms that support cooking
+	static const int32 MaxGameNameLen = 20;
+	static const int32 MaxPathLen = 260;
+
+	// Pad out the game name to the maximum allowed
+	const FString GameName = FApp::GetGameName();
+	FString GameNamePadded = GameName;
+	while (GameNamePadded.Len() < MaxGameNameLen)
+	{
+		GameNamePadded += TEXT(" ");
+	}
+
+	// We use "WindowsNoEditor" below as it's the longest platform name, so will also prove that any shorter platform names will validate correctly
+	const FString AbsoluteRootPath = FPaths::ConvertRelativePathToFull(FPaths::RootDir());
+	const FString AbsoluteGamePath = FPaths::ConvertRelativePathToFull(FPaths::GameDir());
+	const FString AbsoluteCookPath = AbsoluteGamePath / TEXT("Saved") / TEXT("Cooked") / TEXT("WindowsNoEditor") / GameName;
+
+	const FString RelativePathToAsset = FPackageName::LongPackageNameToFilename(PackageName, FPackageName::GetAssetPackageExtension());
+	const FString AbsolutePathToAsset = FPaths::ConvertRelativePathToFull(RelativePathToAsset);
+
+	FString AssetPathWithinCookDir = AbsolutePathToAsset;
+	FPaths::RemoveDuplicateSlashes(AssetPathWithinCookDir);
+	AssetPathWithinCookDir.RemoveFromStart(AbsoluteGamePath, ESearchCase::CaseSensitive);
+
+	// Test that the package can be cooked based on the current project path
+	{
+		FString AbsoluteCookPathToAsset = AbsoluteCookPath / AssetPathWithinCookDir;
+		AbsoluteCookPathToAsset.ReplaceInline(*GameName, *GameNamePadded, ESearchCase::CaseSensitive);
+
+		if (AbsoluteCookPathToAsset.Len() > MaxPathLen)
+		{
+			// The projected length of the path for cooking is too long
+			OutErrorMessage = LOCTEXT("AssetCookingPathTooLong", "The path to the asset is too long for cooking\nPlease choose a shorter name for the asset or create it in a shallower folder structure with shorter folder names.");
+			// Return false to indicate that the user should enter a new name
+			return false;
+		}
+	}
+
+	// See TTP# 332328:
+	// The following checks are done mostly to prevent / alleviate the problems that "long" paths are causing with the BuildFarm and cooked builds.
+	// The BuildFarm uses a verbose path to encode extra information to provide more information when things fail, however this makes the path limitation a problem.
+	//	- We assume a base path of D:/BuildFarm/buildmachine_++depot+UE4-Releases+4.10/
+	//	- We assume the game name is 20 characters (the maximum allowed) to make sure that content can be ported between projects
+	//	- We calculate the cooked game path relative to the game root (eg, Showcases/Infiltrator/Saved/Cooked/WindowsNoEditor/Infiltrator)
+	//	- We calculate the asset path relative to (and including) the Content directory (eg, Content/Environment/Infil1/Infil1_Underground/Infrastructure/Model/SM_Infil1_Tunnel_Ceiling_Pipes_1xEntryCurveOuter_Double.uasset)
+	if (FEngineBuildSettings::IsInternalBuild())
+	{
+		// We assume a constant size for the build machine base path, so strip either the root or game path from the start
+		// (depending on whether the project is part of the main UE4 source tree or located elsewhere)
+		FString CookDirWithoutBasePath = AbsoluteCookPath;
+		if (CookDirWithoutBasePath.StartsWith(AbsoluteRootPath, ESearchCase::CaseSensitive))
+		{
+			CookDirWithoutBasePath.RemoveFromStart(AbsoluteRootPath, ESearchCase::CaseSensitive);
+		}
+		else
+		{
+			CookDirWithoutBasePath.RemoveFromStart(AbsoluteGamePath, ESearchCase::CaseSensitive);
+		}
+
+		FString AbsoluteBuildMachineCookPathToAsset = FString(TEXT("D:/BuildFarm/buildmachine_++depot+UE4-Releases+4.10")) / CookDirWithoutBasePath / AssetPathWithinCookDir;
+		AbsoluteBuildMachineCookPathToAsset.ReplaceInline(*GameName, *GameNamePadded, ESearchCase::CaseSensitive);
+
+		if (AbsoluteBuildMachineCookPathToAsset.Len() > MaxPathLen)
+		{
+			// The projected length of the path for cooking is too long
+			OutErrorMessage = LOCTEXT("AssetCookingPathTooLongForBuildMachine", "The path to the asset is too long for cooking by the build machines\nPlease choose a shorter name for the asset or create it in a shallower folder structure with shorter folder names.");
+			// Return false to indicate that the user should enter a new name
+			return false;
+		}
+	}
+
+	return true;
+}
+
 #undef LOCTEXT_NAMESPACE
diff --git a/Engine/Source/Editor/ContentBrowser/Private/ContentBrowserUtils.h b/Engine/Source/Editor/ContentBrowser/Private/ContentBrowserUtils.h
--- a/Engine/Source/Editor/ContentBrowser/Private/ContentBrowserUtils.h
+++ b/Engine/Source/Editor/ContentBrowser/Private/ContentBrowserUtils.h
@@ -177,4 +177,7 @@ namespace ContentBrowserUtils
 
 	/** Returns true if the specified folder name in the specified path is available for folder creation */
 	bool IsValidFolderPathForCreate(const FString& FolderPath, const FString& NewFolderName, FText& OutErrorMessage);
+
+	/** Checks to see whether the path is within the size restrictions for cooking */
+	bool IsValidPackageForCooking(const FString& PackageName, FText& OutErrorMessage);
 }
\ No newline at end of file
