fixed non-unity errors

diff --git a/Engine/Source/Editor/ContentBrowser/Private/CollectionContextMenu.cpp b/Engine/Source/Editor/ContentBrowser/Private/CollectionContextMenu.cpp
--- a/Engine/Source/Editor/ContentBrowser/Private/CollectionContextMenu.cpp
+++ b/Engine/Source/Editor/ContentBrowser/Private/CollectionContextMenu.cpp
@@ -5,9 +5,12 @@
 #include "CollectionContextMenu.h"
 #include "ISourceControlModule.h"
 #include "ContentBrowserModule.h"
+#include "SColorPicker.h"
+
 
 #define LOCTEXT_NAMESPACE "ContentBrowser"
 
+
 FCollectionContextMenu::FCollectionContextMenu(const TWeakPtr<SCollectionView>& InCollectionView)
 	: CollectionView(InCollectionView)
 	, bProjectUnderSourceControl(false)
diff --git a/Engine/Source/Editor/ContentBrowser/Private/PathContextMenu.cpp b/Engine/Source/Editor/ContentBrowser/Private/PathContextMenu.cpp
--- a/Engine/Source/Editor/ContentBrowser/Private/PathContextMenu.cpp
+++ b/Engine/Source/Editor/ContentBrowser/Private/PathContextMenu.cpp
@@ -9,9 +9,12 @@
 #include "ReferenceViewer.h"
 #include "AssetToolsModule.h"
 #include "Editor/UnrealEd/Public/PackageTools.h"
+#include "SColorPicker.h"
+
 
 #define LOCTEXT_NAMESPACE "ContentBrowser"
 
+
 FPathContextMenu::FPathContextMenu(const TWeakPtr<SWidget>& InParentContent)
 	: ParentContent(InParentContent)
 	, bCanExecuteSCCCheckOut(false)
