Don't create Open Map dialog as a child window if root window is NULL

diff --git a/Engine/Source/Editor/ContentBrowser/Private/ContentBrowserSingleton.cpp b/Engine/Source/Editor/ContentBrowser/Private/ContentBrowserSingleton.cpp
--- a/Engine/Source/Editor/ContentBrowser/Private/ContentBrowserSingleton.cpp
+++ b/Engine/Source/Editor/ContentBrowser/Private/ContentBrowserSingleton.cpp
@@ -348,10 +348,14 @@ void FContentBrowserSingleton::SharedCreateAssetDialogWindow(const TSharedRef<SA
 		{
 			FSlateApplication::Get().AddModalWindow(DialogWindow, MainFrameParentWindow.ToSharedRef());
 		}
-		else
+		else if (FGlobalTabmanager::Get()->GetRootWindow().IsValid())
 		{
 			FSlateApplication::Get().AddWindowAsNativeChild(DialogWindow, MainFrameParentWindow.ToSharedRef());
 		}
+		else
+		{
+			FSlateApplication::Get().AddWindow(DialogWindow);
+		}
 	}
 	else
 	{
