Rename 'Create Copy' to 'Duplicate' in Content Browser

diff --git a/Engine/Source/Editor/ContentBrowser/Private/AssetContextMenu.cpp b/Engine/Source/Editor/ContentBrowser/Private/AssetContextMenu.cpp
--- a/Engine/Source/Editor/ContentBrowser/Private/AssetContextMenu.cpp
+++ b/Engine/Source/Editor/ContentBrowser/Private/AssetContextMenu.cpp
@@ -166,7 +166,7 @@ bool FAssetContextMenu::AddCommonMenuOptions(FMenuBuilder& MenuBuilder)
 		);
 
 	MenuBuilder.AddMenuEntry(
-		LOCTEXT("Duplicate", "Create Copy"),
+		LOCTEXT("Duplicate", "Duplicate"),
 		LOCTEXT("DuplicateTooltip", "Create a copy of the selected assets."),
 		FSlateIcon(),
 		FUIAction(
