Converted widget .Tag instances for tutorials to use MetaData

diff --git a/Engine/Source/Developer/LogVisualizer/Private/SLogVisualizer.cpp b/Engine/Source/Developer/LogVisualizer/Private/SLogVisualizer.cpp
--- a/Engine/Source/Developer/LogVisualizer/Private/SLogVisualizer.cpp
+++ b/Engine/Source/Developer/LogVisualizer/Private/SLogVisualizer.cpp
@@ -360,7 +360,7 @@ void SLogVisualizer::Construct(const FArguments& InArgs, FLogVisualizer* InLogVi
 			[
 				SAssignNew(FilterListPtr, SLogFilterList)
 				.OnFilterChanged(this, &SLogVisualizer::OnLogCategoryFiltersChanged)
-				.Tag(TEXT("CategoryFilters"))
+				.AddMetaData<FTagMetaData>(FTagMetaData(TEXT("CategoryFilters")))
 				/*.OnGetContextMenu(this, &SLogVisualizer::GetFilterContextMenu)*/
 				/*.FrontendFilters(FrontendFilters)*/
 			]
diff --git a/Engine/Source/Editor/ContentBrowser/ContentBrowser.Build.cs b/Engine/Source/Editor/ContentBrowser/ContentBrowser.Build.cs
--- a/Engine/Source/Editor/ContentBrowser/ContentBrowser.Build.cs
+++ b/Engine/Source/Editor/ContentBrowser/ContentBrowser.Build.cs
@@ -49,5 +49,11 @@ public class ContentBrowser : ModuleRules
                 "ReferenceViewer"
 			}
 		);
+		
+		PublicIncludePathModuleNames.AddRange(
+            new string[] {                
+                "IntroTutorials"
+            }
+        );
 	}
 }
diff --git a/Engine/Source/Editor/ContentBrowser/Private/AssetViewWidgets.cpp b/Engine/Source/Editor/ContentBrowser/Private/AssetViewWidgets.cpp
--- a/Engine/Source/Editor/ContentBrowser/Private/AssetViewWidgets.cpp
+++ b/Engine/Source/Editor/ContentBrowser/Private/AssetViewWidgets.cpp
@@ -963,7 +963,7 @@ void SAssetListItem::Construct( const FArguments& InArgs )
 		SNew(SBorder)
 		.BorderImage(this, &SAssetViewItem::GetBorderImage)
 		.Padding(0)
-		.Tag(AssetItem->GetType() == EAssetItemType::Normal ? StaticCastSharedPtr<FAssetViewAsset>(AssetItem)->Data.ObjectPath : NAME_None)
+		.AddMetaData<FTagMetaData>(FTagMetaData(AssetItem->GetType() == EAssetItemType::Normal ? StaticCastSharedPtr<FAssetViewAsset>(AssetItem)->Data.ObjectPath : NAME_None))
 		[
 			SNew(SHorizontalBox)
 
@@ -1175,7 +1175,7 @@ void SAssetTileItem::Construct( const FArguments& InArgs )
 		SNew(SBorder)
 		.BorderImage(this, &SAssetViewItem::GetBorderImage)
 		.Padding(0)
-		.Tag(AssetItem->GetType() == EAssetItemType::Normal ? StaticCastSharedPtr<FAssetViewAsset>(AssetItem)->Data.ObjectPath : NAME_None)
+		.AddMetaData<FTagMetaData>(FTagMetaData(AssetItem->GetType() == EAssetItemType::Normal ? StaticCastSharedPtr<FAssetViewAsset>(AssetItem)->Data.ObjectPath : NAME_None))
 		[
 			SNew(SVerticalBox)
 
@@ -1383,7 +1383,7 @@ TSharedRef<SWidget> SAssetColumnItem::GenerateWidgetForColumn( const FName& Colu
 		const float IconOverlaySize = IconBrush->ImageSize.X * 0.6f;
 
 		Content = SNew(SHorizontalBox)
-			.Tag(AssetItem->GetType() == EAssetItemType::Normal ? StaticCastSharedPtr<FAssetViewAsset>(AssetItem)->Data.ObjectPath : NAME_None)
+			.AddMetaData<FTagMetaData>(FTagMetaData(AssetItem->GetType() == EAssetItemType::Normal ? StaticCastSharedPtr<FAssetViewAsset>(AssetItem)->Data.ObjectPath : NAME_None))
 			// Icon
 			+SHorizontalBox::Slot()
 			.AutoWidth()
diff --git a/Engine/Source/Editor/ContentBrowser/Private/ContentBrowserSingleton.cpp b/Engine/Source/Editor/ContentBrowser/Private/ContentBrowserSingleton.cpp
--- a/Engine/Source/Editor/ContentBrowser/Private/ContentBrowserSingleton.cpp
+++ b/Engine/Source/Editor/ContentBrowser/Private/ContentBrowserSingleton.cpp
@@ -8,6 +8,7 @@
 #include "IDocumentation.h"
 #include "Editor/MainFrame/Public/MainFrame.h"
 #include "SAssetDialog.h"
+#include "TutorialMetaData.h"
 
 #define LOCTEXT_NAMESPACE "ContentBrowser"
 
@@ -464,9 +465,9 @@ TSharedRef<SDockTab> FContentBrowserSingleton::SpawnContentBrowserTab( const FSp
 	}
 
 	// Add wrapper for tutorial highlighting
-	TSharedRef<SBox> Wrapper = 
-		SNew( SBox )
-		.Tag(TEXT("ContentBrowser"))
+	TSharedRef<SBox> Wrapper =
+		SNew(SBox)
+		.AddMetaData<FTutorialMetaData>(FTutorialMetaData(TEXT("ContentBrowser"), TEXT("ContentBrowserTab1")))
 		[
 			NewBrowser
 		];
diff --git a/Engine/Source/Editor/ContentBrowser/Private/SAssetPicker.cpp b/Engine/Source/Editor/ContentBrowser/Private/SAssetPicker.cpp
--- a/Engine/Source/Editor/ContentBrowser/Private/SAssetPicker.cpp
+++ b/Engine/Source/Editor/ContentBrowser/Private/SAssetPicker.cpp
@@ -82,7 +82,7 @@ void SAssetPicker::Construct( const FArguments& InArgs )
 				.OnGetMenuContent( this, &SAssetPicker::MakeAddFilterMenu )
 				.HasDownArrow( true )
 				.ContentPadding( FMargin( 1, 0 ) )
-				.Tag(TEXT("ContentBrowserFiltersCombo"))
+				.AddMetaData<FTagMetaData>(FTagMetaData(TEXT("ContentBrowserFiltersCombo")))
 				.ButtonContent()
 				[
 					SNew( STextBlock )
diff --git a/Engine/Source/Editor/ContentBrowser/Private/SContentBrowser.cpp b/Engine/Source/Editor/ContentBrowser/Private/SContentBrowser.cpp
--- a/Engine/Source/Editor/ContentBrowser/Private/SContentBrowser.cpp
+++ b/Engine/Source/Editor/ContentBrowser/Private/SContentBrowser.cpp
@@ -104,7 +104,7 @@ void SContentBrowser::Construct( const FArguments& InArgs, const FName& InInstan
 							.OnGetMenuContent( this, &SContentBrowser::MakeCreateAssetContextMenu )
 							.ToolTipText( this, &SContentBrowser::GetNewAssetToolTipText )
 							.IsEnabled( this, &SContentBrowser::IsAssetPathSelected )
-							.Tag(TEXT("ContentBrowserNewAsset"))
+							.AddMetaData<FTagMetaData>(FTagMetaData(TEXT("ContentBrowserNewAsset")))
 							.ButtonContent()
 							[
 								SNew( SHorizontalBox )
@@ -144,7 +144,7 @@ void SContentBrowser::Construct( const FArguments& InArgs, const FName& InInstan
 							.IsEnabled( this, &SContentBrowser::IsAssetPathSelected )
 							.OnClicked( this, &SContentBrowser::HandleImportClicked )
 							.ContentPadding( 0 )
-							.Tag(TEXT("ContentBrowserImportAsset"))
+							.AddMetaData<FTagMetaData>(FTagMetaData(TEXT("ContentBrowserImportAsset")))
 							[
 								SNew( SHorizontalBox )
 
@@ -181,7 +181,7 @@ void SContentBrowser::Construct( const FArguments& InArgs, const FName& InInstan
 							.ToolTipText( LOCTEXT( "SaveDirtyPackagesTooltip", "Save all modified assets." ) )
 							.ContentPadding( 0 )
 							.OnClicked( this, &SContentBrowser::OnSaveClicked )
-							.Tag(TEXT("ContentBrowserSaveDirtyPackages"))
+							.AddMetaData<FTagMetaData>(FTagMetaData(TEXT("ContentBrowserSaveDirtyPackages")))
 							[
 								SNew( SImage )
 								.Image( FEditorStyle::GetBrush( "ContentBrowser.SaveDirtyPackages" ) )
@@ -218,7 +218,7 @@ void SContentBrowser::Construct( const FArguments& InArgs, const FName& InInstan
 							.ContentPadding( FMargin(1, 0) )
 							.OnClicked(this, &SContentBrowser::BackClicked)
 							.IsEnabled(this, &SContentBrowser::IsBackEnabled)
-							.Tag(TEXT("ContentBrowserHistoryBack"))
+							.AddMetaData<FTagMetaData>(FTagMetaData(TEXT("ContentBrowserHistoryBack")))
 							[
 								SNew(SImage)
 								.Image(FEditorStyle::GetBrush("ContentBrowser.HistoryBack"))
@@ -244,7 +244,7 @@ void SContentBrowser::Construct( const FArguments& InArgs, const FName& InInstan
 							.ContentPadding( FMargin(1, 0) )
 							.OnClicked(this, &SContentBrowser::ForwardClicked)
 							.IsEnabled(this, &SContentBrowser::IsForwardEnabled)
-							.Tag(TEXT("ContentBrowserHistoryForward"))
+							.AddMetaData<FTagMetaData>(FTagMetaData(TEXT("ContentBrowserHistoryForward")))
 							[
 								SNew(SImage)
 								.Image(FEditorStyle::GetBrush("ContentBrowser.HistoryForward"))
@@ -272,7 +272,7 @@ void SContentBrowser::Construct( const FArguments& InArgs, const FName& InInstan
 						.ToolTipText( LOCTEXT( "PathPickerTooltip", "Choose a path" ) )
 						.OnGetMenuContent( this, &SContentBrowser::GetPathPickerContent )
 						.HasDownArrow( false )
-						.Tag(TEXT("ContentBrowserPathPicker"))
+						.AddMetaData<FTagMetaData>(FTagMetaData(TEXT("ContentBrowserPathPicker")))
 						.ButtonContent()
 						[
 							SNew( SImage )
@@ -301,7 +301,7 @@ void SContentBrowser::Construct( const FArguments& InArgs, const FName& InInstan
 							.InvertTextColorOnHover( false )
 							.OnCrumbClicked(this, &SContentBrowser::OnPathClicked)
 							.GetCrumbMenuContent(this, &SContentBrowser::OnGetCrumbDelimiterContent)
-							.Tag(TEXT("ContentBrowserPath"))
+							.AddMetaData<FTagMetaData>(FTagMetaData(TEXT("ContentBrowserPath")))
 						]
 					]
 
@@ -321,7 +321,7 @@ void SContentBrowser::Construct( const FArguments& InArgs, const FName& InInstan
 							.ToolTipText( LOCTEXT("LockToggleTooltip", "Toggle lock. If locked, this browser will ignore Find in Content Browser requests.") )
 							.ContentPadding( FMargin(1, 0) )
 							.OnClicked(this, &SContentBrowser::ToggleLockClicked)
-							.Tag(TEXT("ContentBrowserLock"))
+							.AddMetaData<FTagMetaData>(FTagMetaData(TEXT("ContentBrowserLock")))
 							[
 								SNew(SImage)
 								.Image( this, &SContentBrowser::GetToggleLockImage)
@@ -368,11 +368,11 @@ void SContentBrowser::Construct( const FArguments& InArgs, const FName& InInstan
 							.FocusSearchBoxWhenOpened( false )
 							.ShowTreeTitle( false )
 							.ShowSeparator( false )
-							.Tag(TEXT("ContentBrowserSources"))
+							.AddMetaData<FTagMetaData>(FTagMetaData(TEXT("ContentBrowserSources")))
 							.SearchContent()
 							[
 								SNew( SVerticalBox )
-								.Tag(TEXT("ContentBrowserSourcesToggle"))
+								.AddMetaData<FTagMetaData>(FTagMetaData(TEXT("ContentBrowserSourcesToggle")))
 								+ SVerticalBox::Slot()
 								.FillHeight( 1.0f )
 								.Padding(0,0,2,0)
@@ -403,7 +403,7 @@ void SContentBrowser::Construct( const FArguments& InArgs, const FName& InInstan
 						[
 							SAssignNew(CollectionViewPtr, SCollectionView)
 							.OnCollectionSelected(this, &SContentBrowser::CollectionSelected)
-							.Tag(TEXT("ContentBrowserCollections"))
+							.AddMetaData<FTagMetaData>(FTagMetaData(TEXT("ContentBrowserCollections")))
 						]
 					]
 				]
@@ -431,7 +431,7 @@ void SContentBrowser::Construct( const FArguments& InArgs, const FName& InInstan
 						.Padding( 0, 0, 4, 0 )
 						[
 							SNew( SVerticalBox )
-							.Tag(TEXT("ContentBrowserSourcesToggle"))
+							.AddMetaData<FTagMetaData>(FTagMetaData(TEXT("ContentBrowserSourcesToggle")))
 							+ SVerticalBox::Slot()
 							.FillHeight( 1.0f )
 							[
@@ -462,7 +462,7 @@ void SContentBrowser::Construct( const FArguments& InArgs, const FName& InInstan
 							.OnGetMenuContent( this, &SContentBrowser::MakeAddFilterMenu )
 							.HasDownArrow( true )
 							.ContentPadding( FMargin( 1, 0 ) )
-							.Tag(TEXT("ContentBrowserFiltersCombo"))
+							.AddMetaData<FTagMetaData>(FTagMetaData(TEXT("ContentBrowserFiltersCombo")))
 							.ButtonContent()
 							[
 								SNew( STextBlock )
@@ -483,7 +483,7 @@ void SContentBrowser::Construct( const FArguments& InArgs, const FName& InInstan
 							.OnTextCommitted( this, &SContentBrowser::OnSearchBoxCommitted )
 							.PossibleSuggestions( this, &SContentBrowser::GetAssetSearchSuggestions )
 							.DelayChangeNotificationsWhileTyping( true )
-							.Tag(TEXT("ContentBrowserSearchAssets"))
+							.AddMetaData<FTagMetaData>(FTagMetaData(TEXT("ContentBrowserSearchAssets")))
 						]
 					]
 
@@ -495,7 +495,7 @@ void SContentBrowser::Construct( const FArguments& InArgs, const FName& InInstan
 						.OnFilterChanged(this, &SContentBrowser::OnFilterChanged)
 						.OnGetContextMenu(this, &SContentBrowser::GetFilterContextMenu)
 						.FrontendFilters(FrontendFilters)
-						.Tag(TEXT("ContentBrowserFilters"))
+						.AddMetaData<FTagMetaData>(FTagMetaData(TEXT("ContentBrowserFilters")))
 					]
 
 					// Assets
@@ -522,7 +522,7 @@ void SContentBrowser::Construct( const FArguments& InArgs, const FName& InInstan
 						.CanShowOnlyAssetsInSelectedFolders(true)
 						.CanShowRealTimeThumbnails(true)
 						.CanShowDevelopersFolder(true)
-						.Tag(TEXT("ContentBrowserAssets"))
+						.AddMetaData<FTagMetaData>(FTagMetaData(TEXT("ContentBrowserAssets")))
 					]
 				]
 			]
diff --git a/Engine/Source/Editor/GraphEditor/Private/KismetNodes/SGraphNodeK2Base.cpp b/Engine/Source/Editor/GraphEditor/Private/KismetNodes/SGraphNodeK2Base.cpp
--- a/Engine/Source/Editor/GraphEditor/Private/KismetNodes/SGraphNodeK2Base.cpp
+++ b/Engine/Source/Editor/GraphEditor/Private/KismetNodes/SGraphNodeK2Base.cpp
@@ -45,8 +45,8 @@ void SGraphNodeK2Base::UpdateCompactNode()
 		NodeToolTip = IDocumentation::Get()->CreateToolTip( TAttribute< FText >( this, &SGraphNode::GetNodeTooltip ), NULL, GraphNode->GetDocumentationLink(), GraphNode->GetDocumentationExcerptName() );
 	}
 
-	// Setup a tag for this node
-	FBlueprintGraphNodeMetaData TagMeta(TEXT("Graphnode"), TEXT("None"), FGuid());
+	// Setup a meta tag for this node
+	FGraphNodeMetaData TagMeta(TEXT("Graphnode"));
 	PopulateMetaTag(&TagMeta);
 
 	//
@@ -75,7 +75,7 @@ void SGraphNodeK2Base::UpdateCompactNode()
 			// NODE CONTENT AREA
 			SNew(SOverlay)
 			.ToolTip( NodeToolTip.ToSharedRef() )
-			.AddMetaData<FBlueprintGraphNodeMetaData>(TagMeta)
+			.AddMetaData<FGraphNodeMetaData>(TagMeta)
 			+SOverlay::Slot()
 			[
 				SNew(SImage)
@@ -522,47 +522,6 @@ void SGraphNodeK2Base::PerformSecondPassLayout(const TMap< UObject*, TSharedRef<
 	PositionThisNodeBetweenOtherNodes(NodeToWidgetLookup, PrevNodes, NextNodes, Height);
 }
 
-void SGraphNodeK2Base::PopulateMetaTag(FBlueprintGraphNodeMetaData* TagMeta) const
-{
-	if (GraphNode != nullptr)
-	{
-		// We want the name of the blueprint as our name - we can find the node from the GUID
-		UObject* Package = GraphNode->GetOutermost();
-		UObject* LastOuter = GraphNode->GetOuter();
-		while (LastOuter->GetOuter() != Package)
-		{
-			LastOuter = LastOuter->GetOuter();
-		}
-		TagMeta->Tag = FName(*FString::Printf(TEXT("GraphNode_%s_%s"), *LastOuter->GetFullName(), *GraphNode->NodeGuid.ToString()));
-		TagMeta->OuterName = LastOuter->GetFullName();
-		TagMeta->GUID = GraphNode->NodeGuid;
-		FString TitleText = GetNodeCompactTitle().ToString();
-		if (GraphNode->Pins.Num() > 0)
-		{
-			const UEdGraphSchema_K2* K2Schema = GetDefault<UEdGraphSchema_K2>();
-			for (int iPin = 0; iPin < GraphNode->Pins.Num(); iPin++)
-			{
-				if (GraphNode->Pins[iPin]->PinType.PinCategory != K2Schema->PC_Exec)
-				{
-					// Include the title if there is one
-					if (TitleText.IsEmpty())
-					{
-						TagMeta->FriendlyName = FString::Printf(TEXT("%s in %s"), *GraphNode->Pins[iPin]->PinName, *TagMeta->OuterName);
-					}
-					else
-					{
-						TagMeta->FriendlyName = FString::Printf(TEXT("%s %s in %s"), *TitleText, *GraphNode->Pins[iPin]->PinName, *TagMeta->OuterName);
-					}
 
-					break;
-				}
-			}
-		}
-		else
-		{
-			TagMeta->FriendlyName = TitleText;
-		}
-	}
-}
 
 #undef LOCTEXT_NAMESPACE
diff --git a/Engine/Source/Editor/GraphEditor/Private/KismetNodes/SGraphNodeK2Base.h b/Engine/Source/Editor/GraphEditor/Private/KismetNodes/SGraphNodeK2Base.h
--- a/Engine/Source/Editor/GraphEditor/Private/KismetNodes/SGraphNodeK2Base.h
+++ b/Engine/Source/Editor/GraphEditor/Private/KismetNodes/SGraphNodeK2Base.h
@@ -32,9 +32,6 @@ protected :
 	/** Retrieves text to tack on to the top of the tooltip (above the standard text) */
 	FText GetToolTipHeading() const;
 
-	/* Populate a meta data tag with information about this graph node */
-	void PopulateMetaTag(class FBlueprintGraphNodeMetaData* TagMeta) const;
-
 protected:
 	static const FLinearColor BreakpointHitColor;
 	static const FLinearColor LatentBubbleColor;
diff --git a/Engine/Source/Editor/GraphEditor/Private/KismetNodes/SGraphNodeK2Var.cpp b/Engine/Source/Editor/GraphEditor/Private/KismetNodes/SGraphNodeK2Var.cpp
--- a/Engine/Source/Editor/GraphEditor/Private/KismetNodes/SGraphNodeK2Var.cpp
+++ b/Engine/Source/Editor/GraphEditor/Private/KismetNodes/SGraphNodeK2Var.cpp
@@ -160,9 +160,9 @@ void SGraphNodeK2Var::UpdateGraphNode()
 
 	TSharedPtr<SWidget> ErrorText = SetupErrorReporting();
 
-	// Setup a tag for this node
-	FBlueprintGraphNodeMetaData TagMeta(TEXT("Graphnode"), TEXT("None"), FGuid());	
-	PopulateMetaTag(&TagMeta, TitleText);
+	// Setup a meta tag for this node
+	FGraphNodeMetaData TagMeta(TEXT("Graphnode"));
+	PopulateMetaTag(&TagMeta);
 
 	//             ________________
 	//            | (>) L |  R (>) |
@@ -188,7 +188,7 @@ void SGraphNodeK2Var::UpdateGraphNode()
 		+SVerticalBox::Slot()
 		[
 			SNew(SOverlay)
-			.AddMetaData<FBlueprintGraphNodeMetaData>(TagMeta)
+			.AddMetaData<FGraphNodeMetaData>(TagMeta)
 			+ SOverlay::Slot()
 			[
 				SNew(SImage)
@@ -265,45 +265,3 @@ const FSlateBrush* SGraphNodeK2Var::GetShadowBrush(bool bSelected) const
 {
 	return bSelected ? FEditorStyle::GetBrush(TEXT("Graph.VarNode.ShadowSelected")) : FEditorStyle::GetBrush(TEXT("Graph.VarNode.Shadow"));
 }
-
-void SGraphNodeK2Var::PopulateMetaTag(FBlueprintGraphNodeMetaData* TagMeta, const FText &TitleText) const
-{
-	if (GraphNode != nullptr)
-	{
-		// We want the name of the blueprint as our name - we can find the node from the GUID
-		UObject* Package = GraphNode->GetOutermost();
-		UObject* LastOuter = GraphNode->GetOuter();
-		while (LastOuter->GetOuter() != Package)
-		{
-			LastOuter = LastOuter->GetOuter();
-		}
-		TagMeta->Tag = FName(*FString::Printf(TEXT("GraphNode_%s_%s"), *LastOuter->GetFullName(), *GraphNode->NodeGuid.ToString()));
-		TagMeta->OuterName = LastOuter->GetFullName();
-		TagMeta->GUID = GraphNode->NodeGuid;
-		if (GraphNode->Pins.Num() > 0)
-		{
-			const UEdGraphSchema_K2* K2Schema = GetDefault<UEdGraphSchema_K2>();
-			for (int iPin = 0; iPin < GraphNode->Pins.Num(); iPin++)
-			{
-				if (GraphNode->Pins[iPin]->PinType.PinCategory != K2Schema->PC_Exec)
-				{
-					// Include the title if there is one
-					if (TitleText.IsEmpty())
-					{
-						TagMeta->FriendlyName = FString::Printf(TEXT("%s in %s"), *GraphNode->Pins[iPin]->PinName, *TagMeta->OuterName);
-					}
-					else
-					{
-						TagMeta->FriendlyName = FString::Printf(TEXT("%s %s in %s"), *TitleText.ToString(), *GraphNode->Pins[iPin]->PinName, *TagMeta->OuterName);
-					}
-
-					break;
-				}
-			}
-		}
-		else
-		{
-			TagMeta->FriendlyName = TitleText.ToString();
-		}
-	}
-}
diff --git a/Engine/Source/Editor/GraphEditor/Private/KismetNodes/SGraphNodeK2Var.h b/Engine/Source/Editor/GraphEditor/Private/KismetNodes/SGraphNodeK2Var.h
--- a/Engine/Source/Editor/GraphEditor/Private/KismetNodes/SGraphNodeK2Var.h
+++ b/Engine/Source/Editor/GraphEditor/Private/KismetNodes/SGraphNodeK2Var.h
@@ -18,9 +18,6 @@ public:
 	virtual void UpdateGraphNode() override;
 	// End of SGraphNode interface
 
-	/* Populates a meta tag with information about the graph node */
-	void PopulateMetaTag(class FBlueprintGraphNodeMetaData* TagMeta, const FText &TitleText) const;
-
 protected:
 	FSlateColor GetVariableColor() const;
 };
diff --git a/Engine/Source/Editor/GraphEditor/Private/SGraphNode.cpp b/Engine/Source/Editor/GraphEditor/Private/SGraphNode.cpp
--- a/Engine/Source/Editor/GraphEditor/Private/SGraphNode.cpp
+++ b/Engine/Source/Editor/GraphEditor/Private/SGraphNode.cpp
@@ -9,6 +9,7 @@
 #include "Editor/UnrealEd/Public/Kismet2/BlueprintEditorUtils.h"
 #include "SLevelOfDetailBranchNode.h"
 #include "IDocumentation.h"
+#include "TutorialMetaData.h"
 
 /////////////////////////////////////////////////////
 // SNodeTitle
@@ -700,19 +701,9 @@ void SGraphNode::UpdateGraphNode()
 		SetToolTip(DefaultToolTip);
 	}
 
-	// Setup a tag for this node
- 	FString TagName;
-	if (GraphNode != nullptr)
- 	{	
-		// We want the name of the blueprint as our name - we can find the node from the GUID
-		UObject* Package = GraphNode->GetOutermost();
-		UObject* LastOuter = GraphNode->GetOuter();
-		while (LastOuter->GetOuter() != Package)
-		{
-			LastOuter = LastOuter->GetOuter();
-		}
-		TagName = FString::Printf(TEXT("GraphNode,%s,%s"), *LastOuter->GetFullName(), *GraphNode->NodeGuid.ToString());
- 	}
+	// Setup a meta tag for this node
+	FGraphNodeMetaData TagMeta(TEXT("Graphnode"));
+	PopulateMetaTag(&TagMeta);
 	
 	TSharedPtr<SVerticalBox> InnerVerticalBox;
 	this->ContentScale.Bind( this, &SGraphNode::GetContentScale );
@@ -727,7 +718,7 @@ void SGraphNode::UpdateGraphNode()
 				SNew(SBorder)
 				.BorderImage( FEditorStyle::GetBrush( "Graph.Node.Body" ) )
 				.Padding(0)
-				.Tag(*TagName)
+				.AddMetaData<FGraphNodeMetaData>(TagMeta)
 				[
 					SAssignNew(InnerVerticalBox, SVerticalBox)
 					+SVerticalBox::Slot()
@@ -1290,3 +1281,21 @@ EVisibility SGraphNode::IsAddPinButtonVisible() const
 
 	return bIsHidden ? EVisibility::Collapsed : EVisibility::Visible;
 }
+
+void SGraphNode::PopulateMetaTag(FGraphNodeMetaData* TagMeta) const
+{
+	if (GraphNode != nullptr)
+	{
+		// We want the name of the blueprint as our name - we can find the node from the GUID
+		UObject* Package = GraphNode->GetOutermost();
+		UObject* LastOuter = GraphNode->GetOuter();
+		while (LastOuter->GetOuter() != Package)
+		{
+			LastOuter = LastOuter->GetOuter();
+		}
+		TagMeta->Tag = FName(*FString::Printf(TEXT("GraphNode_%s_%s"), *LastOuter->GetFullName(), *GraphNode->NodeGuid.ToString()));
+		TagMeta->OuterName = LastOuter->GetFullName();
+		TagMeta->GUID = GraphNode->NodeGuid;
+		TagMeta->FriendlyName = FString::Printf(TEXT("%s in %s"), *GraphNode->GetNodeTitle(ENodeTitleType::ListView).ToString(), *TagMeta->OuterName);		
+	}
+}
\ No newline at end of file
diff --git a/Engine/Source/Editor/GraphEditor/Private/SGraphNodeComment.cpp b/Engine/Source/Editor/GraphEditor/Private/SGraphNodeComment.cpp
--- a/Engine/Source/Editor/GraphEditor/Private/SGraphNodeComment.cpp
+++ b/Engine/Source/Editor/GraphEditor/Private/SGraphNodeComment.cpp
@@ -5,6 +5,7 @@
 #include "SGraphNodeComment.h"
 //#include "TextWrapperHelpers.h"
 #include "BlueprintEditorUtils.h"
+#include "TutorialMetaData.h"
 
 namespace SCommentNodeDefs
 {
@@ -114,6 +115,10 @@ void SGraphNodeComment::UpdateGraphNode()
 
 	TSharedPtr<SWidget> ErrorText = SetupErrorReporting();
 
+	// Setup a meta tag for this node
+	FGraphNodeMetaData TagMeta(TEXT("Graphnode"));
+	PopulateMetaTag(&TagMeta);
+
 	bool bIsSet = GraphNode->IsA(UEdGraphNode_Comment::StaticClass());
 	this->ContentScale.Bind( this, &SGraphNode::GetContentScale );
 	this->ChildSlot
@@ -125,7 +130,7 @@ void SGraphNodeComment::UpdateGraphNode()
 			.ColorAndOpacity( FLinearColor::White )
 			.BorderBackgroundColor( this, &SGraphNodeComment::GetCommentBodyColor )
 			.Padding(  FMargin(3.0f) )
-			.Tag(*TagName)
+			.AddMetaData<FGraphNodeMetaData>(TagMeta)
 			[
 				SNew(SVerticalBox)
 				.ToolTipText( this, &SGraphNode::GetNodeTooltip )
@@ -367,3 +372,21 @@ FSlateRect SGraphNodeComment::GetTitleRect() const
 	FVector2D NodeSize  = TitleBar.IsValid() ? TitleBar->GetDesiredSize() : GetDesiredSize();
 	return FSlateRect( NodePosition.X, NodePosition.Y, NodePosition.X + NodeSize.X, NodePosition.Y + NodeSize.Y ) + SCommentNodeDefs::TitleBarOffset;
 }
+
+void SGraphNodeComment::PopulateMetaTag(FGraphNodeMetaData* TagMeta) const
+{
+	if (GraphNode != nullptr)
+	{
+		// We want the name of the blueprint as our name - we can find the node from the GUID
+		UObject* Package = GraphNode->GetOutermost();
+		UObject* LastOuter = GraphNode->GetOuter();
+		while (LastOuter->GetOuter() != Package)
+		{
+			LastOuter = LastOuter->GetOuter();
+		}
+		TagMeta->Tag = FName(*FString::Printf(TEXT("GraphNode_%s_%s"), *LastOuter->GetFullName(), *GraphNode->NodeGuid.ToString()));
+		TagMeta->OuterName = LastOuter->GetFullName();
+		TagMeta->GUID = GraphNode->NodeGuid;
+		TagMeta->FriendlyName = FString::Printf(TEXT("%s in %s"), *GraphNode->GetNodeTitle(ENodeTitleType::FullTitle).ToString(), *TagMeta->OuterName);
+	}
+}
diff --git a/Engine/Source/Editor/GraphEditor/Private/SGraphNodeDocumentation.cpp b/Engine/Source/Editor/GraphEditor/Private/SGraphNodeDocumentation.cpp
--- a/Engine/Source/Editor/GraphEditor/Private/SGraphNodeDocumentation.cpp
+++ b/Engine/Source/Editor/GraphEditor/Private/SGraphNodeDocumentation.cpp
@@ -5,6 +5,7 @@
 #include "SGraphNodeDocumentation.h"
 #include "IDocumentation.h"
 #include "SLevelOfDetailBranchNode.h"
+#include "TutorialMetaData.h"
 
 #define LOCTEXT_NAMESPACE "SGraphNodeDocumentation"
 
@@ -62,8 +63,13 @@ void SGraphNodeDocumentation::UpdateGraphNode()
 	// Create Node Title
 	TSharedPtr<SNodeTitle> NodeTitle = SNew( SNodeTitle, GraphNode );
 
+	// Setup a meta tag for this node
+	FGraphNodeMetaData TagMeta(TEXT("Graphnode"));
+	PopulateMetaTag(&TagMeta);
+
 	TSharedRef<SOverlay> DefaultTitleAreaWidget =
 		SNew( SOverlay )
+		.AddMetaData<FGraphNodeMetaData>(TagMeta)
 		+SOverlay::Slot()
 		[
 			SNew( SImage )
diff --git a/Engine/Source/Editor/GraphEditor/Public/SGraphNode.h b/Engine/Source/Editor/GraphEditor/Public/SGraphNode.h
--- a/Engine/Source/Editor/GraphEditor/Public/SGraphNode.h
+++ b/Engine/Source/Editor/GraphEditor/Public/SGraphNode.h
@@ -299,6 +299,9 @@ protected:
 	/** Callback function executed when Add pin button is clicked */
 	virtual FReply OnAddPin() {return FReply::Handled();}
 
+	/* Populate a meta data tag with information about this graph node */
+	virtual void PopulateMetaTag(class FGraphNodeMetaData* TagMeta) const;
+
 protected:
 	/** Input pin widgets on this node */
 	TArray< TSharedRef<SGraphPin> > InputPins;
diff --git a/Engine/Source/Editor/GraphEditor/Public/SGraphNodeComment.h b/Engine/Source/Editor/GraphEditor/Public/SGraphNodeComment.h
--- a/Engine/Source/Editor/GraphEditor/Public/SGraphNodeComment.h
+++ b/Engine/Source/Editor/GraphEditor/Public/SGraphNodeComment.h
@@ -44,6 +44,7 @@ protected:
 	// SGraphNode Interface
 	virtual void UpdateGraphNode() override;
 	virtual bool ShouldScaleNodeComment()const override;
+	virtual void PopulateMetaTag(class FGraphNodeMetaData* TagMeta) const override;
 
 	/**
 	 * Helper method to update selection state of comment and any nodes 'contained' within it
diff --git a/Engine/Source/Editor/IntroTutorials/Private/STutorialButton.cpp b/Engine/Source/Editor/IntroTutorials/Private/STutorialButton.cpp
--- a/Engine/Source/Editor/IntroTutorials/Private/STutorialButton.cpp
+++ b/Engine/Source/Editor/IntroTutorials/Private/STutorialButton.cpp
@@ -4,6 +4,7 @@
 #include "STutorialButton.h"
 #include "EditorTutorialSettings.h"
 #include "TutorialStateSettings.h"
+#include "TutorialMetaData.h"
 
 #define LOCTEXT_NAMESPACE "STutorialButton"
 
@@ -26,10 +27,11 @@ void STutorialButton::Construct(const FArguments& InArgs)
 	PulseAnimation.AddCurve(0.0f, TutorialButtonConstants::PulseAnimationLength, ECurveEaseFunction::Linear);
 	PulseAnimation.Play();
 
+	FTutorialMetaData TagMeta(*FString::Printf(TEXT("%s.TutorialLaunchButton"), *Context.ToString()), TEXT("TutorialLaunchButton"));
 	ChildSlot
 	[
 		SNew(SButton)
-		.Tag(*FString::Printf(TEXT("%s.TutorialLaunchButton"), *Context.ToString()))
+		.AddMetaData<FTutorialMetaData>(TagMeta)
 		.ButtonStyle(FEditorStyle::Get(), "TutorialLaunch.Button")
 		.ToolTipText(LOCTEXT("TutorialLaunchToolTip", "Launch Tutorial (right-click for more options)"))
 		.Visibility(this, &STutorialButton::GetVisibility)
diff --git a/Engine/Source/Editor/IntroTutorials/Private/STutorialOverlay.cpp b/Engine/Source/Editor/IntroTutorials/Private/STutorialOverlay.cpp
--- a/Engine/Source/Editor/IntroTutorials/Private/STutorialOverlay.cpp
+++ b/Engine/Source/Editor/IntroTutorials/Private/STutorialOverlay.cpp
@@ -15,10 +15,7 @@ void STutorialOverlay::Construct(const FArguments& InArgs, FTutorialStage* const
 	ParentWindow = InArgs._ParentWindow;
 	bIsStandalone = InArgs._IsStandalone;
 	OnClosed = InArgs._OnClosed;
-
-	// Setup the map for opening of closed tabs by highlighted widgets
-	AddTabInfo();
-
+	
 	TSharedPtr<SOverlay> Overlay;
 
 	ChildSlot
@@ -151,24 +148,12 @@ void STutorialOverlay::PerformWidgetInteractions(const FTutorialWidgetContent &W
 
 void STutorialOverlay::OpenBrowserForWidgetAnchor(const FTutorialWidgetContent &WidgetContent)
 {
-	FString IdentString = WidgetContent.WidgetAnchor.WrapperIdentifier.ToString();
-	FString TabString = "";
-	
-	// See if we have a mapping for the widget ident.We will likely pull this from meta data when we have it)
-	for (TMap<FString, FString>::TConstIterator It(BrowserTabMap); It; ++It)
-	{
-		if (It.Key().StartsWith(IdentString) == true)
-		{
-			TabString = It.Value();
-		}
-	}
-
 	// Open the required tab if we found it in the map
-	if (TabString.IsEmpty() == false)
+	if (WidgetContent.WidgetAnchor.TabTypeToOpen.IsEmpty() == false)
 	{
 		FLevelEditorModule& LevelEditorModule = FModuleManager::GetModuleChecked<FLevelEditorModule>(TEXT("LevelEditor"));
 		TSharedPtr<FTabManager> LevelEditorTabManager = LevelEditorModule.GetLevelEditorTabManager();
-		LevelEditorTabManager->InvokeTab(FName(*TabString));
+		LevelEditorTabManager->InvokeTab(FName(*WidgetContent.WidgetAnchor.TabTypeToOpen));
 	}
 	
 }
@@ -198,14 +183,3 @@ void STutorialOverlay::FocusOnAnyBlueprintNodes(const FTutorialWidgetContent &Wi
 	}
 }
 
-void STutorialOverlay::AddTabInfo()
-{
-	BrowserTabMap.Empty();
-	BrowserTabMap.Add(FString(TEXT("ActorDetails")), FString(TEXT("LevelEditorSelectionDetails")));
-	BrowserTabMap.Add(FString(TEXT("SceneOutliner")), FString(TEXT("LevelEditorSceneOutliner")));
-	BrowserTabMap.Add(FString(TEXT("ContentBrowser")), FString(TEXT("ContentBrowserTab1")));
-	BrowserTabMap.Add(FString(TEXT("ToolsPanel")), FString(TEXT("LevelEditorToolBox")));
-	BrowserTabMap.Add(FString(TEXT("WorldSettings")), FString(TEXT("WorldSettingsTab")));
-	BrowserTabMap.Add(FString(TEXT("EditorViewports")), FString(TEXT("LevelEditorViewport")));	 
-	BrowserTabMap.Add(FString(TEXT("LayerBrowser")), FString(TEXT("LevelEditorLayerBrowser")));	
-}
diff --git a/Engine/Source/Editor/IntroTutorials/Private/TutorialStructCustomization.cpp b/Engine/Source/Editor/IntroTutorials/Private/TutorialStructCustomization.cpp
--- a/Engine/Source/Editor/IntroTutorials/Private/TutorialStructCustomization.cpp
+++ b/Engine/Source/Editor/IntroTutorials/Private/TutorialStructCustomization.cpp
@@ -254,25 +254,35 @@ private:
 			TSharedPtr<IPropertyHandle> TypeProperty = StructPropertyHandle->GetChildHandle(GET_MEMBER_NAME_CHECKED(FTutorialContentAnchor, Type));
 			TypeProperty->SetValue((uint8)ETutorialAnchorIdentifier::NamedWidget);
 			
+			// Set the friendly name to the PickedWidget name - we might not have any metadata
+			StructPropertyHandle->GetChildHandle(GET_MEMBER_NAME_CHECKED(FTutorialContentAnchor, FriendlyName))->SetValue(PickedWidgetName);
+			//StructPropertyHandle->GetChildHandle(GET_MEMBER_NAME_CHECKED(FTutorialContentAnchor, FriendlyName))->SetValue(FString());
+
 			for (const auto& MetaDataEntry : PickedAllMetaData)
 			{
-				if (MetaDataEntry->IsOfType<FBlueprintGraphNodeMetaData>())
+				if (MetaDataEntry->IsOfType<FTutorialMetaData>())
+				{
+					TSharedRef<FTutorialMetaData> GraphNodeMeta = StaticCastSharedRef<FGraphNodeMetaData>(MetaDataEntry);
+					StructPropertyHandle->GetChildHandle(GET_MEMBER_NAME_CHECKED(FTutorialContentAnchor, FriendlyName))->SetValue(GraphNodeMeta->FriendlyName);
+					StructPropertyHandle->GetChildHandle(GET_MEMBER_NAME_CHECKED(FTutorialContentAnchor, TabTypeToOpen))->SetValue(GraphNodeMeta->TabTypeToOpen);
+				}
+				else if (MetaDataEntry->IsOfType<FGraphNodeMetaData>())
 				{
-					TSharedRef<FBlueprintGraphNodeMetaData> GraphNodeMeta = StaticCastSharedRef<FBlueprintGraphNodeMetaData>(MetaDataEntry);
-					StructPropertyHandle->GetChildHandle(GET_MEMBER_NAME_CHECKED(FTutorialContentAnchor, TagIdent))->SetValue(GraphNodeMeta->Tag);
+					TSharedRef<FGraphNodeMetaData> GraphNodeMeta = StaticCastSharedRef<FGraphNodeMetaData>(MetaDataEntry);
 					StructPropertyHandle->GetChildHandle(GET_MEMBER_NAME_CHECKED(FTutorialContentAnchor, GUIDString))->SetValue(GraphNodeMeta->GUID.ToString());
 					StructPropertyHandle->GetChildHandle(GET_MEMBER_NAME_CHECKED(FTutorialContentAnchor, OuterName))->SetValue(GraphNodeMeta->OuterName);
 					StructPropertyHandle->GetChildHandle(GET_MEMBER_NAME_CHECKED(FTutorialContentAnchor, FriendlyName))->SetValue(GraphNodeMeta->FriendlyName);
+					StructPropertyHandle->GetChildHandle(GET_MEMBER_NAME_CHECKED(FTutorialContentAnchor, TabTypeToOpen))->SetValue(GraphNodeMeta->TabTypeToOpen);
 				}
 				else
 				{
 					TSharedRef<FTagMetaData> GraphNodeMeta = StaticCastSharedRef<FTagMetaData>(MetaDataEntry);
-					TSharedPtr<IPropertyHandle> MetaDataPropertyIdent = StructPropertyHandle->GetChildHandle(GET_MEMBER_NAME_CHECKED(FTutorialContentAnchor, TagIdent));
-					MetaDataPropertyIdent->SetValue(GraphNodeMeta->Tag);
+					StructPropertyHandle->GetChildHandle(GET_MEMBER_NAME_CHECKED(FTutorialContentAnchor, FriendlyName))->SetValue(GraphNodeMeta->Tag);
 				}
 			}
 			
 			PickedWidgetName = NAME_None;
+			PickedAllMetaData.Reset();
 
 			if(ParentWindow.IsValid())
 			{
diff --git a/Engine/Source/Editor/IntroTutorials/Public/EditorTutorial.h b/Engine/Source/Editor/IntroTutorials/Public/EditorTutorial.h
--- a/Engine/Source/Editor/IntroTutorials/Public/EditorTutorial.h
+++ b/Engine/Source/Editor/IntroTutorials/Public/EditorTutorial.h
@@ -118,16 +118,23 @@ struct INTROTUTORIALS_API FTutorialContentAnchor
 	/** Whether to draw an animated highlight around the widget */
 	UPROPERTY(EditAnywhere, Category="Anchor")
 	bool bDrawHighlight;
+	
+	/* User friendly name to display in the dialog */
+	UPROPERTY(EditAnywhere, Category = "AnchorMeta|FBlueprintGraphNodeMetaData")
+	FString FriendlyName;
 
-	//FBlueprintGraphNodeMetaData Data
+	/* The type of tab to open for this node*/
 	UPROPERTY(EditAnywhere, Category = "AnchorMeta|FBlueprintGraphNodeMetaData")
-	FString TagIdent;
+	FString TabTypeToOpen;
+
+	//FBlueprintGraphNodeMetaData specific members
+	/* The GUID string */
 	UPROPERTY(EditAnywhere, Category = "AnchorMeta|FBlueprintGraphNodeMetaData")
  	FString GUIDString;
+	
+	/* Name of the outer object - should be the blueprint that 'owns' the node */
 	UPROPERTY(EditAnywhere, Category = "AnchorMeta|FBlueprintGraphNodeMetaData")
 	FString OuterName;
-	UPROPERTY(EditAnywhere, Category = "AnchorMeta|FBlueprintGraphNodeMetaData")
-	FString FriendlyName;
 };
 
 /** Content that is displayed relative to a widget */
diff --git a/Engine/Source/Editor/IntroTutorials/Public/TutorialMetaData.h b/Engine/Source/Editor/IntroTutorials/Public/TutorialMetaData.h
--- a/Engine/Source/Editor/IntroTutorials/Public/TutorialMetaData.h
+++ b/Engine/Source/Editor/IntroTutorials/Public/TutorialMetaData.h
@@ -4,25 +4,44 @@
 
 #include "ISlateMetaData.h"
 
-class INTROTUTORIALS_API FBlueprintGraphNodeMetaData : public FTagMetaData
+class INTROTUTORIALS_API FTutorialMetaData : public FTagMetaData
 {
 public:
-	SLATE_METADATA_TYPE(FBlueprintGraphNodeMetaData, FTagMetaData)
+	SLATE_METADATA_TYPE(FTutorialMetaData, FTagMetaData)
 
-		FBlueprintGraphNodeMetaData(FName InTag, FString InOuterName, FGuid InGUID)
+		FTutorialMetaData(FName InTag, FString InTabType, FString InFriendlyName = FString())
 		: FTagMetaData(InTag)
-		, GUID(InGUID)
+		, TabTypeToOpen(InTabType)
+		, FriendlyName(InFriendlyName)
+		{
+		}
+	
+	FTutorialMetaData(FName InTag)
+		: FTagMetaData(InTag)
+		{
+		}
+	/** User friendly display name */
+	FString FriendlyName;
+
+	/* The type of tab we want to open for this widget where relevant. */
+	FString TabTypeToOpen;
+};
+
+class INTROTUTORIALS_API FGraphNodeMetaData : public FTutorialMetaData
+{
+public:
+	SLATE_METADATA_TYPE(FGraphNodeMetaData, FTutorialMetaData)
+
+		FGraphNodeMetaData(FName InTag, FString InFriendlyName = FString(), FString InOuterName = FString(), FString InTabType = FString())
+		: FTutorialMetaData(InTag, InTabType, InFriendlyName)
 		, OuterName(InOuterName)
 		{
 		}
 
-	
 	/** GUID for the node */
 	FGuid GUID;
 
 	/** Name of the outer (which should be the blueprint) */
 	FString OuterName;
 
-	/** User friendly display name */
-	FString FriendlyName;
 };
\ No newline at end of file
diff --git a/Engine/Source/Editor/Kismet/Private/SBlueprintEditorToolbar.cpp b/Engine/Source/Editor/Kismet/Private/SBlueprintEditorToolbar.cpp
--- a/Engine/Source/Editor/Kismet/Private/SBlueprintEditorToolbar.cpp
+++ b/Engine/Source/Editor/Kismet/Private/SBlueprintEditorToolbar.cpp
@@ -827,7 +827,7 @@ TArray< TSharedPtr< SWidget> > FBlueprintEditorToolbar::GenerateToolbarWidgets(c
 				TEXT("DefaultsMode")))
 			.IconImage(FEditorStyle::GetBrush("FullBlueprintEditor.SwitchToBlueprintDefaultsMode"))
 			.SmallIconImage(FEditorStyle::GetBrush("FullBlueprintEditor.SwitchToBlueprintDefaultsMode.Small"))
-			.Tag(TEXT("DefaultsMode"))
+			.AddMetaData<FTagMetaData>(FTagMetaData(TEXT("DefaultsMode")))
 		);
 
 		Ret.Add(SNew(SBlueprintModeSeparator));
@@ -844,7 +844,7 @@ TArray< TSharedPtr< SWidget> > FBlueprintEditorToolbar::GenerateToolbarWidgets(c
 				TEXT("ComponentsMode")))
 			.IconImage(FEditorStyle::GetBrush("FullBlueprintEditor.SwitchToComponentsMode"))
 			.SmallIconImage(FEditorStyle::GetBrush("FullBlueprintEditor.SwitchToComponentsMode.Small"))
-			.Tag(TEXT("ComponentsMode"))
+			.AddMetaData<FTagMetaData>(FTagMetaData(TEXT("ComponentsMode")))
 		);
 
 		Ret.Add(SNew(SBlueprintModeSeparator));
@@ -862,7 +862,7 @@ TArray< TSharedPtr< SWidget> > FBlueprintEditorToolbar::GenerateToolbarWidgets(c
 			.ToolTipText(LOCTEXT("GraphModeButtonTooltip", "Switch to Graph Editing Mode"))
 			.IconImage(FEditorStyle::GetBrush("FullBlueprintEditor.SwitchToScriptingMode"))
 			.SmallIconImage(FEditorStyle::GetBrush("FullBlueprintEditor.SwitchToScriptingMode.Small"))
-			.Tag(TEXT("GraphMode"))
+			.AddMetaData<FTagMetaData>(FTagMetaData(TEXT("GraphMode")))
 		);
 
 		// Right side padding
diff --git a/Engine/Source/Editor/Kismet/Private/SBlueprintPalette.cpp b/Engine/Source/Editor/Kismet/Private/SBlueprintPalette.cpp
--- a/Engine/Source/Editor/Kismet/Private/SBlueprintPalette.cpp
+++ b/Engine/Source/Editor/Kismet/Private/SBlueprintPalette.cpp
@@ -29,6 +29,7 @@
 #include "BlueprintActionMenuItem.h"
 #include "BlueprintDragDropMenuItem.h"
 #include "BlueprintNodeSpawner.h"
+#include "TutorialMetaData.h"
 
 #define LOCTEXT_NAMESPACE "BlueprintPalette"
 
@@ -879,11 +880,12 @@ void SBlueprintPaletteItem::Construct(const FArguments& InArgs, FCreateWidgetFor
 	GetPaletteItemIcon(GraphAction, InBlueprintEditor.Pin()->GetBlueprintObj(), IconBrush, IconColor, IconToolTip, IconDocLink, IconDocExcerpt);
 	TSharedRef<SWidget> IconWidget = CreateIconWidget(IconToolTip, IconBrush, IconColor, IconDocLink, IconDocExcerpt);
 
-	// Setup a tag for this node
-	FString TagName;
+	// Setup a meta tag for this node
+	FTutorialMetaData TagMeta("PaletteItem"); 
 	if( ActionPtr.IsValid() )
 	{
-		TagName = FString::Printf(TEXT("PaletteItem,%s,%d"), *GraphAction->MenuDescription.ToString(), GraphAction->SectionID);
+		TagMeta.Tag = *FString::Printf(TEXT("PaletteItem,%s,%d"), *GraphAction->MenuDescription.ToString(), GraphAction->SectionID);
+		TagMeta.FriendlyName = GraphAction->MenuDescription.ToString();
 	}
 	// construct the text widget
 	FSlateFontInfo NameFont = FSlateFontInfo(FPaths::EngineContentDir() / TEXT("Slate/Fonts/Roboto-Regular.ttf"), 10);
@@ -894,7 +896,7 @@ void SBlueprintPaletteItem::Construct(const FArguments& InArgs, FCreateWidgetFor
 	ChildSlot
 	[
 		SNew(SHorizontalBox)
-		.Tag(*TagName)
+		.AddMetaData<FTutorialMetaData>(TagMeta)
 		// icon slot
 		+SHorizontalBox::Slot()
 			.AutoWidth()
@@ -1534,20 +1536,20 @@ void SBlueprintPalette::Construct(const FArguments& InArgs, TWeakPtr<FBlueprintE
 			SAssignNew(PaletteSplitter, SSplitter)
 				.Orientation(Orient_Vertical)
 				.OnSplitterFinishedResizing(this, &SBlueprintPalette::OnSplitterResized)
-				.Tag(TEXT("FullBlueprintPalette"))
+				.AddMetaData<FTagMetaData>(FTagMetaData(TEXT("FullBlueprintPalette")))
 
 			+ SSplitter::Slot()
 			.Value(FavoritesHeightRatio)
 			[
 				SNew(SBlueprintFavoritesPalette, InBlueprintEditor)
-				.Tag(TEXT("BlueprintPaletteFavorites"))
+				.AddMetaData<FTagMetaData>(FTagMetaData(TEXT("BlueprintPaletteFavorites")))
 			]
 
 			+ SSplitter::Slot()
 			.Value(LibraryHeightRatio)
 			[
 				SNew(SBlueprintLibraryPalette, InBlueprintEditor)
-				.Tag(TEXT("BlueprintPaletteLibrary"))
+				.AddMetaData<FTagMetaData>(FTagMetaData(TEXT("BlueprintPaletteLibrary")))
 			]
 		];
 	}	
diff --git a/Engine/Source/Editor/Kismet/Private/SGraphTitleBar.cpp b/Engine/Source/Editor/Kismet/Private/SGraphTitleBar.cpp
--- a/Engine/Source/Editor/Kismet/Private/SGraphTitleBar.cpp
+++ b/Engine/Source/Editor/Kismet/Private/SGraphTitleBar.cpp
@@ -86,7 +86,7 @@ void SGraphTitleBar::Construct( const FArguments& InArgs )
 		SNew(SBorder)
 		.BorderImage( FEditorStyle::GetBrush( TEXT("Graph.TitleBackground") ) )
 		.HAlign(HAlign_Fill)
-		.Tag(TEXT("EventGraphTitleBar"))
+		.AddMetaData<FTagMetaData>(FTagMetaData(TEXT("EventGraphTitleBar")))
 		[
 			SNew(SVerticalBox)
 			// Title text/icon
diff --git a/Engine/Source/Editor/Kismet/Private/SKismetInspector.cpp b/Engine/Source/Editor/Kismet/Private/SKismetInspector.cpp
--- a/Engine/Source/Editor/Kismet/Private/SKismetInspector.cpp
+++ b/Engine/Source/Editor/Kismet/Private/SKismetInspector.cpp
@@ -244,7 +244,7 @@ void SKismetInspector::Construct(const FArguments& InArgs)
 	ChildSlot
 	[
 		SNew(SVerticalBox)
-		.Tag( TEXT("BlueprintInspector"))
+		.AddMetaData<FTagMetaData>(FTagMetaData(TEXT("BlueprintInspector")))
 		+ SVerticalBox::Slot()
 		.FillHeight(1.0f)
 		[
diff --git a/Engine/Source/Editor/Kismet/Private/SMyBlueprint.cpp b/Engine/Source/Editor/Kismet/Private/SMyBlueprint.cpp
--- a/Engine/Source/Editor/Kismet/Private/SMyBlueprint.cpp
+++ b/Engine/Source/Editor/Kismet/Private/SMyBlueprint.cpp
@@ -206,7 +206,7 @@ void SMyBlueprint::Construct(const FArguments& InArgs, TWeakPtr<FBlueprintEditor
 		SNew(SBorder)
 		.Padding(4.0f)
 		.BorderImage(FEditorStyle::GetBrush("ToolPanel.GroupBorder"))
-		.Tag(TEXT("MyBlueprintPanel"))
+		.AddMetaData<FTagMetaData>(FTagMetaData(TEXT("MyBlueprintPanel")))
 		[
 			SNew(SVerticalBox)
 
diff --git a/Engine/Source/Editor/Kismet/Private/SSCSEditor.cpp b/Engine/Source/Editor/Kismet/Private/SSCSEditor.cpp
--- a/Engine/Source/Editor/Kismet/Private/SSCSEditor.cpp
+++ b/Engine/Source/Editor/Kismet/Private/SSCSEditor.cpp
@@ -25,6 +25,7 @@
 #include "IDocumentation.h"
 #include "Kismet2NameValidators.h"
 #include "UnrealExporter.h"
+#include "TutorialMetaData.h"
 
 #define LOCTEXT_NAMESPACE "SSCSEditor"
 
@@ -2010,7 +2011,7 @@ void SSCSEditor::Construct( const FArguments& InArgs, TSharedPtr<FBlueprintEdito
 	this->ChildSlot
 	[
 		SNew( SBox )
-		.Tag(TEXT("ComponentsPanel"))
+		.AddMetaData<FTagMetaData>(FTagMetaData(TEXT("ComponentsPanel")))
 		[
 			Contents.ToSharedRef()
 		]
@@ -2045,15 +2046,14 @@ TSharedRef<ITableRow> SSCSEditor::MakeTableRowWidget( FSCSEditorTreeNodePtrType
 		OnRenameComponent(false);
 	}
 
-	
-	// Setup a tag for this node
-	FString TagName;
+	// Setup a meta tag for this node
+	FGraphNodeMetaData TagMeta(TEXT("TableRow"));
 	if (InNodePtr.IsValid() && InNodePtr->GetComponentTemplate() != NULL )
 	{
-		TagName = FString::Printf(TEXT("TableRow,%s,0"), *InNodePtr->GetComponentTemplate()->GetReadableName());		
+		TagMeta.FriendlyName = FString::Printf(TEXT("TableRow,%s,0"), *InNodePtr->GetComponentTemplate()->GetReadableName());
 	}
 	return SNew(SSCS_RowWidget, SharedThis(this), InNodePtr, OwnerTable)
-		.Tag(*TagName);
+		.AddMetaData<FTutorialMetaData>(TagMeta);
 }
 
 void SSCSEditor::GetSelectedItemsForContextMenu(TArray<FComponentEventConstructionData>& OutSelectedItems) const
diff --git a/Engine/Source/Editor/LevelEditor/LevelEditor.Build.cs b/Engine/Source/Editor/LevelEditor/LevelEditor.Build.cs
--- a/Engine/Source/Editor/LevelEditor/LevelEditor.Build.cs
+++ b/Engine/Source/Editor/LevelEditor/LevelEditor.Build.cs
@@ -21,6 +21,7 @@ public class LevelEditor : ModuleRules
 			new string[] {
 				"Settings",
 				"UserFeedback",
+				"IntroTutorials"
 			}
 		);
 
diff --git a/Engine/Source/Editor/LevelEditor/Private/SLevelEditor.cpp b/Engine/Source/Editor/LevelEditor/Private/SLevelEditor.cpp
--- a/Engine/Source/Editor/LevelEditor/Private/SLevelEditor.cpp
+++ b/Engine/Source/Editor/LevelEditor/Private/SLevelEditor.cpp
@@ -36,6 +36,7 @@
 #include "EditorModes.h"
 #include "IDocumentation.h"
 #include "NewsFeed.h"
+#include "TutorialMetaData.h"
 
 
 static const FName LevelEditorBuildAndSubmitTab("LevelEditorBuildAndSubmit");
@@ -187,7 +188,7 @@ void SLevelEditor::Initialize( const TSharedRef<SDockTab>& OwnerTab, const TShar
 			+SOverlay::Slot()
 			[
 				SNew( SBox )
-				.Tag(TEXT("MainMenu"))
+				.AddMetaData<FTagMetaData>(FTagMetaData(TEXT("MainMenu")))
 				[
 					Widget1
 				]
@@ -199,7 +200,7 @@ void SLevelEditor::Initialize( const TSharedRef<SDockTab>& OwnerTab, const TShar
 			.HAlign( HAlign_Right )
 			[
 				SAssignNew( NotificationBarBox, SHorizontalBox )
-				.Tag(TEXT("PerformanceTools"))
+				.AddMetaData<FTagMetaData>(FTagMetaData(TEXT("PerformanceTools")))
 			]
 #endif
 		]
@@ -215,7 +216,7 @@ void SLevelEditor::Initialize( const TSharedRef<SDockTab>& OwnerTab, const TShar
 #if PLATFORM_MAC	
 	OwnerTab->SetRightContent(
 		SAssignNew( NotificationBarBox, SHorizontalBox )
-		.Tag(TEXT("PerformanceTools"))
+		.AddMetaData<FTagMetaData>(FTagMetaData(TEXT("PerformanceTools")))
 		);
 #endif
 
@@ -550,7 +551,7 @@ static TSharedRef<SDockTab> SummonDetailsPanel( FName TabIdentifier )
 		.ToolTip( IDocumentation::Get()->CreateToolTip( Label, nullptr, "Shared/LevelEditor", "DetailsTab" ) )
 		[
 			SNew( SBox )
-			.Tag(TEXT("ActorDetails"))
+			.AddMetaData<FTutorialMetaData>(FTutorialMetaData(TEXT("ActorDetails"), TEXT("LevelEditorSelectionDetails")))
 			[
 				DetailsView
 			]
@@ -584,7 +585,7 @@ TSharedRef<SDockTab> SLevelEditor::SpawnLevelEditorTab( const FSpawnTabArgs& Arg
 			.Icon( FEditorStyle::GetBrush("ToolBar.Icon") )
 			[
 				SNew(SHorizontalBox)
-				.Tag(TEXT("LevelToolbar"))
+				.AddMetaData<FTagMetaData>(FTagMetaData(TEXT("LevelToolbar")))
 				+ SHorizontalBox::Slot()
 				.FillWidth(1)
 				.VAlign(VAlign_Bottom)
@@ -614,7 +615,7 @@ TSharedRef<SDockTab> SLevelEditor::SpawnLevelEditorTab( const FSpawnTabArgs& Arg
 			.Label( NSLOCTEXT( "LevelEditor", "ToolsTabTitle", "Modes" ) )
 			[
 				SNew( SBox )
-				.Tag(TEXT("ToolsPanel"))
+				.AddMetaData<FTutorialMetaData>(FTutorialMetaData(TEXT("ToolsPanel"), TEXT("LevelEditorToolBox")))
 				[
 					NewToolBox
 				]
@@ -667,7 +668,7 @@ TSharedRef<SDockTab> SLevelEditor::SpawnLevelEditorTab( const FSpawnTabArgs& Arg
 				SNew(SBorder)
 				.Padding(4)
 				.BorderImage( FEditorStyle::GetBrush("ToolPanel.GroupBorder") )
-				.Tag(TEXT("SceneOutliner"))
+				.AddMetaData<FTutorialMetaData>(FTutorialMetaData(TEXT("SceneOutliner"), TEXT("LevelEditorSceneOutliner")))
 				[
 					SceneOutlinerModule.CreateSceneOutliner(
 						InitOptions,
@@ -686,7 +687,7 @@ TSharedRef<SDockTab> SLevelEditor::SpawnLevelEditorTab( const FSpawnTabArgs& Arg
 					SNew(SBorder)
 					.Padding( 0 )
 					.BorderImage( FEditorStyle::GetBrush("ToolPanel.GroupBorder") )
-					.Tag(TEXT("LayerBrowser"))
+					.AddMetaData<FTutorialMetaData>(FTutorialMetaData(TEXT("LayerBrowser"), TEXT("LevelEditorLayerBrowser")))
 					[
 						LayersModule.CreateLayerBrowser()
 					]
@@ -779,7 +780,7 @@ TSharedRef<SDockTab> SLevelEditor::SpawnLevelEditorTab( const FSpawnTabArgs& Arg
 		return SNew( SDockTab )
 			.Icon( FEditorStyle::GetBrush( "LevelEditor.WorldProperties.Tab" ) )
 			.Label( NSLOCTEXT("LevelEditor", "WorldSettingsTabTitle", "World Settings") )
-			.Tag(TEXT("WorldSettings"))
+			.AddMetaData<FTutorialMetaData>(FTutorialMetaData(TEXT("WorldSettings"), TEXT("WorldSettingsTab")))
 			[
 				WorldSettingsView.ToSharedRef()
 			];
diff --git a/Engine/Source/Editor/MaterialEditor/Private/MaterialEditor.cpp b/Engine/Source/Editor/MaterialEditor/Private/MaterialEditor.cpp
--- a/Engine/Source/Editor/MaterialEditor/Private/MaterialEditor.cpp
+++ b/Engine/Source/Editor/MaterialEditor/Private/MaterialEditor.cpp
@@ -2436,7 +2436,7 @@ TSharedRef<SDockTab> FMaterialEditor::SpawnTab_Palette(const FSpawnTabArgs& Args
 		.Label(LOCTEXT("MaterialPaletteTitle", "Palette"))
 		[
 			SNew( SBox )
-			.Tag(TEXT("MaterialPalette"))
+			.AddMetaData<FTagMetaData>(FTagMetaData(TEXT("MaterialPalette")))
 			[
 				Palette.ToSharedRef()
 			]
@@ -2454,7 +2454,7 @@ TSharedRef<SDockTab> FMaterialEditor::SpawnTab_Stats(const FSpawnTabArgs& Args)
 		.Label(LOCTEXT("MaterialStatsTitle", "Stats"))
 		[
 			SNew( SBox )
-			.Tag(TEXT("MaterialStats"))
+			.AddMetaData<FTagMetaData>(FTagMetaData(TEXT("MaterialStats")))
 			[
 				Stats.ToSharedRef()
 			]
@@ -2472,7 +2472,7 @@ TSharedRef<SDockTab> FMaterialEditor::SpawnTab_Find(const FSpawnTabArgs& Args)
 		.Label(LOCTEXT("MaterialFindTitle", "Find Results"))
 		[
 			SNew(SBox)
-			.Tag(TEXT("MaterialFind"))
+			.AddMetaData<FTagMetaData>(FTagMetaData(TEXT("MaterialFind")))
 			[
 				FindResults.ToSharedRef()
 			]
diff --git a/Engine/Source/Editor/PlacementMode/Private/SPlacementModeTools.cpp b/Engine/Source/Editor/PlacementMode/Private/SPlacementModeTools.cpp
--- a/Engine/Source/Editor/PlacementMode/Private/SPlacementModeTools.cpp
+++ b/Engine/Source/Editor/PlacementMode/Private/SPlacementModeTools.cpp
@@ -321,7 +321,7 @@ TSharedRef< SWidget > SPlacementModeTools::CreateStandardPanel()
 		.AutoHeight()
 		[
 			SNew( SBox)
-			.Tag(TEXT("PMRecentlyPlaced"))
+			.AddMetaData<FTagMetaData>(FTagMetaData(TEXT("PMRecentlyPlaced")))
 			[
 				CreatePlacementGroupTab( (int32)EPlacementTab::RecentlyPlaced, NSLOCTEXT( "PlacementMode", "RecentlyPlaced", "Recently Placed" ), true )
 			]
@@ -331,7 +331,7 @@ TSharedRef< SWidget > SPlacementModeTools::CreateStandardPanel()
 		.AutoHeight()
 		[
 			SNew( SBox )
-			.Tag(TEXT("PMGeometry"))
+			.AddMetaData<FTagMetaData>(FTagMetaData(TEXT("PMGeometry")))
 			[
 				CreatePlacementGroupTab( (int32)EPlacementTab::Geometry, NSLOCTEXT( "PlacementMode", "Geometry", "Geometry" ), false )
 			]
@@ -341,7 +341,7 @@ TSharedRef< SWidget > SPlacementModeTools::CreateStandardPanel()
 		.AutoHeight()
 		[
 			SNew( SBox )
-			.Tag(TEXT("PMLights"))
+			.AddMetaData<FTagMetaData>(FTagMetaData(TEXT("PMLights")))
 			[
 				CreatePlacementGroupTab( (int32)EPlacementTab::Lights, NSLOCTEXT( "PlacementMode", "Lights", "Lights" ), false )
 			]
@@ -351,7 +351,7 @@ TSharedRef< SWidget > SPlacementModeTools::CreateStandardPanel()
 		.AutoHeight()
 		[
 			SNew( SBox )
-			.Tag(TEXT("PMVisual"))
+			.AddMetaData<FTagMetaData>(FTagMetaData(TEXT("PMVisual")))
 			[
 				CreatePlacementGroupTab( (int32)EPlacementTab::Visual, NSLOCTEXT( "PlacementMode", "Visual", "Visual" ), false )
 			]
@@ -361,7 +361,7 @@ TSharedRef< SWidget > SPlacementModeTools::CreateStandardPanel()
 		.AutoHeight()
 		[
 			SNew( SBox )
-			.Tag(TEXT("PMBasic"))
+			.AddMetaData<FTagMetaData>(FTagMetaData(TEXT("PMBasic")))
 			[
 				CreatePlacementGroupTab( (int32)EPlacementTab::Basic, NSLOCTEXT( "PlacementMode", "Basic", "Basic" ), false )
 			]
@@ -371,7 +371,7 @@ TSharedRef< SWidget > SPlacementModeTools::CreateStandardPanel()
 		.AutoHeight()
 		[
 			SNew( SBox )
-			.Tag(TEXT("PMVolumes"))
+			.AddMetaData<FTagMetaData>(FTagMetaData(TEXT("PMVolumes")))
 			[
 				CreatePlacementGroupTab( (int32)EPlacementTab::Volumes, NSLOCTEXT( "PlacementMode", "Volumes", "Volumes" ), false )
 			]
@@ -381,7 +381,7 @@ TSharedRef< SWidget > SPlacementModeTools::CreateStandardPanel()
 		.AutoHeight()
 		[
 			SNew( SBox )
-			.Tag(TEXT("PMAllClasses"))
+			.AddMetaData<FTagMetaData>(FTagMetaData(TEXT("PMAllClasses")))
 			[
 				CreatePlacementGroupTab( (int32)EPlacementTab::AllClasses, NSLOCTEXT( "PlacementMode", "AllClasses", "All Classes" ), true )
 			]
diff --git a/Engine/Source/Editor/PropertyEditor/Private/DetailItemNode.cpp b/Engine/Source/Editor/PropertyEditor/Private/DetailItemNode.cpp
--- a/Engine/Source/Editor/PropertyEditor/Private/DetailItemNode.cpp
+++ b/Engine/Source/Editor/PropertyEditor/Private/DetailItemNode.cpp
@@ -9,6 +9,7 @@
 #include "DetailPropertyRow.h"
 #include "DetailCustomBuilderRow.h"
 #include "SDetailSingleItemRow.h"
+#include "TutorialMetaData.h"
 
 
 
@@ -137,10 +138,10 @@ void FDetailItemNode::ToggleExpansion()
 
 TSharedRef< ITableRow > FDetailItemNode::GenerateNodeWidget( const TSharedRef<STableViewBase>& OwnerTable, const FDetailColumnSizeData& ColumnSizeData, const TSharedRef<IPropertyUtilities>& PropertyUtilities )
 {
-	FString TagName;
+	FGraphNodeMetaData TagMeta(TEXT("DetailRowItem"));
 	if (Customization.IsValidCustomization() && Customization.GetPropertyNode().IsValid() )
 	{
-		TagName = FString::Printf(TEXT("DetailItemNode,%s,%d"), *Customization.GetPropertyNode()->GetDisplayName(), Customization.GetPropertyNode()->GetArrayIndex());
+		TagMeta.FriendlyName = Customization.GetPropertyNode()->GetDisplayName(); 
 	}
 	if( Customization.HasPropertyNode() && Customization.GetPropertyNode()->AsCategoryNode() )
 	{
@@ -148,7 +149,7 @@ TSharedRef< ITableRow > FDetailItemNode::GenerateNodeWidget( const TSharedRef<ST
 			SNew(SDetailCategoryTableRow, AsShared(), OwnerTable)
 			.IsEnabled(IsParentEnabled)
 			.DisplayName(Customization.GetPropertyNode()->GetDisplayName())
-			.Tag(*TagName)
+			.AddMetaData<FTutorialMetaData>(TagMeta)
 			.InnerCategory( true );
 	}
 	else
@@ -156,7 +157,7 @@ TSharedRef< ITableRow > FDetailItemNode::GenerateNodeWidget( const TSharedRef<ST
 		return
 			SNew(SDetailSingleItemRow, &Customization, HasMultiColumnWidget(), AsShared(), OwnerTable )
 			.IsEnabled( IsParentEnabled )
-			.Tag(*TagName)
+			.AddMetaData<FTutorialMetaData>(TagMeta)
 			.ColumnSizeData(ColumnSizeData);
 	}
 }
diff --git a/Engine/Source/Editor/PropertyEditor/Private/SDetailsView.cpp b/Engine/Source/Editor/PropertyEditor/Private/SDetailsView.cpp
--- a/Engine/Source/Editor/PropertyEditor/Private/SDetailsView.cpp
+++ b/Engine/Source/Editor/PropertyEditor/Private/SDetailsView.cpp
@@ -151,7 +151,7 @@ void SDetailsView::Construct(const FArguments& InArgs)
 				.ContentPadding(0)
 				.ForegroundColor( FSlateColor::UseForeground() )
 				.ButtonStyle( FEditorStyle::Get(), "ToggleButton" )
-				.Tag(TEXT("ViewOptions"))
+				.AddMetaData<FTagMetaData>(FTagMetaData(TEXT("ViewOptions")))
 				.MenuContent()
 				[
 					DetailViewOptions.MakeWidget()
diff --git a/Engine/Source/Editor/PropertyEditor/PropertyEditor.Build.cs b/Engine/Source/Editor/PropertyEditor/PropertyEditor.Build.cs
--- a/Engine/Source/Editor/PropertyEditor/PropertyEditor.Build.cs
+++ b/Engine/Source/Editor/PropertyEditor/PropertyEditor.Build.cs
@@ -12,6 +12,12 @@ public class PropertyEditor : ModuleRules
                 "ActorPickerMode",
 			}
 		);
+		
+        PublicIncludePathModuleNames.AddRange(
+            new string[] {                
+                "IntroTutorials"
+            }
+        );
 
 		PrivateIncludePaths.AddRange(
 			new string[] {
diff --git a/Engine/Source/Editor/UMGEditor/Private/TabFactory/DesignerTabSummoner.cpp b/Engine/Source/Editor/UMGEditor/Private/TabFactory/DesignerTabSummoner.cpp
--- a/Engine/Source/Editor/UMGEditor/Private/TabFactory/DesignerTabSummoner.cpp
+++ b/Engine/Source/Editor/UMGEditor/Private/TabFactory/DesignerTabSummoner.cpp
@@ -29,7 +29,7 @@ TSharedRef<SWidget> FDesignerTabSummoner::CreateTabBody(const FWorkflowTabSpawnI
 		.FillWidth(1)
 		[
 			SNew(SDesignerView, BlueprintEditor.Pin())
-			.Tag(TEXT("Designer"))
+			.AddMetaData<FTagMetaData>(FTagMetaData(TEXT("Designer")))
 		];
 }
 
diff --git a/Engine/Source/Editor/UMGEditor/Private/TabFactory/DetailsTabSummoner.cpp b/Engine/Source/Editor/UMGEditor/Private/TabFactory/DetailsTabSummoner.cpp
--- a/Engine/Source/Editor/UMGEditor/Private/TabFactory/DetailsTabSummoner.cpp
+++ b/Engine/Source/Editor/UMGEditor/Private/TabFactory/DetailsTabSummoner.cpp
@@ -27,7 +27,7 @@ TSharedRef<SWidget> FDetailsTabSummoner::CreateTabBody(const FWorkflowTabSpawnIn
 	TSharedPtr<FWidgetBlueprintEditor> BlueprintEditorPtr = StaticCastSharedPtr<FWidgetBlueprintEditor>(BlueprintEditor.Pin());
 
 	return SNew(SWidgetDetailsView, BlueprintEditorPtr)
-			.Tag(TEXT("Details"));
+		.AddMetaData<FTagMetaData>(FTagMetaData(TEXT("Details")));
 }
 
 #undef LOCTEXT_NAMESPACE 
diff --git a/Engine/Source/Editor/UMGEditor/Private/TabFactory/HierarchyTabSummoner.cpp b/Engine/Source/Editor/UMGEditor/Private/TabFactory/HierarchyTabSummoner.cpp
--- a/Engine/Source/Editor/UMGEditor/Private/TabFactory/HierarchyTabSummoner.cpp
+++ b/Engine/Source/Editor/UMGEditor/Private/TabFactory/HierarchyTabSummoner.cpp
@@ -27,7 +27,7 @@ TSharedRef<SWidget> FHierarchyTabSummoner::CreateTabBody(const FWorkflowTabSpawn
 	TSharedPtr<FWidgetBlueprintEditor> BlueprintEditorPtr = StaticCastSharedPtr<FWidgetBlueprintEditor>(BlueprintEditor.Pin());
 
 	return SNew(SHierarchyView, BlueprintEditorPtr, BlueprintEditorPtr->GetBlueprintObj()->SimpleConstructionScript)
-			.Tag(TEXT("Hierarchy"));
+		.AddMetaData<FTagMetaData>(FTagMetaData(TEXT("Hierarchy")));
 }
 
 #undef LOCTEXT_NAMESPACE 
diff --git a/Engine/Source/Editor/UMGEditor/Private/TabFactory/PaletteTabSummoner.cpp b/Engine/Source/Editor/UMGEditor/Private/TabFactory/PaletteTabSummoner.cpp
--- a/Engine/Source/Editor/UMGEditor/Private/TabFactory/PaletteTabSummoner.cpp
+++ b/Engine/Source/Editor/UMGEditor/Private/TabFactory/PaletteTabSummoner.cpp
@@ -27,7 +27,7 @@ TSharedRef<SWidget> FPaletteTabSummoner::CreateTabBody(const FWorkflowTabSpawnIn
 	TSharedPtr<FWidgetBlueprintEditor> BlueprintEditorPtr = StaticCastSharedPtr<FWidgetBlueprintEditor>(BlueprintEditor.Pin());
 
 	return SNew(SPaletteView, BlueprintEditorPtr)
-			.Tag(TEXT("Palette"));
+		.AddMetaData<FTagMetaData>(FTagMetaData(TEXT("Palette")));
 }
 
 #undef LOCTEXT_NAMESPACE 
diff --git a/Engine/Source/Editor/UMGEditor/Private/WidgetBlueprintEditor.cpp b/Engine/Source/Editor/UMGEditor/Private/WidgetBlueprintEditor.cpp
--- a/Engine/Source/Editor/UMGEditor/Private/WidgetBlueprintEditor.cpp
+++ b/Engine/Source/Editor/UMGEditor/Private/WidgetBlueprintEditor.cpp
@@ -450,7 +450,7 @@ TSharedRef<SWidget> FWidgetBlueprintEditor::CreateSequencerWidget()
 {
 	TSharedRef<SOverlay> SequencerOverlayRef =
 		SNew(SOverlay)
-		.Tag(TEXT("Sequencer"))
+		.AddMetaData<FTagMetaData>(FTagMetaData(TEXT("Sequencer")))
 		+ SOverlay::Slot()
 		[
 			GetSequencer()->GetSequencerWidget()
diff --git a/Engine/Source/Editor/UMGEditor/Private/WidgetBlueprintEditorToolbar.cpp b/Engine/Source/Editor/UMGEditor/Private/WidgetBlueprintEditorToolbar.cpp
--- a/Engine/Source/Editor/UMGEditor/Private/WidgetBlueprintEditorToolbar.cpp
+++ b/Engine/Source/Editor/UMGEditor/Private/WidgetBlueprintEditorToolbar.cpp
@@ -83,7 +83,7 @@ void FWidgetBlueprintEditorToolbar::FillWidgetBlueprintEditorModesToolbar(FToolB
 				TEXT("DefaultsMode")))
 			.IconImage(FEditorStyle::GetBrush("FullBlueprintEditor.SwitchToBlueprintDefaultsMode"))
 			.SmallIconImage(FEditorStyle::GetBrush("FullBlueprintEditor.SwitchToBlueprintDefaultsMode.Small"))
-			.Tag(TEXT("DefaultsMode"))
+			.AddMetaData<FTagMetaData>(FTagMetaData(TEXT("DefaultsMode")))
 		);
 
 		BlueprintEditorPtr->AddToolbarWidget(SNew(SBlueprintModeSeparator));
@@ -101,7 +101,7 @@ void FWidgetBlueprintEditorToolbar::FillWidgetBlueprintEditorModesToolbar(FToolB
 			.ToolTipText(LOCTEXT("GraphModeButtonTooltip", "Switch to Graph Editing Mode"))
 			.IconImage(FEditorStyle::GetBrush("FullBlueprintEditor.SwitchToScriptingMode"))
 			.SmallIconImage(FEditorStyle::GetBrush("FullBlueprintEditor.SwitchToScriptingMode.Small"))
-			.Tag(TEXT("GraphMode"))
+			.AddMetaData<FTagMetaData>(FTagMetaData(TEXT("GraphMode")))
 		);
 		
 		// Right side padding
diff --git a/Engine/Source/Editor/UnrealEd/Private/GraphEditor.cpp b/Engine/Source/Editor/UnrealEd/Private/GraphEditor.cpp
--- a/Engine/Source/Editor/UnrealEd/Private/GraphEditor.cpp
+++ b/Engine/Source/Editor/UnrealEd/Private/GraphEditor.cpp
@@ -33,7 +33,7 @@ void SGraphEditor::ConstructImplementation( const FArguments& InArgs )
 	this->ChildSlot
 	[
 		SNew( SBox )
-		.Tag(TEXT("GraphEditorPanel"))
+		.AddMetaData<FTagMetaData>(FTagMetaData(TEXT("GraphEditorPanel")))
 		[
 			Implementation.ToSharedRef()
 		]
diff --git a/Engine/Source/Editor/UnrealEd/Private/SEditorViewport.cpp b/Engine/Source/Editor/UnrealEd/Private/SEditorViewport.cpp
--- a/Engine/Source/Editor/UnrealEd/Private/SEditorViewport.cpp
+++ b/Engine/Source/Editor/UnrealEd/Private/SEditorViewport.cpp
@@ -6,6 +6,7 @@
 #include "SceneViewport.h"
 #include "EditorViewportCommands.h"
 #include "IDocumentation.h"
+#include "TutorialMetaData.h"
 
 #define LOCTEXT_NAMESPACE "EditorViewport"
 
@@ -37,7 +38,7 @@ void SEditorViewport::Construct( const FArguments& InArgs )
 		SAssignNew( ViewportWidget, SViewport )
 		.ShowEffectWhenDisabled( false )
 		.EnableGammaCorrection( false ) // Scene rendering handles this
-		.Tag(TEXT("EditorViewports"))
+		.AddMetaData<FTutorialMetaData>(FTutorialMetaData(TEXT("EditorViewports"), TEXT("LevelEditorViewport")))
 		[
 			SAssignNew( ViewportOverlay, SOverlay )
 			+SOverlay::Slot()
diff --git a/Engine/Source/Runtime/AppFramework/Private/Widgets/Testing/STestSuite.cpp b/Engine/Source/Runtime/AppFramework/Private/Widgets/Testing/STestSuite.cpp
--- a/Engine/Source/Runtime/AppFramework/Private/Widgets/Testing/STestSuite.cpp
+++ b/Engine/Source/Runtime/AppFramework/Private/Widgets/Testing/STestSuite.cpp
@@ -4535,7 +4535,7 @@ TSharedRef<SDockTab> SpawnTab(const FSpawnTabArgs& Args, FName TabIdentifier)
 			SNew( SDocumentsTest, TabManagerRef )
 			.RenderTransform_Static(&::GetTestRenderTransform)
 			.RenderTransformPivot_Static(&::GetTestRenderTransformPivot)
-			.Tag("DocumentSpawner")
+			.AddMetaData<FTagMetaData>(FTagMetaData("DocumentSpawner"))
 		];
 	}
 	else if (TabIdentifier == FName(TEXT("TableViewTestTab")))
diff --git a/Engine/Source/Runtime/Slate/Private/Framework/MultiBox/MultiBoxBuilder.cpp b/Engine/Source/Runtime/Slate/Private/Framework/MultiBox/MultiBoxBuilder.cpp
--- a/Engine/Source/Runtime/Slate/Private/Framework/MultiBox/MultiBoxBuilder.cpp
+++ b/Engine/Source/Runtime/Slate/Private/Framework/MultiBox/MultiBoxBuilder.cpp
@@ -406,7 +406,7 @@ void FToolBarBuilder::AddWidget( TSharedRef<SWidget> InWidget, FName TutorialHig
 	TSharedRef<SWidget> ChildWidget = InWidget;
 	InWidget = 
 		SNew( SBox )
-		.Tag(TutorialHighlightName)
+		.AddMetaData<FTagMetaData>(FTagMetaData(TutorialHighlightName))
 		[
 			ChildWidget
 		];
diff --git a/Engine/Source/Runtime/Slate/Private/Framework/MultiBox/SMenuEntryBlock.cpp b/Engine/Source/Runtime/Slate/Private/Framework/MultiBox/SMenuEntryBlock.cpp
--- a/Engine/Source/Runtime/Slate/Private/Framework/MultiBox/SMenuEntryBlock.cpp
+++ b/Engine/Source/Runtime/Slate/Private/Framework/MultiBox/SMenuEntryBlock.cpp
@@ -764,7 +764,7 @@ void SMenuEntryBlock::BuildMultiBlockWidget(const ISlateStyle* StyleSet, const F
 		ChildSlot
 		[
 			SNew( SBox )
-			.Tag(TutorialName)
+			.AddMetaData<FTagMetaData>(FTagMetaData(TutorialName))
 			[
 				ChildWidget
 			]
diff --git a/Engine/Source/Runtime/Slate/Private/Framework/MultiBox/SToolBarButtonBlock.cpp b/Engine/Source/Runtime/Slate/Private/Framework/MultiBox/SToolBarButtonBlock.cpp
--- a/Engine/Source/Runtime/Slate/Private/Framework/MultiBox/SToolBarButtonBlock.cpp
+++ b/Engine/Source/Runtime/Slate/Private/Framework/MultiBox/SToolBarButtonBlock.cpp
@@ -154,7 +154,7 @@ void SToolBarButtonBlock::BuildMultiBlockWidget(const ISlateStyle* StyleSet, con
 	TSharedRef< SWidget > ButtonContent =
 
 		SNew(SHorizontalBox)
-		.Tag(TutorialHighlightName)
+		.AddMetaData<FTagMetaData>(FTagMetaData(TutorialHighlightName))
 
 		+ SHorizontalBox::Slot()
 		.FillWidth(1)
diff --git a/Engine/Source/Runtime/Slate/Private/Framework/MultiBox/SToolBarComboButtonBlock.cpp b/Engine/Source/Runtime/Slate/Private/Framework/MultiBox/SToolBarComboButtonBlock.cpp
--- a/Engine/Source/Runtime/Slate/Private/Framework/MultiBox/SToolBarComboButtonBlock.cpp
+++ b/Engine/Source/Runtime/Slate/Private/Framework/MultiBox/SToolBarComboButtonBlock.cpp
@@ -106,7 +106,7 @@ void SToolBarComboButtonBlock::BuildMultiBlockWidget(const ISlateStyle* StyleSet
 
 	bool bBool = TEXT("Bool");
 
-	// Create a tag
+	// Setup the string for the metatag
 	FName TagName;
 	if (ToolBarComboButtonBlock->GetTutorialHighlightName() == NAME_None)
 	{
@@ -120,7 +120,7 @@ void SToolBarComboButtonBlock::BuildMultiBlockWidget(const ISlateStyle* StyleSet
 	// Create the content for our button
 	TSharedRef< SWidget > ButtonContent =
 		SNew( SVerticalBox )
-		.Tag(TagName)
+		.AddMetaData<FTagMetaData>(FTagMetaData(TagName))
 		// Icon image
 		+ SVerticalBox::Slot()
 		.AutoHeight()
