Content Browser fixes:
- Folder hotkeys (rename, delete) work properly in path view
- Deleting and renaming folders in asset view doesn't reset the asset view scope
- Deleting the open folder returns scope to the Game folder (no longer disables "New" and "Import")

diff --git a/Engine/Source/Editor/ContentBrowser/Private/PathContextMenu.cpp b/Engine/Source/Editor/ContentBrowser/Private/PathContextMenu.cpp
--- a/Engine/Source/Editor/ContentBrowser/Private/PathContextMenu.cpp
+++ b/Engine/Source/Editor/ContentBrowser/Private/PathContextMenu.cpp
@@ -34,10 +34,18 @@ void FPathContextMenu::SetOnRenameFolderRequested(const FOnRenameFolderRequested
 	OnRenameFolderRequested = InOnRenameFolderRequested;
 }
 
-TSharedRef<FExtender> FPathContextMenu::MakePathViewContextMenuExtender(const TArray<FString>& InSelectedPaths)
+void FPathContextMenu::SetOnFolderDeleted(const FOnFolderDeleted& InOnFolderDeleted)
+{
+	OnFolderDeleted = InOnFolderDeleted;
+}
+
+void FPathContextMenu::SetSelectedPaths(const TArray<FString>& InSelectedPaths)
 {
 	SelectedPaths = InSelectedPaths;
+}
 
+TSharedRef<FExtender> FPathContextMenu::MakePathViewContextMenuExtender(const TArray<FString>& InSelectedPaths)
+{
 	// Cache any vars that are used in determining if you can execute any actions.
 	// Useful for actions whose "CanExecute" will not change or is expensive to calculate.
 	CacheCanExecuteVars();
@@ -555,6 +563,10 @@ FReply FPathContextMenu::ExecuteDeleteFolderConfirmed()
 	if ( ContentBrowserUtils::DeleteFolders(SelectedPaths) )
 	{
 		ResetColors();
+		if (OnFolderDeleted.IsBound())
+		{
+			OnFolderDeleted.Execute();
+		}
 	}
 
 	return FReply::Handled();
diff --git a/Engine/Source/Editor/ContentBrowser/Private/PathContextMenu.h b/Engine/Source/Editor/ContentBrowser/Private/PathContextMenu.h
--- a/Engine/Source/Editor/ContentBrowser/Private/PathContextMenu.h
+++ b/Engine/Source/Editor/ContentBrowser/Private/PathContextMenu.h
@@ -17,6 +17,13 @@ public:
 	DECLARE_DELEGATE_OneParam(FOnRenameFolderRequested, const FString& /*FolderToRename*/);
 	void SetOnRenameFolderRequested(const FOnRenameFolderRequested& InOnRenameFolderRequested);
 
+	/** Delegate for when the context menu has successfully deleted a folder */
+	DECLARE_DELEGATE(FOnFolderDeleted)
+	void SetOnFolderDeleted(const FOnFolderDeleted& InOnFolderDeleted);
+
+	/** Sets the currently selected paths */
+	void SetSelectedPaths(const TArray<FString>& InSelectedPaths);
+
 	/** Makes the asset tree context menu extender */
 	TSharedRef<FExtender> MakePathViewContextMenuExtender(const TArray<FString>& InSelectedPaths);
 
@@ -119,6 +126,7 @@ private:
 	TWeakPtr<SWidget> ParentContent;
 	FNewAssetContextMenu::FOnNewAssetRequested OnNewAssetRequested;
 	FOnRenameFolderRequested OnRenameFolderRequested;
+	FOnFolderDeleted OnFolderDeleted;
 
 	/** Cached SCC CanExecute vars */
 	bool bCanExecuteSCCCheckOut;
diff --git a/Engine/Source/Editor/ContentBrowser/Private/SContentBrowser.cpp b/Engine/Source/Editor/ContentBrowser/Private/SContentBrowser.cpp
--- a/Engine/Source/Editor/ContentBrowser/Private/SContentBrowser.cpp
+++ b/Engine/Source/Editor/ContentBrowser/Private/SContentBrowser.cpp
@@ -54,6 +54,7 @@ void SContentBrowser::Construct( const FArguments& InArgs, const FName& InInstan
 	PathContextMenu = MakeShareable(new FPathContextMenu( AsShared() ));
 	PathContextMenu->SetOnNewAssetRequested( FNewAssetContextMenu::FOnNewAssetRequested::CreateSP(this, &SContentBrowser::NewAssetRequested) );
 	PathContextMenu->SetOnRenameFolderRequested(FPathContextMenu::FOnRenameFolderRequested::CreateSP(this, &SContentBrowser::OnRenameFolderRequested));
+	PathContextMenu->SetOnFolderDeleted(FPathContextMenu::FOnFolderDeleted::CreateSP(this, &SContentBrowser::OnOpenedFolderDeleted));
 
 	FrontendFilters = MakeShareable(new AssetFilterCollectionType());
 	TextFilter = MakeShareable( new FFrontendFilter_Text() );
@@ -958,11 +959,14 @@ void SContentBrowser::PathSelected(const FString& FolderPath)
 	{
 		PathChangedDelegate.Broadcast(FolderPath);
 	}
+
+	// Update the context menu's selected paths list
+	PathContextMenu->SetSelectedPaths(SelectedPaths);
 }
 
-TSharedRef<FExtender> SContentBrowser::GetPathContextMenuExtender(const TArray<FString>& SelectedPaths) const
+TSharedRef<FExtender> SContentBrowser::GetPathContextMenuExtender(const TArray<FString>& InSelectedPaths) const
 {
-	return PathContextMenu->MakePathViewContextMenuExtender( SelectedPaths );
+	return PathContextMenu->MakePathViewContextMenuExtender(InSelectedPaths);
 }
 
 void SContentBrowser::CollectionSelected(const FCollectionNameType& SelectedCollection)
@@ -1800,6 +1804,18 @@ void SContentBrowser::OnRenameFolderRequested(const FString& FolderToRename)
 	}
 }
 
+void SContentBrowser::OnOpenedFolderDeleted()
+{
+	// Since the contents of the asset view have just been deleted, set the selected path to the default "/Game"
+	TArray<FString> DefaultSelectedPaths;
+	DefaultSelectedPaths.Add(TEXT("/Game"));
+	PathViewPtr->SetSelectedPaths(DefaultSelectedPaths);
+
+	FSourcesData DefaultSourcesData;
+	DefaultSourcesData.PackagePaths.Add(TEXT("/Game"));
+	AssetViewPtr->SetSourcesData(DefaultSourcesData);
+}
+
 void SContentBrowser::OnDuplicateRequested(const TWeakObjectPtr<UObject>& OriginalObject)
 {
 	UObject* Object = OriginalObject.Get();
@@ -1838,8 +1854,6 @@ void SContentBrowser::HandlePathRemoved(const FString& Path)
 {
 	const FName PathName(*Path);
 
-	AssetViewPtr->SetSourcesData(FSourcesData());
-
 	auto RemoveHistoryDelegate = [&](const FHistoryData& HistoryData)
 	{
 		return (HistoryData.SourcesData.PackagePaths.Num() == 1 &&
diff --git a/Engine/Source/Editor/ContentBrowser/Private/SContentBrowser.h b/Engine/Source/Editor/ContentBrowser/Private/SContentBrowser.h
--- a/Engine/Source/Editor/ContentBrowser/Private/SContentBrowser.h
+++ b/Engine/Source/Editor/ContentBrowser/Private/SContentBrowser.h
@@ -275,6 +275,9 @@ private:
 	/** Handler for when the asset context menu requests to rename a folder */
 	void OnRenameFolderRequested(const FString& FolderToRename);
 
+	/** Handler for when the path context menu has successfully deleted a folder */
+	void OnOpenedFolderDeleted();
+
 	/** Handler for when the asset context menu requests to duplicate an asset */
 	void OnDuplicateRequested(const TWeakObjectPtr<UObject>& OriginalObject);
 
