ttp 344629 EDITOR: BP: Cannot re-select an asset in a list after removing it from a node

diff --git a/Engine/Source/Editor/ContentBrowser/Private/SAssetPicker.cpp b/Engine/Source/Editor/ContentBrowser/Private/SAssetPicker.cpp
--- a/Engine/Source/Editor/ContentBrowser/Private/SAssetPicker.cpp
+++ b/Engine/Source/Editor/ContentBrowser/Private/SAssetPicker.cpp
@@ -337,6 +337,10 @@ void SAssetPicker::OnFilterChanged()
 FReply SAssetPicker::OnNoneButtonClicked()
 {
 	OnAssetSelected.ExecuteIfBound(FAssetData());
+	if (AssetViewPtr.IsValid())
+	{
+		AssetViewPtr->ClearSelection(true);
+	}
 	return FReply::Handled();
 }
 
diff --git a/Engine/Source/Editor/ContentBrowser/Private/SAssetView.cpp b/Engine/Source/Editor/ContentBrowser/Private/SAssetView.cpp
--- a/Engine/Source/Editor/ContentBrowser/Private/SAssetView.cpp
+++ b/Engine/Source/Editor/ContentBrowser/Private/SAssetView.cpp
@@ -2578,8 +2578,10 @@ void SAssetView::OnPreviewAssets()
 	OnAssetsActivated.ExecuteIfBound(GetSelectedAssets(), EAssetTypeActivationMethod::Previewed);
 }
 
-void SAssetView::ClearSelection()
+void SAssetView::ClearSelection(bool bForceSilent)
 {
+	const bool bTempBulkSelectingValue = bForceSilent ? true : bBulkSelecting;
+	TGuardValue<bool>(bBulkSelecting, bTempBulkSelectingValue);
 	switch ( GetCurrentViewType() )
 	{
 		case EAssetViewType::List: ListView->ClearSelection(); break;
diff --git a/Engine/Source/Editor/ContentBrowser/Private/SAssetView.h b/Engine/Source/Editor/ContentBrowser/Private/SAssetView.h
--- a/Engine/Source/Editor/ContentBrowser/Private/SAssetView.h
+++ b/Engine/Source/Editor/ContentBrowser/Private/SAssetView.h
@@ -225,7 +225,7 @@ public:
 	void OnPreviewAssets();
 
 	/** Clears the selection of all the lists in the view */
-	void ClearSelection();
+	void ClearSelection(bool bForceSilent = false);
 
 	/** Returns true if the asset view is in thumbnail editing mode */
 	bool IsThumbnailEditMode() const;
