UE4: Removed the "Show only assets in selected folders" editor preference for the Content Browser. This is implied from "Show Folders" and should have been removed when folders were added.

diff --git a/Engine/Config/BaseEditorUserSettings.ini b/Engine/Config/BaseEditorUserSettings.ini
--- a/Engine/Config/BaseEditorUserSettings.ini
+++ b/Engine/Config/BaseEditorUserSettings.ini
@@ -99,8 +99,6 @@ bShowOnlySelectedActors=False
 [/Script/UnrealEd.ContentBrowserSettings]
 ; The number of objects to load at once in the Content Browser before displaying a warning about loading many assets
 NumObjectsToLoadBeforeWarning=20
-; Whether to show only assets in selected folders in the Content Browser, excluding sub-folders
-ShowOnlyAssetsInSelectedFolders=False
 ; Whether to render thumbnails for loaded assets in real-time in the Content Browser
 RealTimeThumbnails=True
 ; Whether to display folders in the assets view of the content browser. Note that this implies 'Show Only Assets in Selected Folders'.
diff --git a/Engine/Source/Editor/ContentBrowser/Private/SAssetPicker.cpp b/Engine/Source/Editor/ContentBrowser/Private/SAssetPicker.cpp
--- a/Engine/Source/Editor/ContentBrowser/Private/SAssetPicker.cpp
+++ b/Engine/Source/Editor/ContentBrowser/Private/SAssetPicker.cpp
@@ -212,7 +212,6 @@ void SAssetPicker::Construct( const FArguments& InArgs )
 		.AllowDragging( InArgs._AssetPickerConfig.bAllowDragging )
 		.CanShowClasses( InArgs._AssetPickerConfig.bCanShowClasses )
 		.CanShowFolders( false )
-		.CanShowOnlyAssetsInSelectedFolders( InArgs._AssetPickerConfig.bCanShowOnlyAssetsInSelectedFolders )
 		.CanShowRealTimeThumbnails( InArgs._AssetPickerConfig.bCanShowRealTimeThumbnails )
 		.CanShowDevelopersFolder( InArgs._AssetPickerConfig.bCanShowDevelopersFolder )
 		.PreloadAssetsForContextMenu( InArgs._AssetPickerConfig.bPreloadAssetsForContextMenu )
diff --git a/Engine/Source/Editor/ContentBrowser/Private/SAssetView.cpp b/Engine/Source/Editor/ContentBrowser/Private/SAssetView.cpp
--- a/Engine/Source/Editor/ContentBrowser/Private/SAssetView.cpp
+++ b/Engine/Source/Editor/ContentBrowser/Private/SAssetView.cpp
@@ -124,8 +124,6 @@ void SAssetView::Construct( const FArguments& InArgs )
 	bCanShowClasses = InArgs._CanShowClasses;
 
 	bCanShowFolders = InArgs._CanShowFolders;
-
-	bCanShowOnlyAssetsInSelectedFolders = InArgs._CanShowOnlyAssetsInSelectedFolders;
 		
 	bCanShowRealTimeThumbnails = InArgs._CanShowRealTimeThumbnails;
 
@@ -1303,7 +1301,7 @@ void SAssetView::RefreshSourceItems()
 		// Assemble the filter using the current sources
 		// force recursion when the user is searching
 		const bool bRecurse = ShouldFilterRecursively();
-		const bool bUsingFolders = GetDefault<UContentBrowserSettings>()->ShowOnlyAssetsInSelectedFolders || IsShowingFolders();
+		const bool bUsingFolders = IsShowingFolders();
 		FARFilter Filter = SourcesData.MakeFilter(bRecurse, bUsingFolders);
 
 		bWereItemsRecursivelyFiltered = bRecurse;
@@ -1463,7 +1461,7 @@ void SAssetView::RefreshFilteredItems()
 	if ( bIsFrontendFilterActive && FrontendFilters.IsValid() )
 	{
 		const bool bRecurse = ShouldFilterRecursively();
-		const bool bUsingFolders = GetDefault<UContentBrowserSettings>()->ShowOnlyAssetsInSelectedFolders || IsShowingFolders();
+		const bool bUsingFolders = IsShowingFolders();
 		FARFilter CombinedFilter = SourcesData.MakeFilter(bRecurse, bUsingFolders);
 		CombinedFilter.Append(BackendFilter);
 
@@ -2050,7 +2048,7 @@ bool SAssetView::PassesCurrentFrontendFilter(const FAssetData& Item) const
 void SAssetView::RunAssetsThroughBackendFilter(TArray<FAssetData>& InOutAssetDataList) const
 {
 	const bool bRecurse = ShouldFilterRecursively();
-	const bool bUsingFolders = GetDefault<UContentBrowserSettings>()->ShowOnlyAssetsInSelectedFolders || IsShowingFolders();
+	const bool bUsingFolders = IsShowingFolders();
 	FARFilter Filter = SourcesData.MakeFilter(bRecurse, bUsingFolders);
 	
 	if ( SourcesData.Collections.Num() > 0 && Filter.ObjectPaths.Num() == 0 )
@@ -2205,19 +2203,6 @@ TSharedRef<SWidget> SAssetView::GetViewButtonContent()
 			EUserInterfaceActionType::ToggleButton
 			);
 
-		MenuBuilder.AddMenuEntry(
-			LOCTEXT("ShowOnlyAssetsInSelectedFolders", "Show Only Assets in Selected Folders"),
-			LOCTEXT("ShowOnlyAssetsInSelectedFoldersToolTip", "Only displays the assets of the selected folders"),
-			FSlateIcon(),
-			FUIAction(
-				FExecuteAction::CreateSP( this, &SAssetView::ToggleShowOnlyAssetsInSelectedFolders ),
-				FCanExecuteAction::CreateSP( this, &SAssetView::CanShowOnlyAssetsInSelectedFolders ),
-				FIsActionChecked::CreateSP( this, &SAssetView::IsShowingOnlyAssetsInSelectedFolders )
-			),
-			NAME_None,
-			EUserInterfaceActionType::ToggleButton
-			);
-
 		MenuBuilder.AddMenuEntry(
 			LOCTEXT("ShowDevelopersFolderOption", "Show Developers Folder"),
 			LOCTEXT("ShowDevelopersFolderOptionToolTip", "Show the developers folder in the view."),
@@ -2322,23 +2307,6 @@ bool SAssetView::IsShowingFolders() const
 	return IsToggleShowFoldersAllowed() ? GetDefault<UContentBrowserSettings>()->DisplayFolders : false;
 }
 
-void SAssetView::ToggleShowOnlyAssetsInSelectedFolders()
-{
-	check( CanShowOnlyAssetsInSelectedFolders() );
-	GetMutableDefault<UContentBrowserSettings>()->ShowOnlyAssetsInSelectedFolders = !GetDefault<UContentBrowserSettings>()->ShowOnlyAssetsInSelectedFolders;
-	RequestSlowFullListRefresh();
-}
-
-bool SAssetView::CanShowOnlyAssetsInSelectedFolders() const
-{
-	 return bCanShowOnlyAssetsInSelectedFolders;
-}
-
-bool SAssetView::IsShowingOnlyAssetsInSelectedFolders() const
-{
-	return CanShowOnlyAssetsInSelectedFolders() ? GetDefault<UContentBrowserSettings>()->ShowOnlyAssetsInSelectedFolders : false;
-}
-
 void SAssetView::ToggleRealTimeThumbnails()
 {
 	check( CanShowRealTimeThumbnails() );
@@ -3732,8 +3700,7 @@ void SAssetView::SetUserSearching(bool bInSearching)
 
 void SAssetView::HandleSettingChanged(FName PropertyName)
 {
-	if ((PropertyName == "ShowOnlyAssetsInSelectedFolders") ||
-		(PropertyName == "DisplayFolders") ||
+	if ((PropertyName == "DisplayFolders") ||
 		(PropertyName == "DisplayDevelopersFolder") ||
 		(PropertyName == "DisplayEngineFolder") ||
 		(PropertyName == NAME_None))	// @todo: Needed if PostEditChange was called manually, for now
diff --git a/Engine/Source/Editor/ContentBrowser/Private/SAssetView.h b/Engine/Source/Editor/ContentBrowser/Private/SAssetView.h
--- a/Engine/Source/Editor/ContentBrowser/Private/SAssetView.h
+++ b/Engine/Source/Editor/ContentBrowser/Private/SAssetView.h
@@ -23,7 +23,6 @@ public:
 		, _AllowThumbnailEditMode(false)
 		, _CanShowClasses(true)
 		, _CanShowFolders(false)
-		, _CanShowOnlyAssetsInSelectedFolders(false)
 		, _CanShowRealTimeThumbnails(false)
 		, _CanShowDevelopersFolder(false)
 		, _PreloadAssetsForContextMenu(true)
@@ -108,9 +107,6 @@ public:
 		/** Indicates if the 'Show Folders' option should be visible */
 		SLATE_ARGUMENT( bool, CanShowFolders )
 
-		/** Indicates if the 'Show Only Assets In Selection' option should be visible */
-		SLATE_ARGUMENT( bool, CanShowOnlyAssetsInSelectedFolders )
-
 		/** Indicates if the 'Real-Time Thumbnails' option should be visible */
 		SLATE_ARGUMENT( bool, CanShowRealTimeThumbnails )
 
@@ -346,17 +342,6 @@ private:
 	/** @return true when we are showing folders */
 	bool IsShowingFolders() const;
 
-
-	/** Toggle whether only assets from the selected folders are shown */
-	void ToggleShowOnlyAssetsInSelectedFolders();
-
-	/** Whether or not it's possible to only show assets from the selected folders */
-	bool CanShowOnlyAssetsInSelectedFolders() const;
-
-	/** @return true when we are showing only the assets from the selected folders */
-	bool IsShowingOnlyAssetsInSelectedFolders() const;
-
-
 	/** Toggle whether to show real-time thumbnails */
 	void ToggleRealTimeThumbnails();
 
@@ -366,7 +351,6 @@ private:
 	/** @return true if we are showing real-time thumbnails */
 	bool IsShowingRealTimeThumbnails() const;
 
-
 	/** Toggle whether plugin content folders should be shown or not */
 	void ToggleShowPluginFolders();
 
@@ -756,9 +740,6 @@ private:
 	/** Indicates if the 'Show Folders' option should be visible */
 	bool bCanShowFolders;
 
-	/** Indicates if the 'Show Only Assets In Selection' option should be visible */
-	bool bCanShowOnlyAssetsInSelectedFolders;
-
 	/** Indicates if the 'Real-Time Thumbnails' option should be visible */
 	bool bCanShowRealTimeThumbnails;
 
diff --git a/Engine/Source/Editor/ContentBrowser/Private/SContentBrowser.cpp b/Engine/Source/Editor/ContentBrowser/Private/SContentBrowser.cpp
--- a/Engine/Source/Editor/ContentBrowser/Private/SContentBrowser.cpp
+++ b/Engine/Source/Editor/ContentBrowser/Private/SContentBrowser.cpp
@@ -520,7 +520,6 @@ void SContentBrowser::Construct( const FArguments& InArgs, const FName& InInstan
 						.AllowThumbnailEditMode(true)
 						.AllowThumbnailHintLabel(false)
 						.CanShowFolders(true)
-						.CanShowOnlyAssetsInSelectedFolders(true)
 						.CanShowRealTimeThumbnails(true)
 						.CanShowDevelopersFolder(true)
 						.AddMetaData<FTagMetaData>(FTagMetaData(TEXT("ContentBrowserAssets")))
diff --git a/Engine/Source/Editor/ContentBrowser/Public/IContentBrowserSingleton.h b/Engine/Source/Editor/ContentBrowser/Public/IContentBrowserSingleton.h
--- a/Engine/Source/Editor/ContentBrowser/Public/IContentBrowserSingleton.h
+++ b/Engine/Source/Editor/ContentBrowser/Public/IContentBrowserSingleton.h
@@ -114,9 +114,6 @@ struct FAssetPickerConfig
 	/** Indicates if this view is allowed to show classes */
 	bool bCanShowClasses;
 
-	/** Indicates if the 'Show Only Assets In Selection' option should be visible */
-	bool bCanShowOnlyAssetsInSelectedFolders;
-
 	/** Indicates if the 'Real-Time Thumbnails' option should be visible */
 	bool bCanShowRealTimeThumbnails;
 
@@ -143,7 +140,6 @@ struct FAssetPickerConfig
 		, bAutohideSearchBar(false)
 		, bAllowDragging(true)
 		, bCanShowClasses(true)
-		, bCanShowOnlyAssetsInSelectedFolders(false)
 		, bCanShowRealTimeThumbnails(false)
 		, bCanShowDevelopersFolder(false)
 		, bPreloadAssetsForContextMenu(true)
diff --git a/Engine/Source/Editor/UnrealEd/Classes/Settings/ContentBrowserSettings.h b/Engine/Source/Editor/UnrealEd/Classes/Settings/ContentBrowserSettings.h
--- a/Engine/Source/Editor/UnrealEd/Classes/Settings/ContentBrowserSettings.h
+++ b/Engine/Source/Editor/UnrealEd/Classes/Settings/ContentBrowserSettings.h
@@ -25,10 +25,6 @@ public:
 	UPROPERTY(EditAnywhere, config, Category=ContentBrowser, meta=(DisplayName = "Assets to Load at Once Before Warning"))
 	int32 NumObjectsToLoadBeforeWarning;
 
-	/** Whether to show only assets in selected folders in the Content Browser, excluding sub-folders. Note that this always the case when 'Display Folders' is enabled. */
-	UPROPERTY(config)
-	bool ShowOnlyAssetsInSelectedFolders;
-
 	/** Whether to render thumbnails for loaded assets in real-time in the Content Browser */
 	UPROPERTY(config)
 	bool RealTimeThumbnails;
