UE4: Asset pickers can now show folders. When showing folders they will ignore the backend filter when determining if they should recursively show assets. Path pickers can now have the path changed after they are spawned. Asset dialogs now show folders like the content browser.

diff --git a/Engine/Source/Editor/ContentBrowser/Private/SAssetDialog.cpp b/Engine/Source/Editor/ContentBrowser/Private/SAssetDialog.cpp
--- a/Engine/Source/Editor/ContentBrowser/Private/SAssetDialog.cpp
+++ b/Engine/Source/Editor/ContentBrowser/Private/SAssetDialog.cpp
@@ -25,7 +25,7 @@ void SAssetDialog::Construct(const FArguments& InArgs, const FSharedAssetDialogC
 	PathPickerConfig.DefaultPath = DefaultPath;
 	PathPickerConfig.bFocusSearchBoxWhenOpened = false;
 	PathPickerConfig.OnPathSelected = FOnPathSelected::CreateSP(this, &SAssetDialog::HandlePathSelected);
-	TSharedPtr<SWidget> PathPicker = FContentBrowserSingleton::Get().CreatePathPicker(PathPickerConfig);
+	PathPickerConfig.SetPathsDelegates.Add(&SetPathsDelegate);
 
 	FAssetPickerConfig AssetPickerConfig;
 	AssetPickerConfig.Filter.ClassNames.Append(AssetClassNames);
@@ -38,7 +38,8 @@ void SAssetDialog::Construct(const FArguments& InArgs, const FSharedAssetDialogC
 	AssetPickerConfig.SetFilterDelegates.Add(&SetFilterDelegate);
 	AssetPickerConfig.GetCurrentSelectionDelegates.Add(&GetCurrentSelectionDelegate);
 	AssetPickerConfig.SaveSettingsName = TEXT("AssetDialog");
-	TSharedPtr<SWidget> AssetPicker = FContentBrowserSingleton::Get().CreateAssetPicker(AssetPickerConfig);
+	AssetPickerConfig.bCanShowFolders = true;
+	AssetPickerConfig.OnFolderEntered = FOnPathSelected::CreateSP(this, &SAssetDialog::HandleAssetViewFolderEntered);
 
 	SetCurrentlySelectedPath(DefaultPath);
 
@@ -68,6 +69,9 @@ void SAssetDialog::Construct(const FArguments& InArgs, const FSharedAssetDialogC
 		ensureMsgf(0, TEXT("AssetDialog type %d is not supported."), DialogType);
 	}
 
+	TSharedPtr<SWidget> PathPicker = FContentBrowserSingleton::Get().CreatePathPicker(PathPickerConfig);
+	TSharedPtr<SWidget> AssetPicker = FContentBrowserSingleton::Get().CreateAssetPicker(AssetPickerConfig);
+
 	// The root widget in this dialog.
 	TSharedRef<SVerticalBox> MainVerticalBox = SNew(SVerticalBox);
 
@@ -251,6 +255,15 @@ void SAssetDialog::HandlePathSelected(const FString& NewPath)
 	SetFilterDelegate.ExecuteIfBound(NewFilter);
 }
 
+void SAssetDialog::HandleAssetViewFolderEntered(const FString& NewPath)
+{
+	SetCurrentlySelectedPath(NewPath);
+
+	TArray<FString> NewPaths;
+	NewPaths.Add(NewPath);
+	SetPathsDelegate.Execute(NewPaths);
+}
+
 bool SAssetDialog::IsConfirmButtonEnabled() const
 {
 	if ( DialogType == EAssetDialogType::Open )
diff --git a/Engine/Source/Editor/ContentBrowser/Private/SAssetDialog.h b/Engine/Source/Editor/ContentBrowser/Private/SAssetDialog.h
--- a/Engine/Source/Editor/ContentBrowser/Private/SAssetDialog.h
+++ b/Engine/Source/Editor/ContentBrowser/Private/SAssetDialog.h
@@ -38,6 +38,9 @@ private:
 	/** Handler for when a path is selected in the path view */
 	void HandlePathSelected(const FString& NewPath);
 
+	/** Handler for when a folder is entered in the asset view */
+	void HandleAssetViewFolderEntered(const FString& NewPath);
+
 	/** Determines if the confirm button (e.g. Open/Save) is enabled. */
 	bool IsConfirmButtonEnabled() const;
 
@@ -75,6 +78,9 @@ private:
 	/** Whether this is an open or save dialog */
 	EAssetDialogType::Type DialogType;
 
+	/** Used to update the path view after it has been created */
+	FSetPathPickerPathsDelegate SetPathsDelegate;
+
 	/** Used to update the asset view after it has been created */
 	FSetARFilterDelegate SetFilterDelegate;
 
diff --git a/Engine/Source/Editor/ContentBrowser/Private/SAssetPicker.cpp b/Engine/Source/Editor/ContentBrowser/Private/SAssetPicker.cpp
--- a/Engine/Source/Editor/ContentBrowser/Private/SAssetPicker.cpp
+++ b/Engine/Source/Editor/ContentBrowser/Private/SAssetPicker.cpp
@@ -11,8 +11,6 @@ SAssetPicker::~SAssetPicker()
 
 void SAssetPicker::Construct( const FArguments& InArgs )
 {
-	TSharedPtr<AssetFilterCollectionType> FrontendFilters = MakeShareable(new AssetFilterCollectionType());
-
 	BindCommands();
 
 	OnAssetsActivated = InArgs._AssetPickerConfig.OnAssetsActivated;
@@ -22,6 +20,7 @@ void SAssetPicker::Construct( const FArguments& InArgs )
 	bPendingFocusNextFrame = InArgs._AssetPickerConfig.bFocusSearchBoxWhenOpened;
 	DefaultFilterMenuExpansion = InArgs._AssetPickerConfig.DefaultFilterMenuExpansion;
 	SaveSettingsName = InArgs._AssetPickerConfig.SaveSettingsName;
+	OnFolderEnteredDelegate = InArgs._AssetPickerConfig.OnFolderEntered;
 
 	for (auto DelegateIt = InArgs._AssetPickerConfig.GetCurrentSelectionDelegates.CreateConstIterator(); DelegateIt; ++DelegateIt)
 	{
@@ -57,12 +56,13 @@ void SAssetPicker::Construct( const FArguments& InArgs )
 	TAttribute< FText > HighlightText;
 	EThumbnailLabel::Type ThumbnailLabel = InArgs._AssetPickerConfig.ThumbnailLabel;
 
+	FrontendFilters = MakeShareable(new AssetFilterCollectionType());
+
 	// Search box
 	if (!InArgs._AssetPickerConfig.bAutohideSearchBar)
 	{
 		TextFilter = MakeShareable( new FFrontendFilter_Text() );
 		TextFilter->SetIncludeClassName(InArgs._AssetPickerConfig.Filter.ClassNames.Num() != 1);
-		FrontendFilters->Add( TextFilter );
 		HighlightText = TAttribute< FText >( this, &SAssetPicker::GetHighlightedText );
 
 		OtherDevelopersFilter = MakeShareable( new FFrontendFilter_ShowOtherDevelopers(nullptr) );
@@ -211,7 +211,8 @@ void SAssetPicker::Construct( const FArguments& InArgs )
 		.OnAssetTagWantsToBeDisplayed(InArgs._AssetPickerConfig.OnAssetTagWantsToBeDisplayed)
 		.AllowDragging( InArgs._AssetPickerConfig.bAllowDragging )
 		.CanShowClasses( InArgs._AssetPickerConfig.bCanShowClasses )
-		.CanShowFolders( false )
+		.CanShowFolders( InArgs._AssetPickerConfig.bCanShowFolders )
+		.FilterRecursivelyWithBackendFilter( false )
 		.CanShowRealTimeThumbnails( InArgs._AssetPickerConfig.bCanShowRealTimeThumbnails )
 		.CanShowDevelopersFolder( InArgs._AssetPickerConfig.bCanShowDevelopersFolder )
 		.PreloadAssetsForContextMenu( InArgs._AssetPickerConfig.bPreloadAssetsForContextMenu )
@@ -219,6 +220,7 @@ void SAssetPicker::Construct( const FArguments& InArgs )
 		.ThumbnailLabel( ThumbnailLabel )
 		.AssetShowWarningText( InArgs._AssetPickerConfig.AssetShowWarningText)
 		.AllowFocusOnSync(false)	// Stop the asset view from stealing focus (we're in control of that)
+		.OnPathSelected(this, &SAssetPicker::FolderEntered)
 	];
 
 	LoadSettings();
@@ -275,19 +277,44 @@ FReply SAssetPicker::OnKeyDown(const FGeometry& MyGeometry, const FKeyboardEvent
 	return FReply::Unhandled();
 }
 
+void SAssetPicker::FolderEntered(const FString& FolderPath)
+{
+	CurrentSourcesData.PackagePaths.Empty();
+	CurrentSourcesData.PackagePaths.Add(FName(*FolderPath));
+
+	AssetViewPtr->SetSourcesData(CurrentSourcesData);
+
+	OnFolderEnteredDelegate.ExecuteIfBound(FolderPath);
+}
+
 FText SAssetPicker::GetHighlightedText() const
 {
 	return TextFilter->GetRawFilterText();
 }
 
+void SAssetPicker::SetSearchBoxText(const FText& InSearchText)
+{
+	TextFilter->SetRawFilterText(InSearchText);
+	if (InSearchText.IsEmpty())
+	{
+		FrontendFilters->Remove(TextFilter);
+		AssetViewPtr->SetUserSearching(false);
+	}
+	else
+	{
+		FrontendFilters->Add(TextFilter);
+		AssetViewPtr->SetUserSearching(true);
+	}
+}
+
 void SAssetPicker::OnSearchBoxChanged(const FText& InSearchText)
 {
-	TextFilter->SetRawFilterText( InSearchText );
+	SetSearchBoxText( InSearchText );
 }
 
 void SAssetPicker::OnSearchBoxCommitted(const FText& InSearchText, ETextCommit::Type CommitInfo)
 {
-	TextFilter->SetRawFilterText( InSearchText );
+	SetSearchBoxText( InSearchText );
 
 	if (CommitInfo == ETextCommit::OnEnter)
 	{
diff --git a/Engine/Source/Editor/ContentBrowser/Private/SAssetPicker.h b/Engine/Source/Editor/ContentBrowser/Private/SAssetPicker.h
--- a/Engine/Source/Editor/ContentBrowser/Private/SAssetPicker.h
+++ b/Engine/Source/Editor/ContentBrowser/Private/SAssetPicker.h
@@ -26,6 +26,11 @@ public:
 	// End of SWidget implementation
 
 private:
+	void FolderEntered(const FString& FolderPath);
+
+	/** Called when the editable text needs to be set or cleared */
+	void SetSearchBoxText(const FText& InSearchText);
+
 	/** Called by the editable text control when the search text is changed by the user */
 	void OnSearchBoxChanged(const FText& InSearchText);
 
@@ -81,6 +86,10 @@ private:
 	void SaveSettings() const;
 
 private:
+
+	/** The list of FrontendFilters currently applied to the asset view */
+	TSharedPtr<AssetFilterCollectionType> FrontendFilters;
+
 	/** The asset view widget */
 	TSharedPtr<SAssetView> AssetViewPtr;
 
@@ -102,6 +111,9 @@ private:
 	/** Called when any number of assets are activated */
 	FOnAssetsActivated OnAssetsActivated;
 
+	/** Called when a folder is entered in the asset view */
+	FOnPathSelected OnFolderEnteredDelegate;
+
 	/** True if the search box will take keyboard focus next frame */
 	bool bPendingFocusNextFrame;
 
diff --git a/Engine/Source/Editor/ContentBrowser/Private/SAssetView.cpp b/Engine/Source/Editor/ContentBrowser/Private/SAssetView.cpp
--- a/Engine/Source/Editor/ContentBrowser/Private/SAssetView.cpp
+++ b/Engine/Source/Editor/ContentBrowser/Private/SAssetView.cpp
@@ -124,6 +124,8 @@ void SAssetView::Construct( const FArguments& InArgs )
 	bCanShowClasses = InArgs._CanShowClasses;
 
 	bCanShowFolders = InArgs._CanShowFolders;
+
+	bFilterRecursivelyWithBackendFilter = InArgs._FilterRecursivelyWithBackendFilter;
 		
 	bCanShowRealTimeThumbnails = InArgs._CanShowRealTimeThumbnails;
 
@@ -1407,7 +1409,14 @@ void SAssetView::RefreshSourceItems()
 bool SAssetView::ShouldFilterRecursively() const
 {
 	// Quick check for conditions which force recursive filtering
-	if (bUserSearching || !BackendFilter.IsEmpty())
+	if (bUserSearching)
+	{
+		return true;
+	}
+
+	// In some cases we want to not filter recursively even if we have a backend filter (e.g. the open level window)
+	// Most of the time, bFilterRecursivelyWithBackendFilter is true
+	if ( bFilterRecursivelyWithBackendFilter && !BackendFilter.IsEmpty() )
 	{
 		return true;
 	}
diff --git a/Engine/Source/Editor/ContentBrowser/Private/SAssetView.h b/Engine/Source/Editor/ContentBrowser/Private/SAssetView.h
--- a/Engine/Source/Editor/ContentBrowser/Private/SAssetView.h
+++ b/Engine/Source/Editor/ContentBrowser/Private/SAssetView.h
@@ -23,6 +23,7 @@ public:
 		, _AllowThumbnailEditMode(false)
 		, _CanShowClasses(true)
 		, _CanShowFolders(false)
+		, _FilterRecursivelyWithBackendFilter(true)
 		, _CanShowRealTimeThumbnails(false)
 		, _CanShowDevelopersFolder(false)
 		, _PreloadAssetsForContextMenu(true)
@@ -104,13 +105,16 @@ public:
 		/** Indicates if this view is allowed to show classes */
 		SLATE_ARGUMENT( bool, CanShowClasses )
 
-		/** Indicates if the 'Show Folders' option should be visible */
+		/** Indicates if the 'Show Folders' option should be enabled or disabled */
 		SLATE_ARGUMENT( bool, CanShowFolders )
 
-		/** Indicates if the 'Real-Time Thumbnails' option should be visible */
+		/** If true, recursive filtering will be caused by applying a backend filter */
+		SLATE_ARGUMENT( bool, FilterRecursivelyWithBackendFilter )
+
+		/** Indicates if the 'Real-Time Thumbnails' option should be enabled or disabled */
 		SLATE_ARGUMENT( bool, CanShowRealTimeThumbnails )
 
-		/** Indicates if the 'Show Developers' option should be visible */
+		/** Indicates if the 'Show Developers' option should be enabled or disabled */
 		SLATE_ARGUMENT( bool, CanShowDevelopersFolder )
 
 		/** Indicates if the context menu is going to load the assets, and if so to preload before the context menu is shown, and warn about the pending load. */
@@ -737,13 +741,16 @@ private:
 	/** Indicates if this view is allowed to show classes */
 	bool bCanShowClasses;
 
-	/** Indicates if the 'Show Folders' option should be visible */
+	/** Indicates if the 'Show Folders' option should be enabled or disabled */
 	bool bCanShowFolders;
 
-	/** Indicates if the 'Real-Time Thumbnails' option should be visible */
+	/** If true, recursive filtering will be caused by applying a backend filter */
+	bool bFilterRecursivelyWithBackendFilter;
+
+	/** Indicates if the 'Real-Time Thumbnails' option should be enabled or disabled */
 	bool bCanShowRealTimeThumbnails;
 
-	/** Indicates if the 'Show Developers' option should be visible */
+	/** Indicates if the 'Show Developers' option should be enabled or disabled */
 	bool bCanShowDevelopersFolder;
 
 	/** Indicates if the context menu is going to load the assets, and if so to preload before the context menu is shown, and warn about the pending load. */
diff --git a/Engine/Source/Editor/ContentBrowser/Private/SPathPicker.cpp b/Engine/Source/Editor/ContentBrowser/Private/SPathPicker.cpp
--- a/Engine/Source/Editor/ContentBrowser/Private/SPathPicker.cpp
+++ b/Engine/Source/Editor/ContentBrowser/Private/SPathPicker.cpp
@@ -8,7 +8,14 @@
 
 void SPathPicker::Construct( const FArguments& InArgs )
 {
-	TSharedPtr<SPathView> PathViewPtr;
+	for (auto DelegateIt = InArgs._PathPickerConfig.SetPathsDelegates.CreateConstIterator(); DelegateIt; ++DelegateIt)
+	{
+		if ((*DelegateIt) != NULL)
+		{
+			(**DelegateIt) = FSetPathPickerPathsDelegate::CreateSP(this, &SPathPicker::SetPaths);
+		}
+	}
+
 	ChildSlot
 	[
 		SAssignNew(PathViewPtr, SPathView)
@@ -74,5 +81,10 @@ void SPathPicker::CreateNewFolder(FString FolderPath, FOnCreateNewFolder InOnCre
 	InOnCreateNewFolder.ExecuteIfBound(DefaultFolderName.ToString(), FolderPath);
 }
 
+void SPathPicker::SetPaths(const TArray<FString>& NewPaths)
+{
+	PathViewPtr->SetSelectedPaths(NewPaths);
+}
+
 
 #undef LOCTEXT_NAMESPACE
\ No newline at end of file
diff --git a/Engine/Source/Editor/ContentBrowser/Private/SPathPicker.h b/Engine/Source/Editor/ContentBrowser/Private/SPathPicker.h
--- a/Engine/Source/Editor/ContentBrowser/Private/SPathPicker.h
+++ b/Engine/Source/Editor/ContentBrowser/Private/SPathPicker.h
@@ -25,4 +25,12 @@ private:
 
 	/** Handler for creating a new folder in the path picker */
 	void CreateNewFolder(FString FolderPath, FOnCreateNewFolder InOnCreateNewFolder);
+
+	/** Sets the selected paths in this picker */
+	void SetPaths(const TArray<FString>& NewPaths);
+
+private:
+
+	/** The path view in this picker */
+	TSharedPtr<SPathView> PathViewPtr;
 };
\ No newline at end of file
diff --git a/Engine/Source/Editor/ContentBrowser/Public/ContentBrowserDelegates.h b/Engine/Source/Editor/ContentBrowser/Public/ContentBrowserDelegates.h
--- a/Engine/Source/Editor/ContentBrowser/Public/ContentBrowserDelegates.h
+++ b/Engine/Source/Editor/ContentBrowser/Public/ContentBrowserDelegates.h
@@ -31,6 +31,9 @@ DECLARE_DELEGATE_OneParam(FSyncToAssetsDelegate, const TArray< FAssetData >& /*A
 /** Called to set a new filter for an existing asset picker */
 DECLARE_DELEGATE_OneParam(FSetARFilterDelegate, const FARFilter& /*NewFilter*/);
 
+/** A pointer to an existing delegate that, when executed, will set the filter an the asset picker after it is created. */
+DECLARE_DELEGATE_OneParam(FSetPathPickerPathsDelegate, const TArray<FString>& /*NewPaths*/);
+
 /** Called to adjust the selection from the current assetdata, should be +1 to increment or -1 to decrement */
 DECLARE_DELEGATE_OneParam( FAdjustSelectionDelegate, const int32 /*direction*/ );
 
diff --git a/Engine/Source/Editor/ContentBrowser/Public/IContentBrowserSingleton.h b/Engine/Source/Editor/ContentBrowser/Public/IContentBrowserSingleton.h
--- a/Engine/Source/Editor/ContentBrowser/Public/IContentBrowserSingleton.h
+++ b/Engine/Source/Editor/ContentBrowser/Public/IContentBrowserSingleton.h
@@ -62,6 +62,9 @@ struct FAssetPickerConfig
 	/** The delegate that fires when an asset was selected */
 	FOnAssetSelected OnAssetSelected;
 
+	/** The delegate that fires when a folder was double clicked */
+	FOnPathSelected OnFolderEntered;
+
 	/** The delegate that fires when an asset is double clicked */
 	FOnAssetDoubleClicked OnAssetDoubleClicked;
 
@@ -114,10 +117,13 @@ struct FAssetPickerConfig
 	/** Indicates if this view is allowed to show classes */
 	bool bCanShowClasses;
 
-	/** Indicates if the 'Real-Time Thumbnails' option should be visible */
+	/** Indicates if the 'Show folders' option should be enabled or disabled */
+	bool bCanShowFolders;
+
+	/** Indicates if the 'Real-Time Thumbnails' option should be enabled or disabled */
 	bool bCanShowRealTimeThumbnails;
 
-	/** Indicates if the 'Show Developers' option should be visible */
+	/** Indicates if the 'Show Developers' option should be enabled or disabled */
 	bool bCanShowDevelopersFolder;
 
 	/** Indicates if the context menu is going to load the assets, and if so to preload before the context menu is shown, and warn about the pending load. */
@@ -140,6 +146,7 @@ struct FAssetPickerConfig
 		, bAutohideSearchBar(false)
 		, bAllowDragging(true)
 		, bCanShowClasses(true)
+		, bCanShowFolders(false)
 		, bCanShowRealTimeThumbnails(false)
 		, bCanShowDevelopersFolder(false)
 		, bPreloadAssetsForContextMenu(true)
@@ -160,6 +167,9 @@ struct FPathPickerConfig
 	/** The delegate that fires when a path is right clicked and a context menu is requested */
 	FContentBrowserMenuExtender_SelectedPaths OnGetPathContextMenuExtender;
 
+	/** A pointer to an existing delegate that, when executed, will set the paths for the path picker after it is created. */
+	TArray<FSetPathPickerPathsDelegate*> SetPathsDelegates;
+
 	/** If true, the search box will gain focus when the path picker is created */
 	bool bFocusSearchBoxWhenOpened;
 
