UE4: The name field of the save-asset-as dialog is now focused and selected initially after opening.

diff --git a/Engine/Source/Editor/ContentBrowser/Private/SAssetDialog.cpp b/Engine/Source/Editor/ContentBrowser/Private/SAssetDialog.cpp
--- a/Engine/Source/Editor/ContentBrowser/Private/SAssetDialog.cpp
+++ b/Engine/Source/Editor/ContentBrowser/Private/SAssetDialog.cpp
@@ -9,6 +9,7 @@ SAssetDialog::SAssetDialog()
 	: DialogType(EAssetDialogType::Open)
 	, ExistingAssetPolicy(ESaveAssetDialogExistingAssetPolicy::Disallow)
 	, bLastInputValidityCheckSuccessful(false)
+	, bPendingFocusNextFrame(true)
 {
 }
 
@@ -52,6 +53,7 @@ void SAssetDialog::Construct(const FArguments& InArgs, const FSharedAssetDialogC
 		PathPickerConfig.bAllowContextMenu = false;
 		ConfirmButtonText = LOCTEXT("AssetDialogOpenButton", "Open");
 		AssetPickerConfig.SelectionMode = OpenAssetConfig.bAllowMultipleSelection ? ESelectionMode::Multi : ESelectionMode::Single;
+		AssetPickerConfig.bFocusSearchBoxWhenOpened = true;
 		bIncludeNameBox = false;
 	}
 	else if ( DialogType == EAssetDialogType::Save )
@@ -60,6 +62,7 @@ void SAssetDialog::Construct(const FArguments& InArgs, const FSharedAssetDialogC
 		PathPickerConfig.bAllowContextMenu = true;
 		ConfirmButtonText = LOCTEXT("AssetDialogSaveButton", "Save");
 		AssetPickerConfig.SelectionMode = ESelectionMode::Single;
+		AssetPickerConfig.bFocusSearchBoxWhenOpened = false;
 		bIncludeNameBox = true;
 		ExistingAssetPolicy = SaveAssetConfig.ExistingAssetPolicy;
 		SetCurrentlyEnteredAssetName(SaveAssetConfig.DefaultAssetName);
@@ -145,10 +148,11 @@ void SAssetDialog::Construct(const FArguments& InArgs, const FSharedAssetDialogC
 			.VAlign(VAlign_Center)
 			.Padding(4, 3)
 			[
-				SNew(SEditableTextBox)
+				SAssignNew(NameEditableText, SEditableTextBox)
 				.Text(this, &SAssetDialog::GetAssetNameText)
 				.OnTextCommitted(this, &SAssetDialog::OnAssetNameTextCommited)
 				.OnTextChanged(this, &SAssetDialog::OnAssetNameTextCommited, ETextCommit::Default)
+				.SelectAllTextWhenFocused(true)
 			];
 	}
 
@@ -203,6 +207,17 @@ FReply SAssetDialog::OnKeyDown( const FGeometry& MyGeometry, const FKeyboardEven
 	return SCompoundWidget::OnKeyDown(MyGeometry, InKeyboardEvent);
 }
 
+void SAssetDialog::Tick(const FGeometry& AllottedGeometry, const double InCurrentTime, const float InDeltaTime)
+{
+	if (bPendingFocusNextFrame)
+	{
+		FocusNameBox();
+		bPendingFocusNextFrame = false;
+	}
+
+	SCompoundWidget::Tick(AllottedGeometry, InCurrentTime, InDeltaTime);
+}
+
 void SAssetDialog::SetOnAssetsChosenForOpen(const FOnAssetsChosenForOpen& InOnAssetsChosenForOpen)
 {
 	OnAssetsChosenForOpen = InOnAssetsChosenForOpen;
@@ -213,6 +228,14 @@ void SAssetDialog::SetOnObjectPathChosenForSave(const FOnObjectPathChosenForSave
 	OnObjectPathChosenForSave = InOnObjectPathChosenForSave;
 }
 
+void SAssetDialog::FocusNameBox()
+{
+	if ( NameEditableText.IsValid() )
+	{
+		FSlateApplication::Get().SetKeyboardFocus(NameEditableText.ToSharedRef(), EKeyboardFocusCause::SetDirectly);
+	}
+}
+
 FText SAssetDialog::GetAssetNameText() const
 {
 	return FText::FromString(CurrentlyEnteredAssetName);
diff --git a/Engine/Source/Editor/ContentBrowser/Private/SAssetDialog.h b/Engine/Source/Editor/ContentBrowser/Private/SAssetDialog.h
--- a/Engine/Source/Editor/ContentBrowser/Private/SAssetDialog.h
+++ b/Engine/Source/Editor/ContentBrowser/Private/SAssetDialog.h
@@ -14,6 +14,7 @@ public:
 	virtual void Construct(const FArguments& InArgs, const FSharedAssetDialogConfig& InConfig);
 
 	virtual FReply OnKeyDown( const FGeometry& MyGeometry, const FKeyboardEvent& InKeyboardEvent ) override;
+	void Tick(const FGeometry& AllottedGeometry, const double InCurrentTime, const float InDeltaTime) override;
 
 	/** Sets the delegate handler for when an open operation is committed */
 	void SetOnAssetsChosenForOpen(const FOnAssetsChosenForOpen& InOnAssetsChosenForOpen);
@@ -23,6 +24,9 @@ public:
 
 private:
 
+	/** Moves keyboard focus to the name box if this is a save dialog */
+	void FocusNameBox();
+
 	/** Gets the name to display in the asset name box */
 	FText GetAssetNameText() const;
 
@@ -99,6 +103,9 @@ private:
 	/** The assets that are currently selected in the asset picker */
 	TArray<FAssetData> CurrentlySelectedAssets;
 
+	/** The name box. Only used in save dialogs. */
+	TSharedPtr<SEditableTextBox> NameEditableText;
+
 	/** The object path of the asset to save. Only used in save dialogs. */
 	FString CurrentlySelectedPath;
 
@@ -113,4 +120,7 @@ private:
 
 	/** True if the last validity check returned that the class name/path is valid for creation */
 	bool bLastInputValidityCheckSuccessful;
+
+	/** Used to focus the name box after opening the dialog */
+	bool bPendingFocusNextFrame;
 };
\ No newline at end of file
