UE4: Fixed a bug that caused text filters in Content Browser widgets to lose their source text when toggling whether they should include the class name or not.

diff --git a/Engine/Source/Editor/ContentBrowser/Private/FrontendFilters.cpp b/Engine/Source/Editor/ContentBrowser/Private/FrontendFilters.cpp
--- a/Engine/Source/Editor/ContentBrowser/Private/FrontendFilters.cpp
+++ b/Engine/Source/Editor/ContentBrowser/Private/FrontendFilters.cpp
@@ -49,6 +49,9 @@ void FFrontendFilter_Text::SetRawFilterText(const FText& InFilterText)
 
 void FFrontendFilter_Text::SetIncludeClassName(bool bIncludeClassName)
 {
+	// Preserve the existing text. The filter is going to get recreated here.
+	const FText ExistingText = TextFilter.GetRawFilterText();
+
 	if ( bIncludeClassName )
 	{
 		TextFilter = TTextFilter<AssetFilterType>::FItemToStringArray::CreateStatic( &AssetDataToExportTextString );
@@ -58,6 +61,9 @@ void FFrontendFilter_Text::SetIncludeClassName(bool bIncludeClassName)
 		TextFilter = TTextFilter<AssetFilterType>::FItemToStringArray::CreateStatic( &AssetDataToObjectPathString );
 	}
 
+	// Apply the existing text before we re-assign the delegate since we want the text preservation to be opaque.
+	TextFilter.SetRawFilterText(ExistingText);
+
 	TextFilter.OnChanged().AddRaw(this, &FFrontendFilter_Text::HandleOnChangedEvent);
 }
 
