Refined text rendering on content browser thumbnails and reverted back to full-width sizing method.

Tiny text is now using a lighter font with a better hinting algorithm, and is centrally justified.
We want to minimise waste of horizontal space, so content browser thumbnails now fill the width of the content browser once again.

diff --git a/Engine/Source/Editor/ContentBrowser/Private/AssetViewWidgets.cpp b/Engine/Source/Editor/ContentBrowser/Private/AssetViewWidgets.cpp
--- a/Engine/Source/Editor/ContentBrowser/Private/AssetViewWidgets.cpp
+++ b/Engine/Source/Editor/ContentBrowser/Private/AssetViewWidgets.cpp
@@ -1146,7 +1146,7 @@ void SAssetTileItem::Construct( const FArguments& InArgs )
 		);
 
 	AssetThumbnail = InArgs._AssetThumbnail;
-	ThumbnailSize = InArgs._ThumbnailSize;
+	ItemWidth = InArgs._ItemWidth;
 	ThumbnailPadding = IsFolder() ? InArgs._ThumbnailPadding + 5.0f : InArgs._ThumbnailPadding;
 
 	TSharedPtr<SWidget> Thumbnail;
@@ -1262,7 +1262,7 @@ void SAssetTileItem::Construct( const FArguments& InArgs )
 			]
 
 			+SVerticalBox::Slot()
-			.Padding(FMargin(2.f, 0, 2.f, ThumbnailPadding))
+			.Padding(FMargin(1.f, 0))
 			.HAlign(HAlign_Center)
 			.VAlign(VAlign_Center)
 			.FillHeight(1.f)
@@ -1277,6 +1277,7 @@ void SAssetTileItem::Construct( const FArguments& InArgs )
 					.IsSelected(InArgs._IsSelected)
 					.IsReadOnly(ThumbnailEditMode)
 					.WrapTextAt(this, &SAssetTileItem::GetNameTextWrapWidth)
+					.Justification(ETextJustify::Center)
 					.LineBreakPolicy(FBreakIterator::CreateCamelCaseBreakIterator())
 			]
 		]
@@ -1304,14 +1305,14 @@ void SAssetTileItem::OnAssetDataChanged()
 	SetToolTip( CreateToolTipWidget() );
 }
 
-FOptionalSize SAssetTileItem::GetSCCImageSize() const
+FOptionalSize SAssetTileItem::GetThumbnailBoxSize() const
 {
-	return GetThumbnailBoxSize().Get() * 0.2;
+	return FOptionalSize( ItemWidth.Get() );
 }
 
-FOptionalSize SAssetTileItem::GetThumbnailBoxSize() const
+FOptionalSize SAssetTileItem::GetSCCImageSize() const
 {
-	return ThumbnailSize.Get(16);
+	return GetThumbnailBoxSize().Get() * 0.2;
 }
 
 FSlateFontInfo SAssetTileItem::GetThumbnailFont() const
@@ -1334,7 +1335,7 @@ FSlateFontInfo SAssetTileItem::GetThumbnailFont() const
 float SAssetTileItem::GetNameTextWrapWidth() const
 {
 	// Wrap to the entire size of the tile, minus some padding
-	return LastGeometry.Size.X - 4.f;
+	return LastGeometry.Size.X - 2.f;
 }
 
 ///////////////////////////////
diff --git a/Engine/Source/Editor/ContentBrowser/Private/AssetViewWidgets.h b/Engine/Source/Editor/ContentBrowser/Private/AssetViewWidgets.h
--- a/Engine/Source/Editor/ContentBrowser/Private/AssetViewWidgets.h
+++ b/Engine/Source/Editor/ContentBrowser/Private/AssetViewWidgets.h
@@ -379,7 +379,7 @@ public:
 		, _ThumbnailLabel( EThumbnailLabel::ClassName )
 		, _ThumbnailHintColorAndOpacity( FLinearColor( 0.0f, 0.0f, 0.0f, 0.0f ) )
 		, _AllowThumbnailHintLabel(true)
-		, _ThumbnailSize(16)
+		, _ItemWidth(16)
 		, _ShouldAllowToolTip(true)
 		, _ThumbnailEditMode(false)
 	{}
@@ -402,8 +402,8 @@ public:
 		/** Whether the thumbnail should ever show it's hint label */
 		SLATE_ARGUMENT( bool, AllowThumbnailHintLabel )
 
-		/** The width of the thumbnail */
-		SLATE_ATTRIBUTE( float, ThumbnailSize )
+		/** The width of the item */
+		SLATE_ATTRIBUTE( float, ItemWidth )
 
 		/** Delegate for when an asset name has entered a rename state */
 		SLATE_EVENT( FOnRenameBegin, OnRenameBegin )
@@ -473,7 +473,7 @@ private:
 	TSharedPtr<FAssetThumbnail> AssetThumbnail;
 
 	/** The width of the item. Used to enforce a square thumbnail. */
-	TAttribute<float> ThumbnailSize;
+	TAttribute<float> ItemWidth;
 
 	/** The padding allotted for the thumbnail */
 	float ThumbnailPadding;
diff --git a/Engine/Source/Editor/ContentBrowser/Private/SAssetView.cpp b/Engine/Source/Editor/ContentBrowser/Private/SAssetView.cpp
--- a/Engine/Source/Editor/ContentBrowser/Private/SAssetView.cpp
+++ b/Engine/Source/Editor/ContentBrowser/Private/SAssetView.cpp
@@ -110,7 +110,7 @@ void SAssetView::Construct( const FArguments& InArgs )
 	TileViewThumbnailResolution = 256;
 	TileViewThumbnailSize = 128;
 	TileViewThumbnailPadding = 5;
-	TileViewNameHeight = 32;
+	TileViewNameHeight = 36;
 	ThumbnailScaleSliderValue = InArgs._ThumbnailScale; 
 
 	if ( !ThumbnailScaleSliderValue.IsBound() )
@@ -118,7 +118,7 @@ void SAssetView::Construct( const FArguments& InArgs )
 		ThumbnailScaleSliderValue = FMath::Clamp<float>(ThumbnailScaleSliderValue.Get(), 0.0f, 1.0f);
 	}
 
-	MinThumbnailScale = 0.3f * ThumbnailScaleRangeScalar;
+	MinThumbnailScale = 0.4f * ThumbnailScaleRangeScalar;
 	MaxThumbnailScale = 2.0f * ThumbnailScaleRangeScalar;
 
 	bCanShowClasses = InArgs._CanShowClasses;
@@ -786,8 +786,6 @@ void SAssetView::Tick(const FGeometry& AllottedGeometry, const double InCurrentT
 
 	CalculateFillScale( AllottedGeometry );
 
-	TileViewNameHeight = int(FMath::Min(GetTileViewThumbnailSize(), 32.f));
-
 	CurrentTime = InCurrentTime;
 
 	// If there were any assets that were recently added via the asset registry, process them now
@@ -2674,7 +2672,7 @@ TSharedRef<ITableRow> SAssetView::MakeTileViewWidget(TSharedPtr<FAssetViewItem>
 		TSharedRef<SAssetTileItem> Item =
 			SNew(SAssetTileItem)
 			.AssetItem(AssetItem)
-			.ThumbnailSize(this, &SAssetView::GetTileViewThumbnailSize)
+			.ItemWidth(this, &SAssetView::GetTileViewItemWidth)
 			.OnRenameBegin(this, &SAssetView::AssetRenameBegin)
 			.OnRenameCommit(this, &SAssetView::AssetRenameCommit)
 			.OnVerifyRenameCommit(this, &SAssetView::AssetVerifyRenameCommit)
@@ -2719,7 +2717,7 @@ TSharedRef<ITableRow> SAssetView::MakeTileViewWidget(TSharedPtr<FAssetViewItem>
 			.AssetThumbnail(AssetThumbnail)
 			.AssetItem(AssetItem)
 			.ThumbnailPadding(TileViewThumbnailPadding)
-			.ThumbnailSize(this, &SAssetView::GetTileViewThumbnailSize)
+			.ItemWidth(this, &SAssetView::GetTileViewItemWidth)
 			.OnRenameBegin(this, &SAssetView::AssetRenameBegin)
 			.OnRenameCommit(this, &SAssetView::AssetRenameCommit)
 			.OnVerifyRenameCommit(this, &SAssetView::AssetVerifyRenameCommit)
@@ -3486,12 +3484,12 @@ float SAssetView::GetListViewItemHeight() const
 
 float SAssetView::GetTileViewItemHeight() const
 {
-	return TileViewNameHeight + GetTileViewThumbnailSize();
+	return TileViewNameHeight + GetTileViewItemBaseHeight() * FillScale;
 }
 
 float SAssetView::GetTileViewItemBaseHeight() const
 {
-	return TileViewThumbnailSize * FMath::Lerp(MinThumbnailScale, MaxThumbnailScale, GetThumbnailScale()) + TileViewThumbnailPadding * 2;
+	return (TileViewThumbnailSize + TileViewThumbnailPadding * 2) * FMath::Lerp(MinThumbnailScale, MaxThumbnailScale, GetThumbnailScale());
 }
 
 float SAssetView::GetTileViewItemWidth() const
@@ -3501,12 +3499,7 @@ float SAssetView::GetTileViewItemWidth() const
 
 float SAssetView::GetTileViewItemBaseWidth() const
 {
-	return FMath::Max(75.f, GetTileViewThumbnailSize() + TileViewThumbnailPadding * 2);
-}
-
-float SAssetView::GetTileViewThumbnailSize() const
-{
-	return TileViewThumbnailSize * FMath::Lerp( MinThumbnailScale, MaxThumbnailScale, GetThumbnailScale() );
+	return ( TileViewThumbnailSize + TileViewThumbnailPadding * 2 ) * FMath::Lerp( MinThumbnailScale, MaxThumbnailScale, GetThumbnailScale() );
 }
 
 EColumnSortMode::Type SAssetView::GetColumnSortMode(const FName ColumnId) const
diff --git a/Engine/Source/Editor/ContentBrowser/Private/SAssetView.h b/Engine/Source/Editor/ContentBrowser/Private/SAssetView.h
--- a/Engine/Source/Editor/ContentBrowser/Private/SAssetView.h
+++ b/Engine/Source/Editor/ContentBrowser/Private/SAssetView.h
@@ -503,7 +503,6 @@ private:
 
 	/** Gets the scaled item width for the tile view before the filler scale is applied */
 	float GetTileViewItemBaseWidth() const;
-	float GetTileViewThumbnailSize() const;
 
 	/** Gets the sort mode for the supplied ColumnId */
 	EColumnSortMode::Type GetColumnSortMode(const FName ColumnId) const;
diff --git a/Engine/Source/Editor/EditorStyle/Private/SlateEditorStyle.cpp b/Engine/Source/Editor/EditorStyle/Private/SlateEditorStyle.cpp
--- a/Engine/Source/Editor/EditorStyle/Private/SlateEditorStyle.cpp
+++ b/Engine/Source/Editor/EditorStyle/Private/SlateEditorStyle.cpp
@@ -4769,8 +4769,8 @@ void FSlateEditorStyle::FStyle::SetupContentBrowserStyle()
 		);
 
 		// Tile view
-		Set( "ContentBrowser.AssetTileViewNameFont", TTF_CORE_FONT( "Fonts/Roboto-Regular", 8 ) );
-		Set( "ContentBrowser.AssetTileViewNameFontSmall", TTF_CORE_FONT( "Fonts/Roboto-Regular", 7 ) );
+		Set( "ContentBrowser.AssetTileViewNameFont", TTF_CORE_FONT( "Fonts/Roboto-Regular", 9 ) );
+		Set( "ContentBrowser.AssetTileViewNameFontSmall", TTF_CORE_FONT( "Fonts/Roboto-Light", 8, EFontHinting::Auto ) );
 		Set( "ContentBrowser.AssetTileViewNameFontDirty", TTF_CORE_FONT( "Fonts/Roboto-Bold", 10 ) );
 		Set( "ContentBrowser.AssetListView.TableRow", FTableRowStyle()
 			.SetEvenRowBackgroundBrush( FSlateNoResource() )
diff --git a/Engine/Source/Editor/UnrealEd/Private/AssetThumbnail.cpp b/Engine/Source/Editor/UnrealEd/Private/AssetThumbnail.cpp
--- a/Engine/Source/Editor/UnrealEd/Private/AssetThumbnail.cpp
+++ b/Engine/Source/Editor/UnrealEd/Private/AssetThumbnail.cpp
@@ -52,6 +52,7 @@ public:
 		AssetThumbnail = InArgs._AssetThumbnail;
 		bHasRenderedThumbnail = false;
 		WidthLastFrame = 0;
+		GenericThumbnailBorderPadding = 2.f;
 
 		AssetThumbnail->OnAssetDataChanged().AddSP(this, &SAssetThumbnail::OnAssetDataChanged);
 
@@ -92,7 +93,7 @@ public:
 			SNew(SBorder)
 			.BorderImage( this, &SAssetThumbnail::GetBackgroundBrush )
 			.BorderBackgroundColor( this, &SAssetThumbnail::GetAssetColor )
-			.Padding( 0 )
+			.Padding( GenericThumbnailBorderPadding )
 			.VAlign(VAlign_Center)
 			.HAlign(HAlign_Center)
 			.Visibility(this, &SAssetThumbnail::GetClassThumbnailVisibility)
@@ -108,7 +109,7 @@ public:
 			SNew(SBorder)
 			.BorderImage( this, &SAssetThumbnail::GetBackgroundBrush )
 			.BorderBackgroundColor( this, &SAssetThumbnail::GetAssetColor )
-			.Padding( 0 )
+			.Padding( GenericThumbnailBorderPadding )
 			.VAlign(VAlign_Center) .HAlign(HAlign_Center)
 			.Visibility(this, &SAssetThumbnail::GetGenericThumbnailVisibility)
 			[
@@ -301,8 +302,7 @@ private:
 
 	float GetTextWrapWidth() const
 	{
-		/** 2px padding either side of the text */
-		return WidthLastFrame;// - 4.f;
+		return WidthLastFrame - GenericThumbnailBorderPadding * 2.f;
 	}
 
 	const FSlateBrush* GetBackgroundBrush() const
@@ -551,6 +551,7 @@ private:
 
 	FLinearColor AssetColor;
 	float WidthLastFrame;
+	float GenericThumbnailBorderPadding;
 	bool bHasRenderedThumbnail;
 	FName Style;
 	TAttribute< FText > HighlightedText;
