UE4: No longer re-applying a text filter to asset view search results when the old filter is the same as the new. This was causing a re-sort when the user loses focus from the search bar in the content browser or asset pickers. Also initialized some uninitialized sorting bools.

diff --git a/Engine/Source/Editor/ContentBrowser/Private/SAssetPicker.cpp b/Engine/Source/Editor/ContentBrowser/Private/SAssetPicker.cpp
--- a/Engine/Source/Editor/ContentBrowser/Private/SAssetPicker.cpp
+++ b/Engine/Source/Editor/ContentBrowser/Private/SAssetPicker.cpp
@@ -294,16 +294,19 @@ FText SAssetPicker::GetHighlightedText() const
 
 void SAssetPicker::SetSearchBoxText(const FText& InSearchText)
 {
-	TextFilter->SetRawFilterText(InSearchText);
-	if (InSearchText.IsEmpty())
+	if ( !InSearchText.EqualToCaseIgnored(TextFilter->GetRawFilterText()) )
 	{
-		FrontendFilters->Remove(TextFilter);
-		AssetViewPtr->SetUserSearching(false);
-	}
-	else
-	{
-		FrontendFilters->Add(TextFilter);
-		AssetViewPtr->SetUserSearching(true);
+		TextFilter->SetRawFilterText(InSearchText);
+		if (InSearchText.IsEmpty())
+		{
+			FrontendFilters->Remove(TextFilter);
+			AssetViewPtr->SetUserSearching(false);
+		}
+		else
+		{
+			FrontendFilters->Add(TextFilter);
+			AssetViewPtr->SetUserSearching(true);
+		}
 	}
 }
 
diff --git a/Engine/Source/Editor/ContentBrowser/Private/SAssetView.cpp b/Engine/Source/Editor/ContentBrowser/Private/SAssetView.cpp
--- a/Engine/Source/Editor/ContentBrowser/Private/SAssetView.cpp
+++ b/Engine/Source/Editor/ContentBrowser/Private/SAssetView.cpp
@@ -176,6 +176,8 @@ void SAssetView::Construct( const FArguments& InArgs )
 	}
 
 	bPendingSortFilteredItems = false;
+	bQuickFrontendListRefreshRequested = false;
+	bSlowFullListRefreshRequested = false;
 	LastSortTime = 0;
 	SortDelaySeconds = 8;
 
diff --git a/Engine/Source/Editor/ContentBrowser/Private/SContentBrowser.cpp b/Engine/Source/Editor/ContentBrowser/Private/SContentBrowser.cpp
--- a/Engine/Source/Editor/ContentBrowser/Private/SContentBrowser.cpp
+++ b/Engine/Source/Editor/ContentBrowser/Private/SContentBrowser.cpp
@@ -1070,16 +1070,19 @@ void SContentBrowser::NewFolderRequested(const FString& SelectedPath)
 
 void SContentBrowser::SetSearchBoxText(const FText& InSearchText)
 {
-	TextFilter->SetRawFilterText( InSearchText );
-	if(InSearchText.IsEmpty())
+	if (!InSearchText.EqualToCaseIgnored(TextFilter->GetRawFilterText()))
 	{
-		FrontendFilters->Remove(TextFilter);
-		AssetViewPtr->SetUserSearching(false);
-	}
-	else
-	{
-		FrontendFilters->Add(TextFilter);
-		AssetViewPtr->SetUserSearching(true);
+		TextFilter->SetRawFilterText( InSearchText );
+		if(InSearchText.IsEmpty())
+		{
+			FrontendFilters->Remove(TextFilter);
+			AssetViewPtr->SetUserSearching(false);
+		}
+		else
+		{
+			FrontendFilters->Add(TextFilter);
+			AssetViewPtr->SetUserSearching(true);
+		}
 	}
 }
 
