Enum values on asset tooltips are not more readable

Enum prefixes are no longer displayed, and CamelCaseEnums are now space-separated.
This addresses TTP#345449 - NUX:Content Browser: NUX: Tooltips for Materials and Textures have weird acronyms all over them

diff --git a/Engine/Source/Editor/ContentBrowser/Private/AssetViewWidgets.cpp b/Engine/Source/Editor/ContentBrowser/Private/AssetViewWidgets.cpp
--- a/Engine/Source/Editor/ContentBrowser/Private/AssetViewWidgets.cpp
+++ b/Engine/Source/Editor/ContentBrowser/Private/AssetViewWidgets.cpp
@@ -499,6 +499,21 @@ TSharedRef<SToolTip> SAssetViewItem::CreateToolTipWidget() const
 						if (UProperty* Field = FindField<UProperty>(AssetClass, TagIt.Key()))
 						{
 							DisplayName = Field->GetDisplayNameText();
+
+							// Strip off enum prefixes if they exist
+							if (UByteProperty* ByteProperty = Cast<UByteProperty>(Field))
+							{
+								if (ByteProperty->Enum)
+								{
+									const FString EnumPrefix = ByteProperty->Enum->GenerateEnumPrefix();
+									if (EnumPrefix.Len() && ValueString.StartsWith(EnumPrefix))
+									{
+										ValueString = ValueString.RightChop(EnumPrefix.Len() + 1);	// +1 to skip over the underscore
+									}
+								}
+
+								ValueString = FName::NameToDisplayString(ValueString, false);
+							}
 						}
 						else
 						{
