Fix Persona Asset browser filter arrow colour

#TTP 344422 - PERSONA: Asset Browser Filters arrow is inconsistent

diff --git a/Engine/Source/Editor/ContentBrowser/Private/SAssetPicker.cpp b/Engine/Source/Editor/ContentBrowser/Private/SAssetPicker.cpp
--- a/Engine/Source/Editor/ContentBrowser/Private/SAssetPicker.cpp
+++ b/Engine/Source/Editor/ContentBrowser/Private/SAssetPicker.cpp
@@ -78,6 +78,7 @@ void SAssetPicker::Construct( const FArguments& InArgs )
 			[
 				SNew( SComboButton )
 				.ComboButtonStyle( FEditorStyle::Get(), "ContentBrowser.Filters.Style" )
+				.ForegroundColor(FLinearColor::White)
 				.ToolTipText( LOCTEXT( "AddFilterToolTip", "Add an asset filter." ) )
 				.OnGetMenuContent( this, &SAssetPicker::MakeAddFilterMenu )
 				.HasDownArrow( true )
