Asset context menu refactor & polish

- Duplicate now works properly in the asset view.

- Consolidated common asset actions into the asset context menu and removed the duplicate code from all asset type actions that had implemented them. The affected actions are Edit, Reimport, Find in Explorer, and Open in External Editor. The latter three actions are only visible for imported assets.

- Added IsImportedAsset() and GetSourceFilePaths() to IAssetTypeActions for use by the context menu.

- Compressed the bottom matter on the context menu by removing section headers and compressing Source Control options into a sub-menu (when connected, otherwise there's just a "connect" menu entry).

- Added lots of icons! (The missing ones are TTP'd)

- Fixed FAssetEditorManager::OpenEditorForAssets to work in a useful fashion. IAssetTypeActions::OpenAssetEditor() already took in a UObject array, so it is up to the class's asset type actions to decide how to handle editing multiple assets. (It appeared to be before, but wasn't really).

- Removed the "Details" option from the context menu (since it's redundant when there's a guaranteed "Edit" option)

- Consolidated all skeleton-related actions into AnimationEditorUtils and removed the duplicate code from FAssetTypeActions_SkeletalMesh and FAssetTypeActions_Skeleton

- Compressed the Paper2D sprite options for textures into a submenu

- Renamed "New" to "Create" in content browser

diff --git a/Engine/Plugins/2D/Paper2D/Source/Paper2DEditor/Private/ContentBrowserExtensions/ContentBrowserExtensions.cpp b/Engine/Plugins/2D/Paper2D/Source/Paper2DEditor/Private/ContentBrowserExtensions/ContentBrowserExtensions.cpp
--- a/Engine/Plugins/2D/Paper2D/Source/Paper2DEditor/Private/ContentBrowserExtensions/ContentBrowserExtensions.cpp
+++ b/Engine/Plugins/2D/Paper2D/Source/Paper2DEditor/Private/ContentBrowserExtensions/ContentBrowserExtensions.cpp
@@ -176,38 +176,49 @@ public:
 		SelectedAssetFunctor->Execute();
 	}
 
-	static void CreateSpriteTextureMenuOptions(FMenuBuilder& Builder, TSharedPtr<FCreateSpriteFromTextureExtension> SpriteCreatorFunctor)
+	static void CreateSpriteActionsSubMenu(FMenuBuilder& MenuBuilder, TSharedPtr<FCreateSpriteFromTextureExtension> SpriteCreator, TSharedPtr<FCreateSpriteFromTextureExtension> SpriteExtractor, TSharedPtr<FConfigureTexturesForSpriteUsageExtension> ConfigFunctor)
 	{
+		MenuBuilder.AddSubMenu(
+			LOCTEXT("SpriteActionsSubMenuLabel", "Sprite Actions"),
+			LOCTEXT("SpriteActionsSubMenuToolTip", "Sprite-related actions for this texture."),
+			FNewMenuDelegate::CreateStatic(&FPaperContentBrowserExtensions_Impl::PopulateSpriteActionsMenu, SpriteCreator, SpriteExtractor, ConfigFunctor),
+			false,
+			FSlateIcon(FEditorStyle::GetStyleSetName(), "ClassIcon.PaperSprite")
+		);
+	}
+
+	static void PopulateSpriteActionsMenu(FMenuBuilder& MenuBuilder, TSharedPtr<FCreateSpriteFromTextureExtension> SpriteCreator, TSharedPtr<FCreateSpriteFromTextureExtension> SpriteExtractor, TSharedPtr<FConfigureTexturesForSpriteUsageExtension> ConfigFunctor)
+	{
+		// Create sprites
 		FUIAction Action_CreateSpritesFromTextures(
-			FExecuteAction::CreateStatic(&FPaperContentBrowserExtensions_Impl::ExecuteSelectedContentFunctor, StaticCastSharedPtr<FContentBrowserSelectedAssetExtensionBase>(SpriteCreatorFunctor)));
-		Builder.AddMenuEntry(
+			FExecuteAction::CreateStatic(&FPaperContentBrowserExtensions_Impl::ExecuteSelectedContentFunctor, StaticCastSharedPtr<FContentBrowserSelectedAssetExtensionBase>(SpriteCreator)));
+		
+		MenuBuilder.AddMenuEntry(
 			LOCTEXT("CB_Extension_Texture_CreateSprite", "Create Sprite"),
 			LOCTEXT("CB_Extension_Texture_CreateSprite_Tooltip", "Create sprites from selected textures"),
-			FSlateIcon(),
+			FSlateIcon(FEditorStyle::GetStyleSetName(), "ClassIcon.PaperSprite"),
 			Action_CreateSpritesFromTextures,
 			NAME_None,
 			EUserInterfaceActionType::Button);
-	}
 
-	static void ExtractSpritesFromTextureMenuOptions(FMenuBuilder& Builder, TSharedPtr<FCreateSpriteFromTextureExtension> SpriteExtractorFunctor)
-	{
+		// Extract Sprites
 		FUIAction Action_ExtractSpritesFromTextures(
-			FExecuteAction::CreateStatic(&FPaperContentBrowserExtensions_Impl::ExecuteSelectedContentFunctor, StaticCastSharedPtr<FContentBrowserSelectedAssetExtensionBase>(SpriteExtractorFunctor)));
-		Builder.AddMenuEntry(
+			FExecuteAction::CreateStatic(&FPaperContentBrowserExtensions_Impl::ExecuteSelectedContentFunctor, StaticCastSharedPtr<FContentBrowserSelectedAssetExtensionBase>(SpriteExtractor)));
+
+		MenuBuilder.AddMenuEntry(
 			LOCTEXT("CB_Extension_Texture_ExtractSprites", "Extract Sprites"),
 			LOCTEXT("CB_Extension_Texture_ExtractSprite_Tooltip", "Detects and extracts sprites from the selected textures using transparency"),
 			FSlateIcon(),
 			Action_ExtractSpritesFromTextures,
 			NAME_None,
 			EUserInterfaceActionType::Button);
-	}
 
-	static void CreateTextureSetupMenuOption(FMenuBuilder& Builder, TSharedPtr<FConfigureTexturesForSpriteUsageExtension> ConfigureFunctor)
-	{
+		// Configure for retro sprites
 		FUIAction Action_ConfigureTexturesForSprites(
-			FExecuteAction::CreateStatic(&FPaperContentBrowserExtensions_Impl::ExecuteSelectedContentFunctor, StaticCastSharedPtr<FContentBrowserSelectedAssetExtensionBase>(ConfigureFunctor)));
-		Builder.AddMenuEntry(
-			LOCTEXT("CB_Extension_Texture_ConfigureTextureForSprites", "Configure For Retro Sprites"), 
+			FExecuteAction::CreateStatic(&FPaperContentBrowserExtensions_Impl::ExecuteSelectedContentFunctor, StaticCastSharedPtr<FContentBrowserSelectedAssetExtensionBase>(ConfigFunctor)));
+
+		MenuBuilder.AddMenuEntry(
+			LOCTEXT("CB_Extension_Texture_ConfigureTextureForSprites", "Configure For Retro Sprites"),
 			LOCTEXT("CB_Extension_Texture_ConfigureTextureForSprites_Tooltip", "Sets compression settings and sampling modes to good defaults for retro sprites (nearest filtering, uncompressed, etc...)"),
 			FSlateIcon(),
 			Action_ConfigureTexturesForSprites,
@@ -215,8 +226,6 @@ public:
 			EUserInterfaceActionType::Button);
 	}
 
-	
-
 	static TSharedRef<FExtender> OnExtendContentBrowserAssetSelectionMenu(const TArray<FAssetData>& SelectedAssets)
 	{
 		TSharedRef<FExtender> Extender(new FExtender());
@@ -231,34 +240,26 @@ public:
 
 		if (bAnyTextures)
 		{
+			// Create Sprite
 			TSharedPtr<FCreateSpriteFromTextureExtension> SpriteCreatorFunctor = MakeShareable(new FCreateSpriteFromTextureExtension());
 			SpriteCreatorFunctor->SelectedAssets = SelectedAssets;
 
-			Extender->AddMenuExtension(
-				"GetAssetActions",
-				EExtensionHook::After,
-				NULL,
-				FMenuExtensionDelegate::CreateStatic(&FPaperContentBrowserExtensions_Impl::CreateSpriteTextureMenuOptions, SpriteCreatorFunctor));
-
-			//@TODO: Already getting nasty, need to refactor the data to be independent of the functor
+			////@TODO: Already getting nasty, need to refactor the data to be independent of the functor
+			// Extract Sprites
 			TSharedPtr<FCreateSpriteFromTextureExtension> SpriteExtractorFunctor = MakeShareable(new FCreateSpriteFromTextureExtension());
 			SpriteExtractorFunctor->SelectedAssets = SelectedAssets;
 			SpriteExtractorFunctor->bExtractSprites = true;
 
-			Extender->AddMenuExtension(
-				"GetAssetActions",
-				EExtensionHook::After,
-				NULL,
-				FMenuExtensionDelegate::CreateStatic(&FPaperContentBrowserExtensions_Impl::ExtractSpritesFromTextureMenuOptions, SpriteExtractorFunctor));
-
+			// Configure for retro sprites
 			TSharedPtr<FConfigureTexturesForSpriteUsageExtension> TextureConfigFunctor = MakeShareable(new FConfigureTexturesForSpriteUsageExtension());
 			TextureConfigFunctor->SelectedAssets = SelectedAssets;
 
+			// Add the sprite actions sub-menu extender
 			Extender->AddMenuExtension(
 				"GetAssetActions",
 				EExtensionHook::After,
 				NULL,
-				FMenuExtensionDelegate::CreateStatic(&FPaperContentBrowserExtensions_Impl::CreateTextureSetupMenuOption, TextureConfigFunctor));
+				FMenuExtensionDelegate::CreateStatic(&FPaperContentBrowserExtensions_Impl::CreateSpriteActionsSubMenu, SpriteCreatorFunctor, SpriteExtractorFunctor, TextureConfigFunctor));
 		}
 
 		return Extender;
diff --git a/Engine/Plugins/2D/Paper2D/Source/Paper2DEditor/Private/SpriteAssetTypeActions.cpp b/Engine/Plugins/2D/Paper2D/Source/Paper2DEditor/Private/SpriteAssetTypeActions.cpp
--- a/Engine/Plugins/2D/Paper2D/Source/Paper2DEditor/Private/SpriteAssetTypeActions.cpp
+++ b/Engine/Plugins/2D/Paper2D/Source/Paper2DEditor/Private/SpriteAssetTypeActions.cpp
@@ -52,14 +52,14 @@ void FSpriteAssetTypeActions::GetActions(const TArray<UObject*>& InObjects, FMen
 	auto Sprites = GetTypedWeakObjectPtrs<UPaperSprite>(InObjects);
 
 	MenuBuilder.AddMenuEntry(
-		LOCTEXT("Sprite_CreateFlipbook", "Create Flipbooks"),
+		LOCTEXT("Sprite_CreateFlipbook", "Create Flipbook"),
 		LOCTEXT("Sprite_CreateFlipbookTooltip", "Creates flipbooks from the selected sprites."),
-		FSlateIcon(),
+		FSlateIcon(FEditorStyle::GetStyleSetName(), "ClassIcon.PaperFlipbook"),
 		FUIAction(
 		FExecuteAction::CreateSP(this, &FSpriteAssetTypeActions::ExecuteCreateFlipbook, Sprites),
 		FCanExecuteAction()
 		)
-		);
+	);
 }
 
 //////////////////////////////////////////////////////////////////////////
diff --git a/Engine/Plugins/2D/Paper2D/Source/PaperJsonImporter/Private/PaperSpriteSheetAssetTypeActions.cpp b/Engine/Plugins/2D/Paper2D/Source/PaperJsonImporter/Private/PaperSpriteSheetAssetTypeActions.cpp
--- a/Engine/Plugins/2D/Paper2D/Source/PaperJsonImporter/Private/PaperSpriteSheetAssetTypeActions.cpp
+++ b/Engine/Plugins/2D/Paper2D/Source/PaperJsonImporter/Private/PaperSpriteSheetAssetTypeActions.cpp
@@ -36,24 +36,31 @@ uint32 FPaperSpriteSheetAssetTypeActions::GetCategories()
 	return EAssetTypeCategories::Misc;
 }
 
+bool FPaperSpriteSheetAssetTypeActions::IsImportedAsset() const
+{
+	return true;
+}
+
+void FPaperSpriteSheetAssetTypeActions::GetResolvedSourceFilePaths(const TArray<UObject*>& TypeAssets, TArray<FString>& OutSourceFilePaths) const
+{
+	for (auto& Asset : TypeAssets)
+	{
+		const auto SpriteSheet = CastChecked<UPaperSpriteSheet>(Asset);
+		if (SpriteSheet->AssetImportData)
+		{
+			OutSourceFilePaths.Add(FReimportManager::ResolveImportFilename(SpriteSheet->AssetImportData->SourceFilePath, SpriteSheet));
+		}
+	}
+}
+
 void FPaperSpriteSheetAssetTypeActions::GetActions(const TArray<UObject*>& InObjects, FMenuBuilder& MenuBuilder)
 {
 	auto SpriteSheetImports = GetTypedWeakObjectPtrs<UPaperSpriteSheet>(InObjects);
 
-	MenuBuilder.AddMenuEntry(
-		LOCTEXT("SpriteSheet_Reimport", "Reimport"),
-		LOCTEXT("SpriteSheet_ReimportTooltip", "Reimports the sprite sheet."),
-		FSlateIcon(),
-		FUIAction(
-		FExecuteAction::CreateSP(this, &FPaperSpriteSheetAssetTypeActions::ExecuteReimport, SpriteSheetImports),
-		FCanExecuteAction()
-		)
-	);
-
 	MenuBuilder.AddMenuEntry(
 		LOCTEXT("SpriteSheet_CreateFlipbooks", "Create Flipbooks"),
 		LOCTEXT("SpriteSheet_CreateFlipbooksTooltip", "Creates flipbooks from sprites in this sprite sheet."),
-		FSlateIcon(),
+		FSlateIcon(FEditorStyle::GetStyleSetName(), "ClassIcon.PaperFlipbook"),
 		FUIAction(
 		FExecuteAction::CreateSP(this, &FPaperSpriteSheetAssetTypeActions::ExecuteCreateFlipbooks, SpriteSheetImports),
 		FCanExecuteAction()
@@ -63,18 +70,6 @@ void FPaperSpriteSheetAssetTypeActions::GetActions(const TArray<UObject*>& InObj
 
 //////////////////////////////////////////////////////////////////////////
 
-void FPaperSpriteSheetAssetTypeActions::ExecuteReimport(TArray<TWeakObjectPtr<UPaperSpriteSheet>> Objects)
-{
-	for (auto ObjIt = Objects.CreateConstIterator(); ObjIt; ++ObjIt)
-	{
-		auto Object = (*ObjIt).Get();
-		if (Object)
-		{
-			FReimportManager::Instance()->Reimport(Object, /*bAskForNewFileIfMissing=*/true);
-		}
-	}
-}
-
 static bool ExtractSpriteNumber(const FString& String, FString& BareString, int& Number)
 {
 	bool bExtracted = false;
diff --git a/Engine/Plugins/2D/Paper2D/Source/PaperJsonImporter/Private/PaperSpriteSheetAssetTypeActions.h b/Engine/Plugins/2D/Paper2D/Source/PaperJsonImporter/Private/PaperSpriteSheetAssetTypeActions.h
--- a/Engine/Plugins/2D/Paper2D/Source/PaperJsonImporter/Private/PaperSpriteSheetAssetTypeActions.h
+++ b/Engine/Plugins/2D/Paper2D/Source/PaperJsonImporter/Private/PaperSpriteSheetAssetTypeActions.h
@@ -16,9 +16,10 @@ public:
 	virtual bool HasActions(const TArray<UObject*>& InObjects) const override { return true; }
 	virtual void GetActions(const TArray<UObject*>& InObjects, FMenuBuilder& MenuBuilder) override;
 	virtual uint32 GetCategories() override;
+	virtual bool IsImportedAsset() const override;
+	virtual void GetResolvedSourceFilePaths(const TArray<UObject*>& TypeAssets, TArray<FString>& OutSourceFilePaths) const override;
 	// End of IAssetTypeActions interface
 
 private:
-	void ExecuteReimport(TArray<TWeakObjectPtr<UPaperSpriteSheet>> Objects);
 	void ExecuteCreateFlipbooks(TArray<TWeakObjectPtr<UPaperSpriteSheet>> Objects);
 };
\ No newline at end of file
diff --git a/Engine/Plugins/Media/MediaPlayerEditor/Source/MediaPlayerEditor/Private/AssetTools/MediaPlayerActions.cpp b/Engine/Plugins/Media/MediaPlayerEditor/Source/MediaPlayerEditor/Private/AssetTools/MediaPlayerActions.cpp
--- a/Engine/Plugins/Media/MediaPlayerEditor/Source/MediaPlayerEditor/Private/AssetTools/MediaPlayerActions.cpp
+++ b/Engine/Plugins/Media/MediaPlayerEditor/Source/MediaPlayerEditor/Private/AssetTools/MediaPlayerActions.cpp
@@ -7,7 +7,7 @@
 #define LOCTEXT_NAMESPACE "AssetTypeActions"
 
 
-/* FMediaPlayerActions structors
+/* FMediaPlayerActions constructors
  *****************************************************************************/
 
 FMediaPlayerActions::FMediaPlayerActions( const TSharedRef<ISlateStyle>& InStyle )
@@ -33,7 +33,7 @@ void FMediaPlayerActions::GetActions( const TArray<UObject*>& InObjects, FMenuBu
 	MenuBuilder.AddMenuEntry(
 		LOCTEXT("MediaPlayer_PlayMovie", "Play Movie"),
 		LOCTEXT("MediaPlayer_PlayMovieToolTip", "Starts playback of the media."),
-		FSlateIcon(),
+		FSlateIcon( FEditorStyle::GetStyleSetName(), "MediaAsset.AssetActions.Play" ),
 		FUIAction(
 			FExecuteAction::CreateSP(this, &FMediaPlayerActions::HandlePlayMovieActionExecute, MediaPlayers),
 			FCanExecuteAction()
@@ -43,7 +43,7 @@ void FMediaPlayerActions::GetActions( const TArray<UObject*>& InObjects, FMenuBu
 	MenuBuilder.AddMenuEntry(
 		LOCTEXT("MediaPlayer_PauseMovie", "Pause Movie"),
 		LOCTEXT("MediaPlayer_PauseMovieToolTip", "Pauses playback of the media."),
-		FSlateIcon(),
+		FSlateIcon( FEditorStyle::GetStyleSetName(), "MediaAsset.AssetActions.Pause" ),
 		FUIAction(
 			FExecuteAction::CreateSP(this, &FMediaPlayerActions::HandlePauseMovieActionExecute, MediaPlayers),
 			FCanExecuteAction()
@@ -55,7 +55,7 @@ void FMediaPlayerActions::GetActions( const TArray<UObject*>& InObjects, FMenuBu
 	MenuBuilder.AddMenuEntry(
 		LOCTEXT("MediaPlayer_CreateMediaSoundWave", "Create Media Sound Wave"),
 		LOCTEXT("MediaPlayer_CreateMediaSoundWaveTooltip", "Creates a new MediaSoundWave using this MediaPlayer asset."),
-		FSlateIcon(),
+		FSlateIcon( FEditorStyle::GetStyleSetName(), "ClassIcon.MediaSoundWave" ),
 		FUIAction(
 			FExecuteAction::CreateSP(this, &FMediaPlayerActions::ExecuteCreateMediaSoundWave, MediaPlayers),
 			FCanExecuteAction()
@@ -65,7 +65,7 @@ void FMediaPlayerActions::GetActions( const TArray<UObject*>& InObjects, FMenuBu
 	MenuBuilder.AddMenuEntry(
 		LOCTEXT("MediaPlayer_CreateMediaTexture", "Create Media Texture"),
 		LOCTEXT("MediaPlayer_CreateMediaTextureTooltip", "Creates a new MediaTexture using this MediaPlayer asset."),
-		FSlateIcon(),
+		FSlateIcon( FEditorStyle::GetStyleSetName(), "ClassIcon.MediaTexture" ),
 		FUIAction(
 			FExecuteAction::CreateSP(this, &FMediaPlayerActions::ExecuteCreateMediaTexture, MediaPlayers),
 			FCanExecuteAction()
diff --git a/Engine/Plugins/Media/MediaPlayerEditor/Source/MediaPlayerEditor/Private/AssetTools/MediaTextureActions.cpp b/Engine/Plugins/Media/MediaPlayerEditor/Source/MediaPlayerEditor/Private/AssetTools/MediaTextureActions.cpp
--- a/Engine/Plugins/Media/MediaPlayerEditor/Source/MediaPlayerEditor/Private/AssetTools/MediaTextureActions.cpp
+++ b/Engine/Plugins/Media/MediaPlayerEditor/Source/MediaPlayerEditor/Private/AssetTools/MediaTextureActions.cpp
@@ -23,20 +23,10 @@ void FMediaTextureActions::GetActions( const TArray<UObject*>& InObjects, FMenuB
 
 	auto Textures = GetTypedWeakObjectPtrs<UTexture>(InObjects);
 
-	MenuBuilder.AddMenuEntry(
-		LOCTEXT("MediaTexture_Edit", "Edit"),
-		LOCTEXT("MediaTexture_EditTooltip", "Opens the selected textures in the texture viewer."),
-		FSlateIcon(),
-		FUIAction(
-			FExecuteAction::CreateSP(this, &FMediaTextureActions::ExecuteEdit, Textures),
-			FCanExecuteAction()
-		)
-	);
-
 	MenuBuilder.AddMenuEntry(
 		LOCTEXT("MediaTexture_CreateMaterial", "Create Material"),
 		LOCTEXT("MediaTexture_CreateMaterialTooltip", "Creates a new material using this texture."),
-		FSlateIcon(),
+		FSlateIcon( FEditorStyle::GetStyleSetName(), "ClassIcon.Material" ),
 		FUIAction(
 			FExecuteAction::CreateSP(this, &FMediaTextureActions::ExecuteCreateMaterial, Textures),
 			FCanExecuteAction()
@@ -46,7 +36,7 @@ void FMediaTextureActions::GetActions( const TArray<UObject*>& InObjects, FMenuB
 /*	MenuBuilder.AddMenuEntry(
 		LOCTEXT("MediaTexture_CreateSlateBrush", "Create Slate Brush"),
 		LOCTEXT("MediaTexture_CreateSlateBrushToolTip", "Creates a new slate brush using this texture."),
-		FSlateIcon(),
+		FSlateIcon( FEditorStyle::GetStyleSetName(), "ClassIcon.SlateBrushAsset" ),
 		FUIAction(
 			FExecuteAction::CreateSP( this, &FAssetTypeActions_MediaTexture::ExecuteCreateSlateBrush, Textures ),
 			FCanExecuteAction()
@@ -237,20 +227,6 @@ void FAssetTypeActions_MediaTexture::ExecuteCreateSlateBrush( TArray<TWeakObject
 }*/
 
 
-void FMediaTextureActions::ExecuteEdit( TArray<TWeakObjectPtr<UTexture>> Objects )
-{
-	for (auto ObjIt = Objects.CreateConstIterator(); ObjIt; ++ObjIt)
-	{
-		auto Object = (*ObjIt).Get();
-
-		if (Object != nullptr)
-		{
-			FAssetEditorManager::Get().OpenEditorForAsset(Object);
-		}
-	}
-}
-
-
 void FMediaTextureActions::ExecuteFindMaterials( TWeakObjectPtr<UTexture> Object )
 {
 	auto Texture = Object.Get();
diff --git a/Engine/Plugins/Media/MediaPlayerEditor/Source/MediaPlayerEditor/Private/AssetTools/MediaTextureActions.h b/Engine/Plugins/Media/MediaPlayerEditor/Source/MediaPlayerEditor/Private/AssetTools/MediaTextureActions.h
--- a/Engine/Plugins/Media/MediaPlayerEditor/Source/MediaPlayerEditor/Private/AssetTools/MediaTextureActions.h
+++ b/Engine/Plugins/Media/MediaPlayerEditor/Source/MediaPlayerEditor/Private/AssetTools/MediaTextureActions.h
@@ -33,9 +33,6 @@ private:
 	/** Callback for selecting CreateSlateBrush. */
 	void ExecuteCreateSlateBrush( TArray<TWeakObjectPtr<UTexture>> Objects );
 
-	/** Callback for selecting Edit. */
-	void ExecuteEdit( TArray<TWeakObjectPtr<UTexture>> Objects );
-
 	/** Callback for selecting FindMaterials. */
 	void ExecuteFindMaterials( TWeakObjectPtr<UTexture> Object );
 };
diff --git a/Engine/Source/Developer/AssetTools/Private/AssetTypeActions/AssetTypeActions_AnimSequence.cpp b/Engine/Source/Developer/AssetTools/Private/AssetTypeActions/AssetTypeActions_AnimSequence.cpp
--- a/Engine/Source/Developer/AssetTools/Private/AssetTypeActions/AssetTypeActions_AnimSequence.cpp
+++ b/Engine/Source/Developer/AssetTools/Private/AssetTypeActions/AssetTypeActions_AnimSequence.cpp
@@ -9,62 +9,23 @@ void FAssetTypeActions_AnimSequence::GetActions( const TArray<UObject*>& InObjec
 {
 	auto Sequences = GetTypedWeakObjectPtrs<UAnimSequence>(InObjects);
 
-	MenuBuilder.AddMenuEntry(
-		LOCTEXT("AnimSequence_Reimport", "Reimport"),
-		LOCTEXT("AnimSequence_ReimportTooltip", "Reimports the selected sequences."),
-		FSlateIcon(),
-		FUIAction(
-			FExecuteAction::CreateSP( this, &FAssetTypeActions_AnimSequence::ExecuteReimport, Sequences ),
-			FCanExecuteAction()
-			)
-		);
-
-	// source menu
-	MenuBuilder.AddSubMenu(
-			LOCTEXT("SourceAnimSubmenu", "Source"),
-			LOCTEXT("SourceAnimSubmenu_ToolTip", "Source data related"),
-			FNewMenuDelegate::CreateSP(this, &FAssetTypeActions_AnimSequence::FillSourceMenu, Sequences));
-
 	// create menu
 	MenuBuilder.AddSubMenu(
-			LOCTEXT("CreateAnimSubmenu", "Create"),
-			LOCTEXT("CreateAnimSubmenu_ToolTip", "Create assets from this anim sequence"),
-			FNewMenuDelegate::CreateSP(this, &FAssetTypeActions_AnimSequence::FillCreateMenu, Sequences));
-
-	FAssetTypeActions_AnimationAsset::GetActions(InObjects, MenuBuilder);
-
-}
-
-void FAssetTypeActions_AnimSequence::FillSourceMenu(FMenuBuilder& MenuBuilder, const TArray<TWeakObjectPtr<UAnimSequence>> Sequences) const
-{
-	MenuBuilder.AddMenuEntry(
-		LOCTEXT("AnimSequence_FindInExplorer", "Find Source"),
-		LOCTEXT("AnimSequence_FindInExplorerTooltip", "Opens explorer at the location of this asset."),
-		FSlateIcon(),
-		FUIAction(
-			FExecuteAction::CreateSP(this, &FAssetTypeActions_AnimSequence::ExecuteFindInExplorer, Sequences),
-			FCanExecuteAction::CreateSP(this, &FAssetTypeActions_AnimSequence::CanExecuteSourceCommands, Sequences)
-			)
-		);
-
-	MenuBuilder.AddMenuEntry(
-		LOCTEXT("AnimSequence_OpenInExternalEditor", "Open Source"),
-		LOCTEXT("AnimSequence_OpenInExternalEditorTooltip", "Opens the selected asset in an external editor."),
-		FSlateIcon(),
-		FUIAction(
-			FExecuteAction::CreateSP(this, &FAssetTypeActions_AnimSequence::ExecuteOpenInExternalEditor, Sequences),
-			FCanExecuteAction::CreateSP(this, &FAssetTypeActions_AnimSequence::CanExecuteSourceCommands, Sequences)
-			)
+		LOCTEXT("CreateAnimSubmenu", "Create"),
+		LOCTEXT("CreateAnimSubmenu_ToolTip", "Create assets from this anim sequence"),
+		FNewMenuDelegate::CreateSP(this, &FAssetTypeActions_AnimSequence::FillCreateMenu, Sequences),
+		false,
+		FSlateIcon(FEditorStyle::GetStyleSetName(), "Persona.AssetActions.CreateAnimAsset")
 		);
 
 	MenuBuilder.AddMenuEntry(
 		LOCTEXT("AnimSequence_ReimportWithNewSource", "Reimport with New Source"),
-		LOCTEXT("AnimSequence_ReimportWithNewSourceTooltip", "Reimport with new source."),
-		FSlateIcon(),
-		FUIAction(
-			FExecuteAction::CreateSP(this, &FAssetTypeActions_AnimSequence::ExecuteReimportWithNewSource, Sequences)
-			)
+		LOCTEXT("AnimSequence_ReimportWithNewSourceTooltip", "Reimport the selected sequence(s) from a new source file."),
+		FSlateIcon(FEditorStyle::GetStyleSetName(), "ContentBrowser.AssetActions.ReimportAsset"),
+		FUIAction(FExecuteAction::CreateSP(this, &FAssetTypeActions_AnimSequence::ExecuteReimportWithNewSource, Sequences))
 		);
+
+	FAssetTypeActions_AnimationAsset::GetActions(InObjects, MenuBuilder);
 }
 
 void FAssetTypeActions_AnimSequence::FillCreateMenu(FMenuBuilder& MenuBuilder, const TArray<TWeakObjectPtr<UAnimSequence>> Sequences) const
@@ -72,47 +33,32 @@ void FAssetTypeActions_AnimSequence::FillCreateMenu(FMenuBuilder& MenuBuilder, c
 	MenuBuilder.AddMenuEntry(
 		LOCTEXT("AnimSequence_NewAnimComposite", "Create AnimComposite"),
 		LOCTEXT("AnimSequence_NewAnimCompositeTooltip", "Creates an AnimComposite using the selected anim sequence."),
-		FSlateIcon(),
+		FSlateIcon(FEditorStyle::GetStyleSetName(), "ClassIcon.AnimComposite"),
 		FUIAction(
-		FExecuteAction::CreateSP(this, &FAssetTypeActions_AnimSequence::ExecuteNewAnimComposite, Sequences),
-		FCanExecuteAction()
-		)
+			FExecuteAction::CreateSP(this, &FAssetTypeActions_AnimSequence::ExecuteNewAnimComposite, Sequences),
+			FCanExecuteAction()
+			)
 		);
 
 	MenuBuilder.AddMenuEntry(
 		LOCTEXT("AnimSequence_NewAnimMontage", "Create AnimMontage"),
 		LOCTEXT("AnimSequence_NewAnimMontageTooltip", "Creates an AnimMontage using the selected anim sequence."),
-		FSlateIcon(),
+		FSlateIcon(FEditorStyle::GetStyleSetName(), "ClassIcon.AnimMontage"),
 		FUIAction(
-		FExecuteAction::CreateSP(this, &FAssetTypeActions_AnimSequence::ExecuteNewAnimMontage, Sequences),
-		FCanExecuteAction()
-		)
+			FExecuteAction::CreateSP(this, &FAssetTypeActions_AnimSequence::ExecuteNewAnimMontage, Sequences),
+			FCanExecuteAction()
+			)
 		);
 }
-void FAssetTypeActions_AnimSequence::ExecuteReimport(TArray<TWeakObjectPtr<UAnimSequence>> Objects)
-{
-	for (auto ObjIt = Objects.CreateConstIterator(); ObjIt; ++ObjIt)
-	{
-		auto Object = (*ObjIt).Get();
-		if ( Object )
-		{
-			FReimportManager::Instance()->Reimport( Object, /*bAskForNewFileIfMissing=*/true );
-		}
-	}
-}
 
-void FAssetTypeActions_AnimSequence::ExecuteFindInExplorer(TArray<TWeakObjectPtr<UAnimSequence>> Objects)
+void FAssetTypeActions_AnimSequence::GetResolvedSourceFilePaths(const TArray<UObject*>& TypeAssets, TArray<FString>& OutSourceFilePaths) const
 {
-	for (auto ObjIt = Objects.CreateConstIterator(); ObjIt; ++ObjIt)
+	for (auto& Asset : TypeAssets)
 	{
-		auto Object = (*ObjIt).Get();
-		if ( Object && Object->AssetImportData )
+		const auto AnimSequence = CastChecked<UAnimSequence>(Asset);
+		if (AnimSequence->AssetImportData)
 		{
-			const FString SourceFilePath = FReimportManager::ResolveImportFilename(Object->AssetImportData->SourceFilePath, Object);
-			if ( SourceFilePath.Len() && IFileManager::Get().FileSize( *SourceFilePath ) != INDEX_NONE )
-			{
-				FPlatformProcess::ExploreFolder( *FPaths::GetPath(SourceFilePath) );
-			}
+			OutSourceFilePaths.Add(FReimportManager::ResolveImportFilename(AnimSequence->AssetImportData->SourceFilePath, AnimSequence));
 		}
 	}
 }
@@ -130,43 +76,14 @@ void FAssetTypeActions_AnimSequence::ExecuteReimportWithNewSource(TArray<TWeakOb
 		}
 	}
 
-	ExecuteReimport(Objects);
-}
-
-void FAssetTypeActions_AnimSequence::ExecuteOpenInExternalEditor(TArray<TWeakObjectPtr<UAnimSequence>> Objects)
-{
-	for (auto ObjIt = Objects.CreateConstIterator(); ObjIt; ++ObjIt)
-	{
-		auto Object = (*ObjIt).Get();
-		if ( Object && Object->AssetImportData )
-		{
-			const FString SourceFilePath = FReimportManager::ResolveImportFilename(Object->AssetImportData->SourceFilePath, Object);
-			if ( SourceFilePath.Len() && IFileManager::Get().FileSize( *SourceFilePath ) != INDEX_NONE )
-			{
-				FPlatformProcess::LaunchFileInDefaultExternalApplication( *SourceFilePath, NULL, ELaunchVerb::Edit );
-			}
-		}
-	}
-}
-
-bool FAssetTypeActions_AnimSequence::CanExecuteSourceCommands(TArray<TWeakObjectPtr<UAnimSequence>> Objects) const
-{
-	bool bHaveSourceAsset = false;
 	for (auto ObjIt = Objects.CreateConstIterator(); ObjIt; ++ObjIt)
 	{
 		auto Object = (*ObjIt).Get();
-		if ( Object && Object->AssetImportData )
+		if (Object)
 		{
-			const FString& SourceFilePath = FReimportManager::ResolveImportFilename(Object->AssetImportData->SourceFilePath, Object);
-
-			if ( SourceFilePath.Len() && IFileManager::Get().FileSize(*SourceFilePath) != INDEX_NONE )
-			{
-				return true;
-			}
+			FReimportManager::Instance()->Reimport(Object, /*bAskForNewFileIfMissing=*/true);
 		}
 	}
-
-	return false;
 }
 
 void FAssetTypeActions_AnimSequence::ExecuteNewAnimComposite(TArray<TWeakObjectPtr<UAnimSequence>> Objects) const
diff --git a/Engine/Source/Developer/AssetTools/Private/AssetTypeActions/AssetTypeActions_AnimSequence.h b/Engine/Source/Developer/AssetTools/Private/AssetTypeActions/AssetTypeActions_AnimSequence.h
--- a/Engine/Source/Developer/AssetTools/Private/AssetTypeActions/AssetTypeActions_AnimSequence.h
+++ b/Engine/Source/Developer/AssetTools/Private/AssetTypeActions/AssetTypeActions_AnimSequence.h
@@ -13,23 +13,13 @@ public:
 	virtual UClass* GetSupportedClass() const override { return UAnimSequence::StaticClass(); }
 	virtual void GetActions( const TArray<UObject*>& InObjects, FMenuBuilder& MenuBuilder ) override;
 	virtual bool CanFilter() override { return true; }
+	virtual bool IsImportedAsset() const override { return true; }
+	virtual void GetResolvedSourceFilePaths(const TArray<UObject*>& TypeAssets, TArray<FString>& OutSourceFilePaths) const override;
 
 private:
-	/** Handler for when Reimport is selected */
-	void ExecuteReimport(TArray<TWeakObjectPtr<UAnimSequence>> Objects);
-
-	/** Handler for when FindInExplorer is selected */
-	void ExecuteFindInExplorer(TArray<TWeakObjectPtr<UAnimSequence>> Objects);
-
-	/** Handler for when OpenInExternalEditor is selected */
-	void ExecuteOpenInExternalEditor(TArray<TWeakObjectPtr<UAnimSequence>> Objects);
-
 	/** Handler for when OpenInExternalEditor is selected */
 	void ExecuteReimportWithNewSource(TArray<TWeakObjectPtr<UAnimSequence>> Objects);
 
-	/** Returns true to allow execution of source file commands */
-	bool CanExecuteSourceCommands(TArray<TWeakObjectPtr<UAnimSequence>> Objects) const;
-
 	/** Handler for when Create AnimComposite is selected */
 	void ExecuteNewAnimComposite(TArray<TWeakObjectPtr<UAnimSequence>> Objects) const;
 
@@ -45,6 +35,5 @@ private:
 	/** Creates animation assets of the supplied class */
 	void CreateAnimationAssets(const TArray<TWeakObjectPtr<UAnimSequence>>& AnimSequences, TSubclassOf<UAnimationAsset> AssetClass, UFactory* AssetFactory, const FString& InSuffix, FOnConfigureFactory OnConfigureFactory) const;
 
-	void FillSourceMenu(FMenuBuilder& MenuBuilder, const TArray<TWeakObjectPtr<UAnimSequence>> Sequences) const;
 	void FillCreateMenu(FMenuBuilder& MenuBuilder, const TArray<TWeakObjectPtr<UAnimSequence>> Sequences) const;
 };
\ No newline at end of file
diff --git a/Engine/Source/Developer/AssetTools/Private/AssetTypeActions/AssetTypeActions_Blueprint.cpp b/Engine/Source/Developer/AssetTools/Private/AssetTypeActions/AssetTypeActions_Blueprint.cpp
--- a/Engine/Source/Developer/AssetTools/Private/AssetTypeActions/AssetTypeActions_Blueprint.cpp
+++ b/Engine/Source/Developer/AssetTools/Private/AssetTypeActions/AssetTypeActions_Blueprint.cpp
@@ -15,26 +15,33 @@
 void FAssetTypeActions_Blueprint::GetActions( const TArray<UObject*>& InObjects, FMenuBuilder& MenuBuilder )
 {
 	auto Blueprints = GetTypedWeakObjectPtrs<UBlueprint>(InObjects);
+	
+	if (Blueprints.Num() > 1)
+	{
+		// Ensure that all the selected blueprints are actors
+		bool bCanEditSharedDefaults = true;
+		for (auto Blueprint : Blueprints)
+		{
+			if (!Blueprint.Get()->ParentClass->IsChildOf(AActor::StaticClass()))
+			{
+				bCanEditSharedDefaults = false;
+				break;
+			}
+		}
 
-	MenuBuilder.AddMenuEntry(
-		LOCTEXT("Blueprint_EditAction", "Open in Full Editor"),
-		LOCTEXT("Blueprint_EditActionTooltip", "Opens the selected blueprints in the either the blueprint editor or in Persona (depending on type)."),
-		FSlateIcon(),
-		FUIAction(
-			FExecuteAction::CreateSP( this, &FAssetTypeActions_Blueprint::ExecuteEdit, Blueprints ),
-			FCanExecuteAction()
-			)
-		);
-
-	MenuBuilder.AddMenuEntry(
-		LOCTEXT("Blueprint_EditDefaults", "Edit Defaults"),
-		LOCTEXT("Blueprint_EditDefaultsTooltip", "Edits the default properties for the selected blueprints."),
-		FSlateIcon(),
-		FUIAction(
-			FExecuteAction::CreateSP( this, &FAssetTypeActions_Blueprint::ExecuteEditDefaults, Blueprints ),
-			FCanExecuteAction()
-			)
-		);
+		if (bCanEditSharedDefaults)
+		{
+			MenuBuilder.AddMenuEntry(
+			LOCTEXT("Blueprint_EditDefaults", "Edit Shared Defaults"),
+			LOCTEXT("Blueprint_EditDefaultsTooltip", "Edit the shared default properties of the selected blueprints."),
+			FSlateIcon(FEditorStyle::GetStyleSetName(), "Kismet.Tabs.BlueprintDefaults"),
+			FUIAction(
+				FExecuteAction::CreateSP( this, &FAssetTypeActions_Blueprint::ExecuteEditDefaults, Blueprints ),
+				FCanExecuteAction()
+				)
+			);
+		}	
+	}
 
 	if ( Blueprints.Num() == 1 && CanCreateNewDerivedBlueprint() )
 	{
@@ -45,7 +52,7 @@ void FAssetTypeActions_Blueprint::GetActions( const TArray<UObject*>& InObjects,
 		MenuBuilder.AddMenuEntry(
 			LOCTEXT("Blueprint_NewDerivedBlueprint", "Create Blueprint based on this"),
 			DynamicTooltipAttribute,
-			FSlateIcon(),
+			FSlateIcon(FEditorStyle::GetStyleSetName(), "LevelEditor.CreateClassBlueprint"),
 			FUIAction(
 				FExecuteAction::CreateSP( this, &FAssetTypeActions_Blueprint::ExecuteNewDerivedBlueprint, Blueprints[0] ),
 				FCanExecuteAction::CreateSP( this, &FAssetTypeActions_Blueprint::CanExecuteNewDerivedBlueprint, Blueprints[0] )
@@ -63,8 +70,14 @@ void FAssetTypeActions_Blueprint::OpenAssetEditor( const TArray<UObject*>& InObj
 		auto Blueprint = Cast<UBlueprint>(*ObjIt);
 		if (Blueprint && Blueprint->SkeletonGeneratedClass && Blueprint->GeneratedClass )
 		{
+			// Open to the full editor
+			const bool bSavedForceFullEditor = Blueprint->bForceFullEditor;
+			Blueprint->bForceFullEditor = true;
+
 			FBlueprintEditorModule& BlueprintEditorModule = FModuleManager::LoadModuleChecked<FBlueprintEditorModule>( "Kismet" );
 			TSharedRef< IBlueprintEditor > NewKismetEditor = BlueprintEditorModule.CreateBlueprintEditor( Mode, EditWithinLevelEditor, Blueprint, ShouldUseDataOnlyEditor(Blueprint) );
+			
+			Blueprint->bForceFullEditor = bSavedForceFullEditor;
 		}
 		else
 		{
@@ -73,24 +86,6 @@ void FAssetTypeActions_Blueprint::OpenAssetEditor( const TArray<UObject*>& InObj
 	}
 }
 
-void FAssetTypeActions_Blueprint::ExecuteEdit(TArray<TWeakObjectPtr<UBlueprint>> Objects)
-{
-	for (auto ObjIt = Objects.CreateConstIterator(); ObjIt; ++ObjIt)
-	{
-		auto Object = (*ObjIt).Get();
-		if ( Object )
-		{
-			// Force full (non data-only) editor if possible.
-			const bool bSavedForceFullEditor = Object->bForceFullEditor;
-			Object->bForceFullEditor = true;
-
-			FAssetEditorManager::Get().OpenEditorForAsset(Object);
-
-			Object->bForceFullEditor = bSavedForceFullEditor;
-		}
-	}
-}
-
 bool FAssetTypeActions_Blueprint::CanCreateNewDerivedBlueprint() const
 {
 	return true;
diff --git a/Engine/Source/Developer/AssetTools/Private/AssetTypeActions/AssetTypeActions_Blueprint.h b/Engine/Source/Developer/AssetTools/Private/AssetTypeActions/AssetTypeActions_Blueprint.h
--- a/Engine/Source/Developer/AssetTools/Private/AssetTypeActions/AssetTypeActions_Blueprint.h
+++ b/Engine/Source/Developer/AssetTools/Private/AssetTypeActions/AssetTypeActions_Blueprint.h
@@ -20,9 +20,6 @@ public:
 	virtual FText GetAssetDescription(const FAssetData& AssetData) const override;
 
 protected:
-	/** Handler for when Edit is selected */
-	virtual void ExecuteEdit(TArray<TWeakObjectPtr<UBlueprint>> Objects);
-
 	/** Whether or not this asset can create derived blueprints */
 	virtual bool CanCreateNewDerivedBlueprint() const;
 
diff --git a/Engine/Source/Developer/AssetTools/Private/AssetTypeActions/AssetTypeActions_CSVAssetBase.cpp b/Engine/Source/Developer/AssetTools/Private/AssetTypeActions/AssetTypeActions_CSVAssetBase.cpp
--- a/Engine/Source/Developer/AssetTools/Private/AssetTypeActions/AssetTypeActions_CSVAssetBase.cpp
+++ b/Engine/Source/Developer/AssetTools/Private/AssetTypeActions/AssetTypeActions_CSVAssetBase.cpp
@@ -4,112 +4,52 @@
 
 #define LOCTEXT_NAMESPACE "AssetTypeActions"
 
-void FAssetTypeActions_CSVAssetBase::ExecuteEdit(TArray< TWeakObjectPtr<UObject> > Objects)
-{
-	for (auto ObjIt = Objects.CreateConstIterator(); ObjIt; ++ObjIt)
-	{
-		auto Object = (*ObjIt).Get();
-		if ( Object )
-		{
-			FAssetEditorManager::Get().OpenEditorForAsset(Object);
-		}
-	}
-}
-
-void FAssetTypeActions_CSVAssetBase::ExecuteReimport( TArray< TWeakObjectPtr<UObject> > Objects )
-{
-	for (auto ObjIt = Objects.CreateConstIterator(); ObjIt; ++ObjIt)
-	{
-		auto Object = (*ObjIt).Get();
-		if ( Object )
-		{
-			FReimportManager::Instance()->Reimport(Object, /*bAskForNewFileIfMissing=*/true);
-		}
-	}
-}
-
-void FAssetTypeActions_CSVAssetBase::ExecuteFindInExplorer(TArray<FString> Filenames)
-{
-	for (TArray<FString>::TConstIterator FilenameIter(Filenames); FilenameIter; ++FilenameIter)
-	{
-		const FString FullImportPath = FPaths::GetPath(FPaths::ConvertRelativePathToFull(*FilenameIter));
-		if (!FullImportPath.IsEmpty() && IFileManager::Get().DirectoryExists(*FullImportPath))
-		{
-			FPlatformProcess::ExploreFolder(*FullImportPath);
-		}
-	}
-}
-
-void FAssetTypeActions_CSVAssetBase::ExecuteOpenInExternalEditor(TArray<FString> Filenames, TArray<FString> OverrideExtensions)
+void FAssetTypeActions_CSVAssetBase::ExecuteFindExcelFileInExplorer(TArray<FString> Filenames, TArray<FString> OverrideExtensions)
 {
 	for (TArray<FString>::TConstIterator FilenameIter(Filenames); FilenameIter; ++FilenameIter)
 	{
 		const FString CSVFilename = FPaths::ConvertRelativePathToFull(*FilenameIter);
-		if (OverrideExtensions.Num() > 0)
+		const FString RootPath = FPaths::GetPath(CSVFilename);
+		const FString BaseFilename = FPaths::GetBaseFilename(CSVFilename, true);
+		
+		for (TArray<FString>::TConstIterator ExtensionItr(OverrideExtensions); ExtensionItr; ++ExtensionItr)
 		{
-			const FString RootPath = FPaths::GetPath(CSVFilename);
-			const FString BaseFilename = FPaths::GetBaseFilename(CSVFilename, true);
-			for (TArray<FString>::TConstIterator ExtensionItr(OverrideExtensions); ExtensionItr; ++ExtensionItr)
+			const FString FilenameWithExtension(FString::Printf(TEXT("%s/%s%s"), *RootPath, *BaseFilename, **ExtensionItr));
+			if (VerifyFileExists(FilenameWithExtension))
 			{
-				const FString FilenameWithExtension(FString::Printf(TEXT("%s/%s%s"), *RootPath, *BaseFilename, **ExtensionItr));
-				if (VerifyFileExists(FilenameWithExtension))
-				{
-					FPlatformProcess::LaunchFileInDefaultExternalApplication(*FilenameWithExtension, NULL, ELaunchVerb::Edit);
-					break;
-				}
+				FPlatformProcess::LaunchFileInDefaultExternalApplication(*FilenameWithExtension, NULL, ELaunchVerb::Edit);
+				break;
 			}
 		}
-		else if (VerifyFileExists(CSVFilename))
-		{
-			FPlatformProcess::LaunchFileInDefaultExternalApplication(*CSVFilename, NULL, ELaunchVerb::Edit);
-		}
 	}
 }
 
-bool FAssetTypeActions_CSVAssetBase::CanExecuteLaunchExternalSourceCommands(TArray<FString> Filenames, TArray<FString> OverrideExtensions) const
+bool FAssetTypeActions_CSVAssetBase::CanExecuteFindExcelFileInExplorer(TArray<FString> Filenames, TArray<FString> OverrideExtensions) const
 {
-	bool bCanExecute = false;
+	// Verify that extensions were provided
+	if (OverrideExtensions.Num() == 0)
+	{
+		return false;
+	}
 
+	// Verify that the file exists with any of the given extensions
 	for (TArray<FString>::TConstIterator FilenameIter(Filenames); FilenameIter; ++FilenameIter)
 	{
 		const FString CSVFilename = FPaths::ConvertRelativePathToFull(*FilenameIter);
-		if (OverrideExtensions.Num() > 0)
+		const FString RootPath = FPaths::GetPath(CSVFilename);
+		const FString BaseFilename = FPaths::GetBaseFilename(CSVFilename, true);
+
+		for (TArray<FString>::TConstIterator ExtensionItr(OverrideExtensions); ExtensionItr; ++ExtensionItr)
 		{
-			const FString RootPath = FPaths::GetPath(CSVFilename);
-			const FString BaseFilename = FPaths::GetBaseFilename(CSVFilename, true);
-			for (TArray<FString>::TConstIterator ExtensionItr(OverrideExtensions); ExtensionItr; ++ExtensionItr)
+			const FString FilenameWithExtension(FString::Printf(TEXT("%s/%s%s"), *RootPath, *BaseFilename, **ExtensionItr));
+			if (VerifyFileExists(FilenameWithExtension))
 			{
-				const FString FilenameWithExtension(FString::Printf(TEXT("%s/%s%s"), *RootPath, *BaseFilename, **ExtensionItr));
-				if (VerifyFileExists(FilenameWithExtension))
-				{
-					bCanExecute = true;
-					break;
-				}
+				return true;
 			}
 		}
-		else if (VerifyFileExists(CSVFilename))
-		{
-			bCanExecute = true;
-		}
-	}
-
-	return bCanExecute;
-}
-
-bool FAssetTypeActions_CSVAssetBase::CanExecuteFindInExplorerSourceCommand(TArray<FString> Filenames) const
-{
-	bool bCanExecute = false;
-	for (TArray<FString>::TConstIterator FilenameItr(Filenames); FilenameItr; ++FilenameItr)
-	{
-		const FString FullImportPath = FPaths::GetPath(FPaths::ConvertRelativePathToFull(*FilenameItr));
-		if (!FullImportPath.IsEmpty() && IFileManager::Get().DirectoryExists(*FullImportPath))
-		{
-			bCanExecute = true;
-			break;
-		}
 	}
 
-	return bCanExecute;
+	return false;
 }
 
 bool FAssetTypeActions_CSVAssetBase::VerifyFileExists(const FString& InFileName) const
diff --git a/Engine/Source/Developer/AssetTools/Private/AssetTypeActions/AssetTypeActions_Curve.cpp b/Engine/Source/Developer/AssetTools/Private/AssetTypeActions/AssetTypeActions_Curve.cpp
--- a/Engine/Source/Developer/AssetTools/Private/AssetTypeActions/AssetTypeActions_Curve.cpp
+++ b/Engine/Source/Developer/AssetTools/Private/AssetTypeActions/AssetTypeActions_Curve.cpp
@@ -1,38 +1,8 @@
 // Copyright 1998-2014 Epic Games, Inc. All Rights Reserved.
 
 #include "AssetToolsPrivatePCH.h"
-#include "Toolkits/AssetEditorManager.h"
 #include "CurveAssetEditorModule.h"
 #include "ICurveAssetEditor.h"
-#include "Editor/UnrealEd/Public/MiniCurveEditor.h"
-#include "Editor/WorkspaceMenuStructure/Public/WorkspaceMenuStructureModule.h"
-
-#define LOCTEXT_NAMESPACE "AssetTypeActions"
-
-void FAssetTypeActions_Curve::GetActions( const TArray<UObject*>& InObjects, FMenuBuilder& MenuBuilder )
-{
-	auto Curves = GetTypedWeakObjectPtrs<UCurveBase>(InObjects);
-
-	MenuBuilder.AddMenuEntry(
-		LOCTEXT("Curve_Edit", "Edit"),
-		LOCTEXT("Curve_EditTooltip", "Opens the selected curves in the curve editor."),
-		FSlateIcon(),
-		FUIAction(
-			FExecuteAction::CreateSP( this, &FAssetTypeActions_Curve::ExecuteEdit, Curves ),
-			FCanExecuteAction()
-			)
-		);
-
-	MenuBuilder.AddMenuEntry(
-		LOCTEXT("Curve_Reimport", "Reimport"),
-		LOCTEXT("Curve_ReimportTooltip", "Reimports the selected Curve from file."),
-		FSlateIcon(),
-		FUIAction(
-		FExecuteAction::CreateSP( this, &FAssetTypeActions_Curve::ExecuteReimport, Curves ),
-		FCanExecuteAction::CreateSP(this, &FAssetTypeActions_Curve::CanReimportCurves, Curves) //if it was from a file?
-		)
-		);
-}
 
 void FAssetTypeActions_Curve::OpenAssetEditor( const TArray<UObject*>& InObjects, TSharedPtr<IToolkitHost> EditWithinLevelEditor )
 {
@@ -49,44 +19,11 @@ void FAssetTypeActions_Curve::OpenAssetEditor( const TArray<UObject*>& InObjects
 	}
 }
 
-void FAssetTypeActions_Curve::ExecuteEdit(TArray<TWeakObjectPtr<UCurveBase>> Objects)
-{
-	for (auto ObjIt = Objects.CreateConstIterator(); ObjIt; ++ObjIt)
-	{
-		auto Object = (*ObjIt).Get();
-		if ( Object )
-		{
-			FAssetEditorManager::Get().OpenEditorForAsset(Object);
-		}
-	}
-}
-
-void FAssetTypeActions_Curve::ExecuteReimport( TArray<TWeakObjectPtr<UCurveBase>> Objects )
-{
-	for (auto ObjIt = Objects.CreateConstIterator(); ObjIt; ++ObjIt)
-	{
-		auto Object = (*ObjIt).Get();
-		if ( Object )
-		{
-			FReimportManager::Instance()->Reimport(Object, /*bAskForNewFileIfMissing=*/true);
-		}
-	}
-}
-
-bool FAssetTypeActions_Curve::CanReimportCurves( TArray<TWeakObjectPtr<UCurveBase>> Objects )const
+void FAssetTypeActions_Curve::GetResolvedSourceFilePaths(const TArray<UObject*>& TypeAssets, TArray<FString>& OutSourceFilePaths) const
 {
-	for(auto It = Objects.CreateIterator();It;++It)
+	for (auto& Asset : TypeAssets)
 	{
-		auto& Obj = *It;
-		if(UCurveBase* Curve = Obj.Get())
-		{
-			if(!Curve->ImportPath.IsEmpty())
-			{
-				return true;
-			}
-		}
+		const auto Curve = CastChecked<UCurveBase>(Asset);
+		OutSourceFilePaths.Add(FReimportManager::ResolveImportFilename(Curve->ImportPath, Curve));
 	}
-	return false;
-}
-
-#undef LOCTEXT_NAMESPACE
+}
\ No newline at end of file
diff --git a/Engine/Source/Developer/AssetTools/Private/AssetTypeActions/AssetTypeActions_Curve.h b/Engine/Source/Developer/AssetTools/Private/AssetTypeActions/AssetTypeActions_Curve.h
--- a/Engine/Source/Developer/AssetTools/Private/AssetTypeActions/AssetTypeActions_Curve.h
+++ b/Engine/Source/Developer/AssetTools/Private/AssetTypeActions/AssetTypeActions_Curve.h
@@ -9,19 +9,9 @@ public:
 	virtual FText GetName() const override { return NSLOCTEXT("AssetTypeActions", "AssetTypeActions_Curve", "Curve"); }
 	virtual FColor GetTypeColor() const override { return FColor(78, 40, 165); }
 	virtual UClass* GetSupportedClass() const override { return UCurveBase::StaticClass(); }
-	virtual bool HasActions ( const TArray<UObject*>& InObjects ) const override { return true; }
-	virtual void GetActions( const TArray<UObject*>& InObjects, FMenuBuilder& MenuBuilder ) override;
 	virtual void OpenAssetEditor( const TArray<UObject*>& InObjects, TSharedPtr<class IToolkitHost> EditWithinLevelEditor = TSharedPtr<IToolkitHost>() ) override;
 	virtual bool CanFilter() override { return false; }
 	virtual uint32 GetCategories() override { return EAssetTypeCategories::Misc; }
-
-private:
-	/** Handler for when Edit is selected */
-	void ExecuteEdit(TArray<TWeakObjectPtr<UCurveBase>> Objects);
-
-	/** Handler for when Reimport is selected */
-	void ExecuteReimport(TArray<TWeakObjectPtr<UCurveBase>> Objects);
-
-	/** Can curves be reimported? */
-	bool CanReimportCurves( TArray<TWeakObjectPtr<UCurveBase>> Objects ) const;
+	virtual bool IsImportedAsset() const override { return true; }
+	virtual void GetResolvedSourceFilePaths(const TArray<UObject*>& TypeAssets, TArray<FString>& OutSourceFilePaths) const override;
 };
\ No newline at end of file
diff --git a/Engine/Source/Developer/AssetTools/Private/AssetTypeActions/AssetTypeActions_CurveTable.cpp b/Engine/Source/Developer/AssetTools/Private/AssetTypeActions/AssetTypeActions_CurveTable.cpp
--- a/Engine/Source/Developer/AssetTools/Private/AssetTypeActions/AssetTypeActions_CurveTable.cpp
+++ b/Engine/Source/Developer/AssetTools/Private/AssetTypeActions/AssetTypeActions_CurveTable.cpp
@@ -21,58 +21,16 @@ void FAssetTypeActions_CurveTable::GetActions( const TArray<UObject*>& InObjects
 		}
 	}
 
-	MenuBuilder.AddMenuEntry(
-		LOCTEXT("CurveTable_Edit", "Edit"),
-		LOCTEXT("CurveTable_EditTooltip", "Opens the selected tables in the table editor."),
-		FSlateIcon(),
-		FUIAction(
-			FExecuteAction::CreateSP( this, &FAssetTypeActions_CurveTable::ExecuteEdit, Tables ),
-			FCanExecuteAction()
-			)
-		);
-
-
-	MenuBuilder.AddMenuEntry(
-		LOCTEXT("CurveTable_Reimport", "Reimport"),
-		LOCTEXT("CurveTable_ReimportTooltip", "Reimports the selected Curve Table from file."),
-		FSlateIcon(),
-		FUIAction(
-		FExecuteAction::CreateSP( this, &FAssetTypeActions_CurveTable::ExecuteReimport, Tables ),
-		FCanExecuteAction() 
-		)
-		);
-
-	MenuBuilder.AddMenuEntry(
-		LOCTEXT("CurveTable_FindSource", "Find Source"),
-		LOCTEXT("CurveTable_FindSourceTooltip", "Opens explorer at the location of this asset's source data."),
-		FSlateIcon(),
-		FUIAction(
-			FExecuteAction::CreateSP( this, &FAssetTypeActions_CurveTable::ExecuteFindInExplorer, ImportPaths ),
-			FCanExecuteAction::CreateSP(this, &FAssetTypeActions_CurveTable::CanExecuteFindInExplorerSourceCommand, ImportPaths )
-			)
-		);
-
-	static const TArray<FString> EmptyArray;
-	MenuBuilder.AddMenuEntry(
-		LOCTEXT("CurveTable_OpenSourceCSV", "Open Source (.csv)"),
-		LOCTEXT("CurveTable_OpenSourceCSVTooltip", "Opens the selected asset's source CSV data in an external editor."),
-		FSlateIcon(),
-		FUIAction(
-			FExecuteAction::CreateSP( this, &FAssetTypeActions_CurveTable::ExecuteOpenInExternalEditor, ImportPaths, EmptyArray ),
-			FCanExecuteAction::CreateSP(this, &FAssetTypeActions_CurveTable::CanExecuteLaunchExternalSourceCommands, ImportPaths, EmptyArray )
-			)
-		);
-
 	TArray<FString> XLSExtensions;
 	XLSExtensions.Add(TEXT(".xls"));
 	XLSExtensions.Add(TEXT(".xlsm"));
 	MenuBuilder.AddMenuEntry(
 		LOCTEXT("CurveTable_OpenSourceXLS", "Open Source (.xls/.xlsm)"),
-		LOCTEXT("CurveTable_OpenSourceXLSTooltip", "Opens the selected asset's source XLS/XLSM data in an external editor."),
+		LOCTEXT("CurveTable_OpenSourceXLSTooltip", "Opens the curve table's source XLS/XLSM file in an external editor."),
 		FSlateIcon(),
 		FUIAction(
-			FExecuteAction::CreateSP( this, &FAssetTypeActions_CurveTable::ExecuteOpenInExternalEditor, ImportPaths, XLSExtensions ),
-			FCanExecuteAction::CreateSP(this, &FAssetTypeActions_CurveTable::CanExecuteLaunchExternalSourceCommands, ImportPaths, XLSExtensions )
+			FExecuteAction::CreateSP( this, &FAssetTypeActions_CurveTable::ExecuteFindExcelFileInExplorer, ImportPaths, XLSExtensions ),
+			FCanExecuteAction::CreateSP(this, &FAssetTypeActions_CurveTable::CanExecuteFindExcelFileInExplorer, ImportPaths, XLSExtensions)
 			)
 		);
 }
@@ -90,4 +48,13 @@ void FAssetTypeActions_CurveTable::OpenAssetEditor( const TArray<UObject*>& InOb
 	}
 }
 
+void FAssetTypeActions_CurveTable::GetResolvedSourceFilePaths(const TArray<UObject*>& TypeAssets, TArray<FString>& OutSourceFilePaths) const
+{
+	for (auto& Asset : TypeAssets)
+	{
+		const auto CurveTable = CastChecked<UCurveTable>(Asset);
+		OutSourceFilePaths.Add(FReimportManager::ResolveImportFilename(CurveTable->ImportPath, CurveTable));
+	}
+}
+
 #undef LOCTEXT_NAMESPACE
\ No newline at end of file
diff --git a/Engine/Source/Developer/AssetTools/Private/AssetTypeActions/AssetTypeActions_CurveTable.h b/Engine/Source/Developer/AssetTools/Private/AssetTypeActions/AssetTypeActions_CurveTable.h
--- a/Engine/Source/Developer/AssetTools/Private/AssetTypeActions/AssetTypeActions_CurveTable.h
+++ b/Engine/Source/Developer/AssetTools/Private/AssetTypeActions/AssetTypeActions_CurveTable.h
@@ -7,10 +7,9 @@ class FAssetTypeActions_CurveTable : public FAssetTypeActions_CSVAssetBase
 public:
 	// IAssetTypeActions Implementation
 	virtual FText GetName() const override { return NSLOCTEXT("AssetTypeActions", "AssetTypeActions_CurveTable", "Curve Table"); }
-	virtual FColor GetTypeColor() const override { return FColor(62, 140, 35); }
 	virtual UClass* GetSupportedClass() const override { return UCurveTable::StaticClass(); }
-	virtual bool HasActions ( const TArray<UObject*>& InObjects ) const override { return true; }
 	virtual void GetActions( const TArray<UObject*>& InObjects, FMenuBuilder& MenuBuilder ) override;
 	virtual void OpenAssetEditor( const TArray<UObject*>& InObjects, TSharedPtr<class IToolkitHost> EditWithinLevelEditor = TSharedPtr<IToolkitHost>() ) override;
-	virtual uint32 GetCategories() override { return EAssetTypeCategories::Misc; }
+	virtual void GetResolvedSourceFilePaths(const TArray<UObject*>& TypeAssets, TArray<FString>& OutSourceFilePaths) const override;
+	// End IAssetTypeActions
 };
\ No newline at end of file
diff --git a/Engine/Source/Developer/AssetTools/Private/AssetTypeActions/AssetTypeActions_DataTable.cpp b/Engine/Source/Developer/AssetTools/Private/AssetTypeActions/AssetTypeActions_DataTable.cpp
--- a/Engine/Source/Developer/AssetTools/Private/AssetTypeActions/AssetTypeActions_DataTable.cpp
+++ b/Engine/Source/Developer/AssetTools/Private/AssetTypeActions/AssetTypeActions_DataTable.cpp
@@ -10,7 +10,7 @@
 void FAssetTypeActions_DataTable::GetActions( const TArray<UObject*>& InObjects, FMenuBuilder& MenuBuilder )
 {
 	auto Tables = GetTypedWeakObjectPtrs<UObject>(InObjects);
-
+	
 	TArray<FString> ImportPaths;
 	for (auto TableIter = Tables.CreateConstIterator(); TableIter; ++TableIter)
 	{
@@ -22,69 +22,27 @@ void FAssetTypeActions_DataTable::GetActions( const TArray<UObject*>& InObjects,
 	}
 
 	MenuBuilder.AddMenuEntry(
-		LOCTEXT("DataTable_Edit", "Edit"),
-		LOCTEXT("DataTable_EditTooltip", "Opens the selected tables in the table editor."),
+		LOCTEXT("DataTable_JSON", "JSON"),
+		LOCTEXT("DataTable_JSONTooltip", "Creates a JSON version of the data table in the lock."),
 		FSlateIcon(),
 		FUIAction(
-			FExecuteAction::CreateSP( this, &FAssetTypeActions_DataTable::ExecuteEdit, Tables ),
+			FExecuteAction::CreateSP( this, &FAssetTypeActions_DataTable::ExecuteJSON, Tables ),
 			FCanExecuteAction()
 			)
 		);
 
-
-	MenuBuilder.AddMenuEntry(
-		LOCTEXT("DataTable_Reimport", "Reimport"),
-		LOCTEXT("DataTable_ReimportTooltip", "Reimports the selected Data Table from file."),
-		FSlateIcon(),
-		FUIAction(
-			FExecuteAction::CreateSP( this, &FAssetTypeActions_DataTable::ExecuteReimport, Tables ),
-			FCanExecuteAction() 
-			)
-		);
-
-	MenuBuilder.AddMenuEntry(
-		LOCTEXT("DataTable_FindSource", "Find Source"),
-		LOCTEXT("DataTable_FindSourceTooltip", "Opens explorer at the location of this asset's source data."),
-		FSlateIcon(),
-		FUIAction(
-			FExecuteAction::CreateSP( this, &FAssetTypeActions_DataTable::ExecuteFindInExplorer, ImportPaths ),
-			FCanExecuteAction::CreateSP(this, &FAssetTypeActions_DataTable::CanExecuteFindInExplorerSourceCommand, ImportPaths )
-			)
-		);
-
-	static const TArray<FString> EmptyArray;
-	MenuBuilder.AddMenuEntry(
-		LOCTEXT("DataTable_OpenSourceCSV", "Open Source (.csv)"),
-		LOCTEXT("DataTable_OpenSourceCSVTooltip", "Opens the selected asset's source CSV data in an external editor."),
-		FSlateIcon(),
-		FUIAction(
-			FExecuteAction::CreateSP( this, &FAssetTypeActions_DataTable::ExecuteOpenInExternalEditor, ImportPaths, EmptyArray ),
-			FCanExecuteAction::CreateSP(this, &FAssetTypeActions_DataTable::CanExecuteLaunchExternalSourceCommands, ImportPaths, EmptyArray )
-			)
-		);
-
 	TArray<FString> XLSExtensions;
 	XLSExtensions.Add(TEXT(".xls"));
 	XLSExtensions.Add(TEXT(".xlsm"));
 	MenuBuilder.AddMenuEntry(
 		LOCTEXT("DataTable_OpenSourceXLS", "Open Source (.xls/.xlsm)"),
-		LOCTEXT("DataTable_OpenSourceXLSTooltip", "Opens the selected asset's source XLS/XLSM data in an external editor."),
+		LOCTEXT("DataTable_OpenSourceXLSTooltip", "Opens the data table's source XLS/XLSM file in an external editor."),
 		FSlateIcon(),
 		FUIAction(
-			FExecuteAction::CreateSP( this, &FAssetTypeActions_DataTable::ExecuteOpenInExternalEditor, ImportPaths, XLSExtensions ),
-			FCanExecuteAction::CreateSP(this, &FAssetTypeActions_DataTable::CanExecuteLaunchExternalSourceCommands, ImportPaths, XLSExtensions )
+			FExecuteAction::CreateSP( this, &FAssetTypeActions_DataTable::ExecuteFindExcelFileInExplorer, ImportPaths, XLSExtensions ),
+			FCanExecuteAction::CreateSP(this, &FAssetTypeActions_DataTable::CanExecuteFindExcelFileInExplorer, ImportPaths, XLSExtensions)
 			)
 		);
-
-	MenuBuilder.AddMenuEntry(
-		LOCTEXT("DataTable_JSON", "JSON"),
-		LOCTEXT("DataTable_JSONTooltip", "Creates a JSON version of the data table in the lock."),
-		FSlateIcon(),
-		FUIAction(
-		FExecuteAction::CreateSP( this, &FAssetTypeActions_DataTable::ExecuteJSON, Tables ),
-		FCanExecuteAction()
-		)
-		);
 }
 
 void FAssetTypeActions_DataTable::ExecuteJSON(TArray< TWeakObjectPtr<UObject> > Objects)
@@ -112,4 +70,13 @@ void FAssetTypeActions_DataTable::OpenAssetEditor( const TArray<UObject*>& InObj
 	}
 }
 
+void FAssetTypeActions_DataTable::GetResolvedSourceFilePaths(const TArray<UObject*>& TypeAssets, TArray<FString>& OutSourceFilePaths) const
+{
+	for (auto& Asset : TypeAssets)
+	{
+		const auto DataTable = CastChecked<UDataTable>(Asset);
+		OutSourceFilePaths.Add(FReimportManager::ResolveImportFilename(DataTable->ImportPath, DataTable));
+	}
+}
+
 #undef LOCTEXT_NAMESPACE
\ No newline at end of file
diff --git a/Engine/Source/Developer/AssetTools/Private/AssetTypeActions/AssetTypeActions_DataTable.h b/Engine/Source/Developer/AssetTools/Private/AssetTypeActions/AssetTypeActions_DataTable.h
--- a/Engine/Source/Developer/AssetTools/Private/AssetTypeActions/AssetTypeActions_DataTable.h
+++ b/Engine/Source/Developer/AssetTools/Private/AssetTypeActions/AssetTypeActions_DataTable.h
@@ -7,13 +7,13 @@ class FAssetTypeActions_DataTable : public FAssetTypeActions_CSVAssetBase
 public:
 	// IAssetTypeActions Implementation
 	virtual FText GetName() const override { return NSLOCTEXT("AssetTypeActions", "AssetTypeActions_DataTable", "Data Table"); }
-	virtual FColor GetTypeColor() const override { return FColor(62, 140, 35); }
 	virtual UClass* GetSupportedClass() const override { return UDataTable::StaticClass(); }
-	virtual bool HasActions ( const TArray<UObject*>& InObjects ) const override { return true; }
 	virtual void GetActions( const TArray<UObject*>& InObjects, FMenuBuilder& MenuBuilder ) override;
 	virtual void OpenAssetEditor( const TArray<UObject*>& InObjects, TSharedPtr<class IToolkitHost> EditWithinLevelEditor = TSharedPtr<IToolkitHost>() ) override;
-	virtual uint32 GetCategories() override { return EAssetTypeCategories::Misc; }
+	virtual void GetResolvedSourceFilePaths(const TArray<UObject*>& TypeAssets, TArray<FString>& OutSourceFilePaths) const override;
+	// End IAssetTypeActions
 
+private:
 	/** Handler for when JSON is selected */
 	void ExecuteJSON(TArray< TWeakObjectPtr<UObject> > Objects);
 };
\ No newline at end of file
diff --git a/Engine/Source/Developer/AssetTools/Private/AssetTypeActions/AssetTypeActions_DestructibleMesh.cpp b/Engine/Source/Developer/AssetTools/Private/AssetTypeActions/AssetTypeActions_DestructibleMesh.cpp
--- a/Engine/Source/Developer/AssetTools/Private/AssetTypeActions/AssetTypeActions_DestructibleMesh.cpp
+++ b/Engine/Source/Developer/AssetTools/Private/AssetTypeActions/AssetTypeActions_DestructibleMesh.cpp
@@ -14,31 +14,6 @@ UClass* FAssetTypeActions_DestructibleMesh::GetSupportedClass() const
 	return UDestructibleMesh::StaticClass();
 }
 
-void FAssetTypeActions_DestructibleMesh::GetActions( const TArray<UObject*>& InObjects, FMenuBuilder& MenuBuilder )
-{
-	auto Meshes = GetTypedWeakObjectPtrs<UDestructibleMesh>(InObjects);
-
-	MenuBuilder.AddMenuEntry(
-		LOCTEXT("ObjectContext_EditUsingDestructibleMeshEditor", "Edit"),
-		LOCTEXT("ObjectContext_EditUsingDestructibleMeshEditorTooltip", "Opens the selected meshes in the Destructible Mesh editor."),
-		FSlateIcon(),
-		FUIAction(
-			FExecuteAction::CreateSP( this, &FAssetTypeActions_DestructibleMesh::ExecuteEdit, Meshes ),
-			FCanExecuteAction()
-			)
-		);
-
-	MenuBuilder.AddMenuEntry(
-		LOCTEXT("ObjectContext_ReimportDestructibleMesh", "Reimport"),
-		LOCTEXT("ObjectContext_ReimportDestructibleMeshTooltip", "Reimports the selected meshes from file."),
-		FSlateIcon(),
-		FUIAction(
-			FExecuteAction::CreateSP( this, &FAssetTypeActions_DestructibleMesh::ExecuteReimport, Meshes ),
-			FCanExecuteAction()
-			)
-		);
-}
-
 void FAssetTypeActions_DestructibleMesh::OpenAssetEditor( const TArray<UObject*>& InObjects, TSharedPtr<IToolkitHost> EditWithinLevelEditor )
 {
 	for (auto ObjIt = InObjects.CreateConstIterator(); ObjIt; ++ObjIt)
@@ -50,30 +25,4 @@ void FAssetTypeActions_DestructibleMesh::OpenAssetEditor( const TArray<UObject*>
 			TSharedRef< IDestructibleMeshEditor > NewDestructibleMeshEditor = DestructibleMeshEditorModule.CreateDestructibleMeshEditor( EToolkitMode::Standalone, EditWithinLevelEditor, Mesh );
 		}
 	}
-}
-
-void FAssetTypeActions_DestructibleMesh::ExecuteEdit(TArray<TWeakObjectPtr<UDestructibleMesh>> Objects)
-{
-	for (auto ObjIt = Objects.CreateConstIterator(); ObjIt; ++ObjIt)
-	{
-		auto Object = (*ObjIt).Get();
-		if ( Object )
-		{
-			FAssetEditorManager::Get().OpenEditorForAsset(Object);
-		}
-	}
-}
-
-void FAssetTypeActions_DestructibleMesh::ExecuteReimport(TArray<TWeakObjectPtr<UDestructibleMesh>> Objects)
-{
-	for (auto ObjIt = Objects.CreateConstIterator(); ObjIt; ++ObjIt)
-	{
-		auto Object = (*ObjIt).Get();
-		if ( Object )
-		{
-			FReimportManager::Instance()->Reimport(Object, true);
-		}
-	}
-}
-
-#undef LOCTEXT_NAMESPACE
+}
\ No newline at end of file
diff --git a/Engine/Source/Developer/AssetTools/Private/AssetTypeActions/AssetTypeActions_DialogueWave.cpp b/Engine/Source/Developer/AssetTools/Private/AssetTypeActions/AssetTypeActions_DialogueWave.cpp
--- a/Engine/Source/Developer/AssetTools/Private/AssetTypeActions/AssetTypeActions_DialogueWave.cpp
+++ b/Engine/Source/Developer/AssetTools/Private/AssetTypeActions/AssetTypeActions_DialogueWave.cpp
@@ -5,21 +5,6 @@
 
 #define LOCTEXT_NAMESPACE "AssetTypeActions"
 
-void FAssetTypeActions_DialogueWave::GetActions( const TArray<UObject*>& InObjects, FMenuBuilder& MenuBuilder )
-{
-	auto DialogueWaves = GetTypedWeakObjectPtrs<UDialogueWave>(InObjects);
-
-	MenuBuilder.AddMenuEntry(
-		LOCTEXT("DialogueWave_Edit", "Edit"),
-		LOCTEXT("DialogueWave_EditTooltip", "Opens the selected Dialogue Waves in the Dialogue Wave editor."),
-		FSlateIcon(),
-		FUIAction(
-		FExecuteAction::CreateSP( this, &FAssetTypeActions_DialogueWave::ExecuteEdit, DialogueWaves ),
-		FCanExecuteAction()
-		)
-		);
-}
-
 void FAssetTypeActions_DialogueWave::OpenAssetEditor( const TArray<UObject*>& InObjects, TSharedPtr<IToolkitHost> EditWithinLevelEditor )
 {
 	EToolkitMode::Type Mode = EditWithinLevelEditor.IsValid() ? EToolkitMode::WorldCentric : EToolkitMode::Standalone;
@@ -34,16 +19,4 @@ void FAssetTypeActions_DialogueWave::OpenAssetEditor( const TArray<UObject*>& In
 	}
 }
 
-void FAssetTypeActions_DialogueWave::ExecuteEdit(TArray<TWeakObjectPtr<UDialogueWave>> Objects)
-{
-	for (auto ObjIt = Objects.CreateConstIterator(); ObjIt; ++ObjIt)
-	{
-		auto Object = (*ObjIt).Get();
-		if ( Object )
-		{
-			FAssetEditorManager::Get().OpenEditorForAsset(Object);
-		}
-	}
-}
-
 #undef LOCTEXT_NAMESPACE
\ No newline at end of file
diff --git a/Engine/Source/Developer/AssetTools/Private/AssetTypeActions/AssetTypeActions_DialogueWave.h b/Engine/Source/Developer/AssetTools/Private/AssetTypeActions/AssetTypeActions_DialogueWave.h
--- a/Engine/Source/Developer/AssetTools/Private/AssetTypeActions/AssetTypeActions_DialogueWave.h
+++ b/Engine/Source/Developer/AssetTools/Private/AssetTypeActions/AssetTypeActions_DialogueWave.h
@@ -11,13 +11,7 @@ public:
 	virtual FText GetName() const override { return NSLOCTEXT("AssetTypeActions", "AssetTypeActions_DialogueWave", "Dialogue Wave"); }
 	virtual FColor GetTypeColor() const override { return FColor(0,0,255); }
 	virtual UClass* GetSupportedClass() const override { return UDialogueWave::StaticClass(); }
-	virtual bool HasActions ( const TArray<UObject*>& InObjects ) const override { return true; }
-	virtual void GetActions( const TArray<UObject*>& InObjects, FMenuBuilder& MenuBuilder ) override;
 	virtual uint32 GetCategories() override { return EAssetTypeCategories::Sounds; }
 	virtual bool CanFilter() override { return true; }
 	virtual void OpenAssetEditor( const TArray<UObject*>& InObjects, TSharedPtr<class IToolkitHost> EditWithinLevelEditor = TSharedPtr<IToolkitHost>() ) override;
-
-private:
-	/** Handler for when Edit is selected */
-	void ExecuteEdit(TArray<TWeakObjectPtr<UDialogueWave>> Objects);
 };
\ No newline at end of file
diff --git a/Engine/Source/Developer/AssetTools/Private/AssetTypeActions/AssetTypeActions_Font.cpp b/Engine/Source/Developer/AssetTools/Private/AssetTypeActions/AssetTypeActions_Font.cpp
--- a/Engine/Source/Developer/AssetTools/Private/AssetTypeActions/AssetTypeActions_Font.cpp
+++ b/Engine/Source/Developer/AssetTools/Private/AssetTypeActions/AssetTypeActions_Font.cpp
@@ -5,28 +5,15 @@
 
 #define LOCTEXT_NAMESPACE "AssetTypeActions"
 
-void FAssetTypeActions_Font::GetActions( const TArray<UObject*>& InObjects, FMenuBuilder& MenuBuilder )
+void FAssetTypeActions_Font::GetActions(const TArray<UObject*>& InObjects, FMenuBuilder& MenuBuilder)
 {
 	auto Fonts = GetTypedWeakObjectPtrs<UFont>(InObjects);
 
 	MenuBuilder.AddMenuEntry(
-		LOCTEXT("Font_Edit", "Edit"),
-		LOCTEXT("Font_EditTooltip", "Opens the selected fonts in the font editor."),
-		FSlateIcon(),
-		FUIAction(
-			FExecuteAction::CreateSP( this, &FAssetTypeActions_Font::ExecuteEdit, Fonts ),
-			FCanExecuteAction()
-			)
-		);
-
-	MenuBuilder.AddMenuEntry(
-		LOCTEXT("Font_Reimport", "Reimport"),
-		LOCTEXT("Font_ReimportTooltip", "Reimports the selected fonts."),
-		FSlateIcon(),
-		FUIAction(
-			FExecuteAction::CreateSP( this, &FAssetTypeActions_Font::ExecuteReimport, Fonts ),
-			FCanExecuteAction()
-			)
+		LOCTEXT("ReimportFontLabel", "Reimport"),
+		LOCTEXT("ReimportFontTooltip", "Reimport the selected font(s)."),
+		FSlateIcon(FEditorStyle::GetStyleSetName(), "ContentBrowser.AssetActions.ReimportAsset"),
+		FUIAction(FExecuteAction::CreateSP(this, &FAssetTypeActions_Font::ExecuteReimport, Fonts))
 		);
 }
 
@@ -45,26 +32,15 @@ void FAssetTypeActions_Font::OpenAssetEditor( const TArray<UObject*>& InObjects,
 	}
 }
 
-void FAssetTypeActions_Font::ExecuteEdit(TArray<TWeakObjectPtr<UFont>> Objects)
-{
-	for (auto ObjIt = Objects.CreateConstIterator(); ObjIt; ++ObjIt)
-	{
-		auto Object = (*ObjIt).Get();
-		if ( Object )
-		{
-			FAssetEditorManager::Get().OpenEditorForAsset(Object);
-		}
-	}
-}
-
-void FAssetTypeActions_Font::ExecuteReimport(TArray<TWeakObjectPtr<UFont>> Objects)
+void FAssetTypeActions_Font::ExecuteReimport(const TArray<TWeakObjectPtr<UFont>> Objects) const
 {
 	for (auto ObjIt = Objects.CreateConstIterator(); ObjIt; ++ObjIt)
 	{
 		auto Object = (*ObjIt).Get();
-		if ( Object )
+		if (Object)
 		{
-			FReimportManager::Instance()->Reimport( Object );
+			// Fonts fail to reimport if they ask for a new file if missing
+			FReimportManager::Instance()->Reimport(Object, /*bAskForNewFileIfMissing=*/false);
 		}
 	}
 }
diff --git a/Engine/Source/Developer/AssetTools/Private/AssetTypeActions/AssetTypeActions_Font.h b/Engine/Source/Developer/AssetTools/Private/AssetTypeActions/AssetTypeActions_Font.h
--- a/Engine/Source/Developer/AssetTools/Private/AssetTypeActions/AssetTypeActions_Font.h
+++ b/Engine/Source/Developer/AssetTools/Private/AssetTypeActions/AssetTypeActions_Font.h
@@ -9,15 +9,12 @@ public:
 	virtual FText GetName() const override { return NSLOCTEXT("AssetTypeActions", "AssetTypeActions_Font", "Font"); }
 	virtual FColor GetTypeColor() const override { return FColor(255,0,0); }
 	virtual UClass* GetSupportedClass() const override { return UFont::StaticClass(); }
-	virtual bool HasActions ( const TArray<UObject*>& InObjects ) const override { return true; }
-	virtual void GetActions( const TArray<UObject*>& InObjects, FMenuBuilder& MenuBuilder ) override;
+	virtual bool HasActions(const TArray<UObject*>& InObjects) const override { return true; }
+	virtual void GetActions(const TArray<UObject*>& InObjects, FMenuBuilder& MenuBuilder) override;
 	virtual void OpenAssetEditor( const TArray<UObject*>& InObjects, TSharedPtr<class IToolkitHost> EditWithinLevelEditor = TSharedPtr<IToolkitHost>() ) override;
 	virtual uint32 GetCategories() override { return EAssetTypeCategories::MaterialsAndTextures; }
 
 private:
-	/** Handler for when Edit is selected */
-	void ExecuteEdit(TArray<TWeakObjectPtr<UFont>> Objects);
-
 	/** Handler for when Reimport is selected */
-	void ExecuteReimport(TArray<TWeakObjectPtr<UFont>> Objects);
+	void ExecuteReimport(const TArray<TWeakObjectPtr<UFont>> Objects) const;
 };
\ No newline at end of file
diff --git a/Engine/Source/Developer/AssetTools/Private/AssetTypeActions/AssetTypeActions_Material.cpp b/Engine/Source/Developer/AssetTools/Private/AssetTypeActions/AssetTypeActions_Material.cpp
--- a/Engine/Source/Developer/AssetTools/Private/AssetTypeActions/AssetTypeActions_Material.cpp
+++ b/Engine/Source/Developer/AssetTools/Private/AssetTypeActions/AssetTypeActions_Material.cpp
@@ -3,24 +3,6 @@
 #include "AssetToolsPrivatePCH.h"
 #include "Editor/MaterialEditor/Public/MaterialEditorModule.h"
 
-#define LOCTEXT_NAMESPACE "AssetTypeActions"
-
-void FAssetTypeActions_Material::GetActions( const TArray<UObject*>& InObjects, FMenuBuilder& MenuBuilder )
-{
-	auto Materials = GetTypedWeakObjectPtrs<UMaterial>(InObjects);
-
-	MenuBuilder.AddMenuEntry(
-		LOCTEXT("Material_Edit", "Edit"),
-		LOCTEXT("Material_EditTooltip", "Opens the selected materials in the material editor."),
-		FSlateIcon(),
-		FUIAction(
-			FExecuteAction::CreateSP( this, &FAssetTypeActions_Material::ExecuteEdit, Materials ),
-			FCanExecuteAction()
-			)
-		);
-	
-	FAssetTypeActions_MaterialInterface::GetActions(InObjects, MenuBuilder);
-}
 
 void FAssetTypeActions_Material::OpenAssetEditor( const TArray<UObject*>& InObjects, TSharedPtr<IToolkitHost> EditWithinLevelEditor )
 {
@@ -35,18 +17,4 @@ void FAssetTypeActions_Material::OpenAssetEditor( const TArray<UObject*>& InObje
 			MaterialEditorModule->CreateMaterialEditor(Mode, EditWithinLevelEditor, Material);
 		}
 	}
-}
-
-void FAssetTypeActions_Material::ExecuteEdit(TArray<TWeakObjectPtr<UMaterial>> Objects)
-{
-	for (auto ObjIt = Objects.CreateConstIterator(); ObjIt; ++ObjIt)
-	{
-		auto Object = (*ObjIt).Get();
-		if ( Object )
-		{
-			FAssetEditorManager::Get().OpenEditorForAsset(Object);
-		}
-	}
-}
-
-#undef LOCTEXT_NAMESPACE
\ No newline at end of file
+}
\ No newline at end of file
diff --git a/Engine/Source/Developer/AssetTools/Private/AssetTypeActions/AssetTypeActions_Material.h b/Engine/Source/Developer/AssetTools/Private/AssetTypeActions/AssetTypeActions_Material.h
--- a/Engine/Source/Developer/AssetTools/Private/AssetTypeActions/AssetTypeActions_Material.h
+++ b/Engine/Source/Developer/AssetTools/Private/AssetTypeActions/AssetTypeActions_Material.h
@@ -9,12 +9,7 @@ public:
 	virtual FText GetName() const override { return NSLOCTEXT("AssetTypeActions", "AssetTypeActions_Material", "Material"); }
 	virtual FColor GetTypeColor() const override { return FColor(0,255,0); }
 	virtual UClass* GetSupportedClass() const override { return UMaterial::StaticClass(); }
-	virtual void GetActions( const TArray<UObject*>& InObjects, FMenuBuilder& MenuBuilder ) override;
 	virtual void OpenAssetEditor( const TArray<UObject*>& InObjects, TSharedPtr<class IToolkitHost> EditWithinLevelEditor = TSharedPtr<IToolkitHost>() ) override;
 	virtual bool CanFilter() override { return true; }
 	virtual uint32 GetCategories() override { return FAssetTypeActions_MaterialInterface::GetCategories() | EAssetTypeCategories::Basic; }
-
-private:
-	/** Handler for when Edit is selected */
-	void ExecuteEdit(TArray<TWeakObjectPtr<UMaterial>> Objects);
 };
\ No newline at end of file
diff --git a/Engine/Source/Developer/AssetTools/Private/AssetTypeActions/AssetTypeActions_MaterialFunction.cpp b/Engine/Source/Developer/AssetTools/Private/AssetTypeActions/AssetTypeActions_MaterialFunction.cpp
--- a/Engine/Source/Developer/AssetTools/Private/AssetTypeActions/AssetTypeActions_MaterialFunction.cpp
+++ b/Engine/Source/Developer/AssetTools/Private/AssetTypeActions/AssetTypeActions_MaterialFunction.cpp
@@ -9,20 +9,10 @@ void FAssetTypeActions_MaterialFunction::GetActions( const TArray<UObject*>& InO
 {
 	auto Materials = GetTypedWeakObjectPtrs<UMaterialFunction>(InObjects);
 
-	MenuBuilder.AddMenuEntry(
-		LOCTEXT("MaterialFunction_Edit", "Edit"),
-		LOCTEXT("MaterialFunction_EditTooltip", "Opens the selected material functions in the material editor."),
-		FSlateIcon(),
-		FUIAction(
-			FExecuteAction::CreateSP( this, &FAssetTypeActions_MaterialFunction::ExecuteEdit, Materials ),
-			FCanExecuteAction()
-			)
-		);
-
 	MenuBuilder.AddMenuEntry(
 		LOCTEXT("MaterialFunction_FindMaterials", "Find Materials Using This"),
 		LOCTEXT("MaterialFunction_FindMaterialsTooltip", "Finds the materials that reference this material function in the content browser."),
-		FSlateIcon(),
+		FSlateIcon(FEditorStyle::GetStyleSetName(), "Kismet.Tabs.FindResults"),
 		FUIAction(
 			FExecuteAction::CreateSP( this, &FAssetTypeActions_MaterialFunction::ExecuteFindMaterials, Materials ),
 			FCanExecuteAction()
@@ -45,18 +35,6 @@ void FAssetTypeActions_MaterialFunction::OpenAssetEditor( const TArray<UObject*>
 	}
 }
 
-void FAssetTypeActions_MaterialFunction::ExecuteEdit(TArray<TWeakObjectPtr<UMaterialFunction>> Objects)
-{
-	for (auto ObjIt = Objects.CreateConstIterator(); ObjIt; ++ObjIt)
-	{
-		auto Object = (*ObjIt).Get();
-		if ( Object )
-		{
-			FAssetEditorManager::Get().OpenEditorForAsset(Object);
-		}
-	}
-}
-
 void FAssetTypeActions_MaterialFunction::ExecuteFindMaterials(TArray<TWeakObjectPtr<UMaterialFunction>> Objects)
 {
 	TArray<UObject*> ObjectsToSync;
diff --git a/Engine/Source/Developer/AssetTools/Private/AssetTypeActions/AssetTypeActions_MaterialFunction.h b/Engine/Source/Developer/AssetTools/Private/AssetTypeActions/AssetTypeActions_MaterialFunction.h
--- a/Engine/Source/Developer/AssetTools/Private/AssetTypeActions/AssetTypeActions_MaterialFunction.h
+++ b/Engine/Source/Developer/AssetTools/Private/AssetTypeActions/AssetTypeActions_MaterialFunction.h
@@ -18,9 +18,6 @@ public:
 	virtual class UThumbnailInfo* GetThumbnailInfo(UObject* Asset) const override;
 
 private:
-	/** Handler for when Edit is selected */
-	void ExecuteEdit(TArray<TWeakObjectPtr<UMaterialFunction>> Objects);
-
 	/** Handler for when FindMaterials is selected */
 	void ExecuteFindMaterials(TArray<TWeakObjectPtr<UMaterialFunction>> Objects);
 };
\ No newline at end of file
diff --git a/Engine/Source/Developer/AssetTools/Private/AssetTypeActions/AssetTypeActions_MaterialInstanceConstant.cpp b/Engine/Source/Developer/AssetTools/Private/AssetTypeActions/AssetTypeActions_MaterialInstanceConstant.cpp
--- a/Engine/Source/Developer/AssetTools/Private/AssetTypeActions/AssetTypeActions_MaterialInstanceConstant.cpp
+++ b/Engine/Source/Developer/AssetTools/Private/AssetTypeActions/AssetTypeActions_MaterialInstanceConstant.cpp
@@ -9,16 +9,6 @@ void FAssetTypeActions_MaterialInstanceConstant::GetActions( const TArray<UObjec
 {
 	auto MICs = GetTypedWeakObjectPtrs<UMaterialInstanceConstant>(InObjects);
 
-	MenuBuilder.AddMenuEntry(
-		LOCTEXT("MaterialInstanceConstant_Edit", "Edit"),
-		LOCTEXT("MaterialInstanceConstant_EditTooltip", "Opens the selected materials in the material editor."),
-		FSlateIcon(),
-		FUIAction(
-			FExecuteAction::CreateSP( this, &FAssetTypeActions_MaterialInstanceConstant::ExecuteEdit, MICs ),
-			FCanExecuteAction()
-			)
-		);
-
 	FAssetTypeActions_MaterialInterface::GetActions(InObjects, MenuBuilder);
 
 	MenuBuilder.AddMenuEntry(
@@ -47,18 +37,6 @@ void FAssetTypeActions_MaterialInstanceConstant::OpenAssetEditor( const TArray<U
 	}
 }
 
-void FAssetTypeActions_MaterialInstanceConstant::ExecuteEdit(TArray<TWeakObjectPtr<UMaterialInstanceConstant>> Objects)
-{
-	for (auto ObjIt = Objects.CreateConstIterator(); ObjIt; ++ObjIt)
-	{
-		auto Object = (*ObjIt).Get();
-		if ( Object )
-		{
-			FAssetEditorManager::Get().OpenEditorForAsset(Object);
-		}
-	}
-}
-
 void FAssetTypeActions_MaterialInstanceConstant::ExecuteFindParent(TArray<TWeakObjectPtr<UMaterialInstanceConstant>> Objects)
 {
 	TArray<UObject*> ObjectsToSyncTo;
diff --git a/Engine/Source/Developer/AssetTools/Private/AssetTypeActions/AssetTypeActions_MaterialInstanceConstant.h b/Engine/Source/Developer/AssetTools/Private/AssetTypeActions/AssetTypeActions_MaterialInstanceConstant.h
--- a/Engine/Source/Developer/AssetTools/Private/AssetTypeActions/AssetTypeActions_MaterialInstanceConstant.h
+++ b/Engine/Source/Developer/AssetTools/Private/AssetTypeActions/AssetTypeActions_MaterialInstanceConstant.h
@@ -14,9 +14,6 @@ public:
 	virtual bool CanFilter() override { return true; }
 
 private:
-	/** Handler for when Edit is selected */
-	void ExecuteEdit(TArray<TWeakObjectPtr<UMaterialInstanceConstant>> Objects);
-
 	/** Handler for when FindParent is selected */
 	void ExecuteFindParent(TArray<TWeakObjectPtr<UMaterialInstanceConstant>> Objects);
 };
\ No newline at end of file
diff --git a/Engine/Source/Developer/AssetTools/Private/AssetTypeActions/AssetTypeActions_MaterialInterface.cpp b/Engine/Source/Developer/AssetTools/Private/AssetTypeActions/AssetTypeActions_MaterialInterface.cpp
--- a/Engine/Source/Developer/AssetTools/Private/AssetTypeActions/AssetTypeActions_MaterialInterface.cpp
+++ b/Engine/Source/Developer/AssetTools/Private/AssetTypeActions/AssetTypeActions_MaterialInterface.cpp
@@ -15,7 +15,7 @@ void FAssetTypeActions_MaterialInterface::GetActions( const TArray<UObject*>& In
 	MenuBuilder.AddMenuEntry(
 		LOCTEXT("Material_NewMIC", "Create Material Instance"),
 		LOCTEXT("Material_NewMICTooltip", "Creates a parameterized material using this material as a base."),
-		FSlateIcon(),
+		FSlateIcon(FEditorStyle::GetStyleSetName(), "ClassIcon.MaterialInstanceActor"),
 		FUIAction(
 			FExecuteAction::CreateSP( this, &FAssetTypeActions_MaterialInterface::ExecuteNewMIC, MaterialInterfaces )
 			)
diff --git a/Engine/Source/Developer/AssetTools/Private/AssetTypeActions/AssetTypeActions_MorphTarget.cpp b/Engine/Source/Developer/AssetTools/Private/AssetTypeActions/AssetTypeActions_MorphTarget.cpp
--- a/Engine/Source/Developer/AssetTools/Private/AssetTypeActions/AssetTypeActions_MorphTarget.cpp
+++ b/Engine/Source/Developer/AssetTools/Private/AssetTypeActions/AssetTypeActions_MorphTarget.cpp
@@ -13,16 +13,6 @@ void FAssetTypeActions_MorphTarget::GetActions( const TArray<UObject*>& InObject
 {
 	auto MorphTargets = GetTypedWeakObjectPtrs<UMorphTarget>(InObjects);
 
-	MenuBuilder.AddMenuEntry(
-		LOCTEXT("MorphTarget_Edit", "Edit"),
-		LOCTEXT("MorphTarget_EditTooltip", "Opens the selected morph targets in Persona."),
-		FSlateIcon(),
-		FUIAction(
-		FExecuteAction::CreateSP( this, &FAssetTypeActions_MorphTarget::ExecuteEdit, MorphTargets ),
-		FCanExecuteAction()
-		)
-		);
-
 	MenuBuilder.AddMenuEntry(
 		LOCTEXT("MorphTarget_Convert", "Move to SkeletalMesh"),
 		LOCTEXT("MorphTarget_ConvertTooltip", "Moves the selected morph targets to its own Mesh and delete as an asset."),
@@ -52,18 +42,6 @@ void FAssetTypeActions_MorphTarget::OpenAssetEditor( const TArray<UObject*>& InO
 	}
 }
 
-void FAssetTypeActions_MorphTarget::ExecuteEdit(TArray<TWeakObjectPtr<UMorphTarget>> Objects)
-{
-	for (auto ObjIt = Objects.CreateConstIterator(); ObjIt; ++ObjIt)
-	{
-		auto Object = (*ObjIt).Get();
-		if ( Object && Object->BaseSkelMesh )
-		{
-			FAssetEditorManager::Get().OpenEditorForAsset(Object);
-		}
-	}
-}
-
 void FAssetTypeActions_MorphTarget::ExecuteMovetoMesh(TArray<TWeakObjectPtr<UMorphTarget>> Objects)
 {
 	if ( FMessageDialog::Open(EAppMsgType::YesNo, LOCTEXT("MoveToMeshWarning", "This will move the selected morphtargets to internal SkeletalMesh and delete the assets. This is recommended. Would you like to continue?")) == EAppReturnType::Yes )
diff --git a/Engine/Source/Developer/AssetTools/Private/AssetTypeActions/AssetTypeActions_MorphTarget.h b/Engine/Source/Developer/AssetTools/Private/AssetTypeActions/AssetTypeActions_MorphTarget.h
--- a/Engine/Source/Developer/AssetTools/Private/AssetTypeActions/AssetTypeActions_MorphTarget.h
+++ b/Engine/Source/Developer/AssetTools/Private/AssetTypeActions/AssetTypeActions_MorphTarget.h
@@ -15,8 +15,6 @@ public:
 	virtual uint32 GetCategories() override { return EAssetTypeCategories::Animation; }
 
 private:
-	/** Handler for when Edit is selected */
-	void ExecuteEdit(TArray<TWeakObjectPtr<UMorphTarget>> Objects);
 	/** Handler for when Move to Mesh is requested **/
 	void ExecuteMovetoMesh(TArray<TWeakObjectPtr<UMorphTarget>> Objects);
 };
\ No newline at end of file
diff --git a/Engine/Source/Developer/AssetTools/Private/AssetTypeActions/AssetTypeActions_NiagaraEffect.cpp b/Engine/Source/Developer/AssetTools/Private/AssetTypeActions/AssetTypeActions_NiagaraEffect.cpp
--- a/Engine/Source/Developer/AssetTools/Private/AssetTypeActions/AssetTypeActions_NiagaraEffect.cpp
+++ b/Engine/Source/Developer/AssetTools/Private/AssetTypeActions/AssetTypeActions_NiagaraEffect.cpp
@@ -10,21 +10,6 @@
 
 #define LOCTEXT_NAMESPACE "AssetTypeActions"
 
-void FAssetTypeActions_NiagaraEffect::GetActions(const TArray<UObject*>& InObjects, FMenuBuilder& MenuBuilder)
-{
-	auto Scripts = GetTypedWeakObjectPtrs<UNiagaraEffect>(InObjects);
-
-	MenuBuilder.AddMenuEntry(
-		LOCTEXT("NiagaraEffect_Edit", "Edit"),
-		LOCTEXT("NiagaraEffect_EditTooltip", "Opens the selected effect in editor."),
-		FSlateIcon(),
-		FUIAction(
-		FExecuteAction::CreateSP(this, &FAssetTypeActions_NiagaraEffect::ExecuteEdit, Scripts),
-		FCanExecuteAction()
-		)
-		);
-}
-
 void FAssetTypeActions_NiagaraEffect::OpenAssetEditor(const TArray<UObject*>& InObjects, TSharedPtr<IToolkitHost> EditWithinLevelEditor)
 {
 	EToolkitMode::Type Mode = EditWithinLevelEditor.IsValid() ? EToolkitMode::WorldCentric : EToolkitMode::Standalone;
@@ -40,18 +25,6 @@ void FAssetTypeActions_NiagaraEffect::OpenAssetEditor(const TArray<UObject*>& In
 	}
 }
 
-void FAssetTypeActions_NiagaraEffect::ExecuteEdit(TArray<TWeakObjectPtr<UNiagaraEffect>> Objects)
-{
-	for (auto ObjIt = Objects.CreateConstIterator(); ObjIt; ++ObjIt)
-	{
-		auto Object = (*ObjIt).Get();
-		if (Object)
-		{
-			FAssetEditorManager::Get().OpenEditorForAsset(Object);
-		}
-	}
-}
-
 UClass* FAssetTypeActions_NiagaraEffect::GetSupportedClass() const
 {
 	return UNiagaraEffect::StaticClass();
diff --git a/Engine/Source/Developer/AssetTools/Private/AssetTypeActions/AssetTypeActions_NiagaraEffect.h b/Engine/Source/Developer/AssetTools/Private/AssetTypeActions/AssetTypeActions_NiagaraEffect.h
--- a/Engine/Source/Developer/AssetTools/Private/AssetTypeActions/AssetTypeActions_NiagaraEffect.h
+++ b/Engine/Source/Developer/AssetTools/Private/AssetTypeActions/AssetTypeActions_NiagaraEffect.h
@@ -9,12 +9,6 @@ public:
 	virtual FText GetName() const override { return NSLOCTEXT("AssetTypeActions", "AssetTypeActions_NiagaraEffect", "Niagara Effect"); }
 	virtual FColor GetTypeColor() const override { return FColor(255, 194, 0); }
 	virtual UClass* GetSupportedClass() const override;
-	virtual bool HasActions(const TArray<UObject*>& InObjects) const override { return true; }
-	virtual void GetActions(const TArray<UObject*>& InObjects, FMenuBuilder& MenuBuilder) override;
 	virtual void OpenAssetEditor(const TArray<UObject*>& InObjects, TSharedPtr<class IToolkitHost> EditWithinLevelEditor = TSharedPtr<IToolkitHost>()) override;
 	virtual uint32 GetCategories() override { return EAssetTypeCategories::Misc; }
-
-private:
-	/** Handler for when Edit is selected */
-	void ExecuteEdit(TArray<TWeakObjectPtr<class UNiagaraEffect>> Objects);
 };
\ No newline at end of file
diff --git a/Engine/Source/Developer/AssetTools/Private/AssetTypeActions/AssetTypeActions_NiagaraScript.cpp b/Engine/Source/Developer/AssetTools/Private/AssetTypeActions/AssetTypeActions_NiagaraScript.cpp
--- a/Engine/Source/Developer/AssetTools/Private/AssetTypeActions/AssetTypeActions_NiagaraScript.cpp
+++ b/Engine/Source/Developer/AssetTools/Private/AssetTypeActions/AssetTypeActions_NiagaraScript.cpp
@@ -10,21 +10,6 @@
 
 #define LOCTEXT_NAMESPACE "AssetTypeActions"
 
-void FAssetTypeActions_NiagaraScript::GetActions( const TArray<UObject*>& InObjects, FMenuBuilder& MenuBuilder )
-{
-	auto Scripts = GetTypedWeakObjectPtrs<UNiagaraScript>(InObjects);
-
-	MenuBuilder.AddMenuEntry(
-		LOCTEXT("NiagaraScript_Edit", "Edit"),
-		LOCTEXT("NiagaraScript_EditTooltip", "Opens the selected script in editor."),
-		FSlateIcon(),
-		FUIAction(
-		FExecuteAction::CreateSP( this, &FAssetTypeActions_NiagaraScript::ExecuteEdit, Scripts ),
-		FCanExecuteAction()
-		)
-		);
-}
-
 void FAssetTypeActions_NiagaraScript::OpenAssetEditor( const TArray<UObject*>& InObjects, TSharedPtr<IToolkitHost> EditWithinLevelEditor )
 {
 	EToolkitMode::Type Mode = EditWithinLevelEditor.IsValid() ? EToolkitMode::WorldCentric : EToolkitMode::Standalone;
@@ -40,18 +25,6 @@ void FAssetTypeActions_NiagaraScript::OpenAssetEditor( const TArray<UObject*>& I
 	}
 }
 
-void FAssetTypeActions_NiagaraScript::ExecuteEdit(TArray<TWeakObjectPtr<UNiagaraScript>> Objects)
-{
-	for (auto ObjIt = Objects.CreateConstIterator(); ObjIt; ++ObjIt)
-	{
-		auto Object = (*ObjIt).Get();
-		if ( Object )
-		{
-			FAssetEditorManager::Get().OpenEditorForAsset(Object);
-		}
-	}
-}
-
 UClass* FAssetTypeActions_NiagaraScript::GetSupportedClass() const
 { 
 	return UNiagaraScript::StaticClass(); 
diff --git a/Engine/Source/Developer/AssetTools/Private/AssetTypeActions/AssetTypeActions_NiagaraScript.h b/Engine/Source/Developer/AssetTools/Private/AssetTypeActions/AssetTypeActions_NiagaraScript.h
--- a/Engine/Source/Developer/AssetTools/Private/AssetTypeActions/AssetTypeActions_NiagaraScript.h
+++ b/Engine/Source/Developer/AssetTools/Private/AssetTypeActions/AssetTypeActions_NiagaraScript.h
@@ -9,12 +9,6 @@ public:
 	virtual FText GetName() const override { return NSLOCTEXT("AssetTypeActions", "AssetTypeActions_NiagaraScript", "Niagara Script"); }
 	virtual FColor GetTypeColor() const override { return FColor(255,194,0); }
 	virtual UClass* GetSupportedClass() const override;
-	virtual bool HasActions ( const TArray<UObject*>& InObjects ) const override { return true; }
-	virtual void GetActions( const TArray<UObject*>& InObjects, FMenuBuilder& MenuBuilder ) override;
 	virtual void OpenAssetEditor( const TArray<UObject*>& InObjects, TSharedPtr<class IToolkitHost> EditWithinLevelEditor = TSharedPtr<IToolkitHost>() ) override;
 	virtual uint32 GetCategories() override { return EAssetTypeCategories::Misc; }
-
-private:
-	/** Handler for when Edit is selected */
-	void ExecuteEdit(TArray<TWeakObjectPtr<class UNiagaraScript>> Objects);
 };
\ No newline at end of file
diff --git a/Engine/Source/Developer/AssetTools/Private/AssetTypeActions/AssetTypeActions_ParticleSystem.cpp b/Engine/Source/Developer/AssetTools/Private/AssetTypeActions/AssetTypeActions_ParticleSystem.cpp
--- a/Engine/Source/Developer/AssetTools/Private/AssetTypeActions/AssetTypeActions_ParticleSystem.cpp
+++ b/Engine/Source/Developer/AssetTools/Private/AssetTypeActions/AssetTypeActions_ParticleSystem.cpp
@@ -11,16 +11,6 @@ void FAssetTypeActions_ParticleSystem::GetActions( const TArray<UObject*>& InObj
 {
 	auto ParticleSystems = GetTypedWeakObjectPtrs<UParticleSystem>(InObjects);
 
-	MenuBuilder.AddMenuEntry(
-		LOCTEXT("ParticleSystem_Edit", "Edit"),
-		LOCTEXT("ParticleSystem_EditTooltip", "Opens the selected particle systems in cascade."),
-		FSlateIcon(),
-		FUIAction(
-			FExecuteAction::CreateSP( this, &FAssetTypeActions_ParticleSystem::ExecuteEdit, ParticleSystems ),
-			FCanExecuteAction()
-			)
-		);
-
 	MenuBuilder.AddMenuEntry(
 		LOCTEXT("ParticleSystem_CopyParameters", "Copy Parameters"),
 		LOCTEXT("ParticleSystem_CopyParametersTooltip", "Copies particle system parameters to the clipboard."),
@@ -57,18 +47,6 @@ void FAssetTypeActions_ParticleSystem::OpenAssetEditor( const TArray<UObject*>&
 	}
 }
 
-void FAssetTypeActions_ParticleSystem::ExecuteEdit(TArray<TWeakObjectPtr<UParticleSystem>> Objects)
-{
-	for (auto ObjIt = Objects.CreateConstIterator(); ObjIt; ++ObjIt)
-	{
-		auto Object = (*ObjIt).Get();
-		if ( Object )
-		{
-			FAssetEditorManager::Get().OpenEditorForAsset(Object);
-		}
-	}
-}
-
 void FAssetTypeActions_ParticleSystem::ExecuteCopyParameters(TArray<TWeakObjectPtr<UParticleSystem>> Objects)
 {
 	FString ClipboardString = TEXT("");
diff --git a/Engine/Source/Developer/AssetTools/Private/AssetTypeActions/AssetTypeActions_ParticleSystem.h b/Engine/Source/Developer/AssetTools/Private/AssetTypeActions/AssetTypeActions_ParticleSystem.h
--- a/Engine/Source/Developer/AssetTools/Private/AssetTypeActions/AssetTypeActions_ParticleSystem.h
+++ b/Engine/Source/Developer/AssetTools/Private/AssetTypeActions/AssetTypeActions_ParticleSystem.h
@@ -17,9 +17,6 @@ public:
 	virtual uint32 GetCategories() override { return EAssetTypeCategories::Basic; }
 	
 private:
-	/** Handler for when Edit is selected */
-	void ExecuteEdit(TArray<TWeakObjectPtr<UParticleSystem>> Objects);
-
 	/** Handler for when Copy Parameters is selected */
 	void ExecuteCopyParameters(TArray<TWeakObjectPtr<UParticleSystem>> Objects);
 
diff --git a/Engine/Source/Developer/AssetTools/Private/AssetTypeActions/AssetTypeActions_PhysicsAsset.cpp b/Engine/Source/Developer/AssetTools/Private/AssetTypeActions/AssetTypeActions_PhysicsAsset.cpp
--- a/Engine/Source/Developer/AssetTools/Private/AssetTypeActions/AssetTypeActions_PhysicsAsset.cpp
+++ b/Engine/Source/Developer/AssetTools/Private/AssetTypeActions/AssetTypeActions_PhysicsAsset.cpp
@@ -11,21 +11,6 @@ UClass* FAssetTypeActions_PhysicsAsset::GetSupportedClass() const
 	return UPhysicsAsset::StaticClass();
 }
 
-void FAssetTypeActions_PhysicsAsset::GetActions( const TArray<UObject*>& InObjects, FMenuBuilder& MenuBuilder )
-{
-	auto PhysicsAssets = GetTypedWeakObjectPtrs<UPhysicsAsset>(InObjects);
-
-	MenuBuilder.AddMenuEntry(
-		LOCTEXT("PhysicsAsset_Edit", "Edit"),
-		LOCTEXT("PhysicsAsset_EditTooltip", "Opens the selected physics asset in PhAT."),
-		FSlateIcon(),
-		FUIAction(
-			FExecuteAction::CreateSP( this, &FAssetTypeActions_PhysicsAsset::ExecuteEdit, PhysicsAssets ),
-			FCanExecuteAction()
-			)
-		);
-}
-
 void FAssetTypeActions_PhysicsAsset::OpenAssetEditor( const TArray<UObject*>& InObjects, TSharedPtr<IToolkitHost> EditWithinLevelEditor )
 {
 	EToolkitMode::Type Mode = EditWithinLevelEditor.IsValid() ? EToolkitMode::WorldCentric : EToolkitMode::Standalone;
@@ -42,16 +27,4 @@ void FAssetTypeActions_PhysicsAsset::OpenAssetEditor( const TArray<UObject*>& In
 	}
 }
 
-void FAssetTypeActions_PhysicsAsset::ExecuteEdit(TArray<TWeakObjectPtr<UPhysicsAsset>> Objects)
-{
-	for (auto ObjIt = Objects.CreateConstIterator(); ObjIt; ++ObjIt)
-	{
-		auto Object = (*ObjIt).Get();
-		if ( Object )
-		{
-			FAssetEditorManager::Get().OpenEditorForAsset(Object);
-		}
-	}
-}
-
 #undef LOCTEXT_NAMESPACE
\ No newline at end of file
diff --git a/Engine/Source/Developer/AssetTools/Private/AssetTypeActions/AssetTypeActions_PhysicsAsset.h b/Engine/Source/Developer/AssetTools/Private/AssetTypeActions/AssetTypeActions_PhysicsAsset.h
--- a/Engine/Source/Developer/AssetTools/Private/AssetTypeActions/AssetTypeActions_PhysicsAsset.h
+++ b/Engine/Source/Developer/AssetTools/Private/AssetTypeActions/AssetTypeActions_PhysicsAsset.h
@@ -9,12 +9,6 @@ public:
 	virtual FText GetName() const override { return NSLOCTEXT("AssetTypeActions", "AssetTypeActions_PhysicsAsset", "Physics Asset"); }
 	virtual FColor GetTypeColor() const override { return FColor(255,192,128); }
 	virtual UClass* GetSupportedClass() const override;
-	virtual bool HasActions ( const TArray<UObject*>& InObjects ) const override { return true; }
-	virtual void GetActions( const TArray<UObject*>& InObjects, FMenuBuilder& MenuBuilder ) override;
 	virtual void OpenAssetEditor( const TArray<UObject*>& InObjects, TSharedPtr<class IToolkitHost> EditWithinLevelEditor = TSharedPtr<IToolkitHost>() ) override;
 	virtual uint32 GetCategories() override { return EAssetTypeCategories::Physics; }
-
-private:
-	/** Handler for when Edit is selected */
-	void ExecuteEdit(TArray<TWeakObjectPtr<UPhysicsAsset>> Objects);
 };
\ No newline at end of file
diff --git a/Engine/Source/Developer/AssetTools/Private/AssetTypeActions/AssetTypeActions_SkeletalMesh.cpp b/Engine/Source/Developer/AssetTools/Private/AssetTypeActions/AssetTypeActions_SkeletalMesh.cpp
--- a/Engine/Source/Developer/AssetTools/Private/AssetTypeActions/AssetTypeActions_SkeletalMesh.cpp
+++ b/Engine/Source/Developer/AssetTools/Private/AssetTypeActions/AssetTypeActions_SkeletalMesh.cpp
@@ -8,6 +8,7 @@
 #include "Editor/UnrealEd/Public/SSkeletonWidget.h"
 #include "Editor/PhAT/Public/PhATModule.h"
 #include "Editor/Persona/Public/PersonaModule.h"
+#include "AnimationEditorUtils.h"
 
 #include "ContentBrowserModule.h"
 #include "AssetToolsModule.h"
@@ -418,60 +419,50 @@ void FAssetTypeActions_SkeletalMesh::GetActions( const TArray<UObject*>& InObjec
 {
 	auto Meshes = GetTypedWeakObjectPtrs<USkeletalMesh>(InObjects);
 
-	MenuBuilder.AddMenuEntry(
-		LOCTEXT("SkeletalMesh_Edit", "Edit"),
-		LOCTEXT("SkeletalMesh_EditTooltip", "Opens the selected meshes in Persona."),
-		FSlateIcon(),
-		FUIAction(
-			FExecuteAction::CreateSP( this, &FAssetTypeActions_SkeletalMesh::ExecuteEdit, Meshes ),
-			FCanExecuteAction()
-			)
-		);
-
 	MenuBuilder.AddSubMenu(
 			LOCTEXT("CreateSkeletalMeshSubmenu", "Create"),
 			LOCTEXT("CreateSkeletalMeshSubmenu_ToolTip", "Create related assets"),
-			FNewMenuDelegate::CreateSP(this, &FAssetTypeActions_SkeletalMesh::FillCreateMenu, Meshes));
-
-	MenuBuilder.AddMenuEntry(
-		LOCTEXT("SkeletalMesh_Reimport", "Reimport"),
-		LOCTEXT("SkeletalMesh_ReimportTooltip", "Reimports the selected meshes from file."),
-		FSlateIcon(),
-		FUIAction(
-			FExecuteAction::CreateSP( this, &FAssetTypeActions_SkeletalMesh::ExecuteReimport, Meshes ),
-			FCanExecuteAction()
-			)
+			FNewMenuDelegate::CreateSP(this, &FAssetTypeActions_SkeletalMesh::FillCreateMenu, Meshes),
+			false,
+			FSlateIcon(FEditorStyle::GetStyleSetName(), "Persona.AssetActions.CreateAnimAsset")
+			);
+
+	MenuBuilder.AddSubMenu(	
+		LOCTEXT("SkeletalMesh_LODImport", "Import LOD"),
+		LOCTEXT("SkeletalMesh_LODImportTooltip", "Select which LODs to import."),
+		FNewMenuDelegate::CreateSP(this, &FAssetTypeActions_SkeletalMesh::GetLODMenu, Meshes)
 		);
 
-	MenuBuilder.AddSubMenu(	LOCTEXT("SkeletalMesh_LODImport", "LOD Import"),
-							LOCTEXT("SkeletalMesh_LODImportTooltip", "Select which LODs to import."),
-							FNewMenuDelegate::CreateSP( this, &FAssetTypeActions_SkeletalMesh::GetLODMenu, Meshes));
-
-	// source menu
-	MenuBuilder.AddSubMenu(
-			LOCTEXT("SourceSkeletalMeshSubmenu", "Source"),
-			LOCTEXT("SourceSkeletalMeshSubmenu_ToolTip", "Source data related"),
-			FNewMenuDelegate::CreateSP(this, &FAssetTypeActions_SkeletalMesh::FillSourceMenu, Meshes));
-
 	// Add actions that do not apply to destructible meshes
 	GetNonDestructibleActions(Meshes, MenuBuilder);
 }
 
 void FAssetTypeActions_SkeletalMesh::FillCreateMenu(FMenuBuilder& MenuBuilder, const TArray<TWeakObjectPtr<USkeletalMesh>> Meshes) const
 {
-	MenuBuilder.AddSubMenu(LOCTEXT("SkeletalMesh_CreateAssets", "Create Animation Assets"),
-		LOCTEXT("SkeletalMesh_CreateAssetsTooltip", "Create animation assets from this skeletal mesh."),
-		FNewMenuDelegate::CreateSP(this, &FAssetTypeActions_SkeletalMesh::GetCreateMenu, Meshes));
+	MenuBuilder.BeginSection("CreatePhysicsAsset", LOCTEXT("CreatePhysicsAssetMenuHeading", "Physics Asset"));
+	{
+		MenuBuilder.AddMenuEntry(
+			LOCTEXT("SkeletalMesh_NewPhysicsAsset", "Physics Asset"),
+			LOCTEXT("SkeletalMesh_NewPhysicsAssetTooltip", "Creates a new physics asset for each of the selected meshes."),
+			FSlateIcon(),
+			FUIAction(
+				FExecuteAction::CreateSP(this, &FAssetTypeActions_SkeletalMesh::ExecuteNewPhysicsAsset, Meshes),
+				FCanExecuteAction()
+				)
+			);
+	}
+	MenuBuilder.EndSection();
+
+	// Get the skeleton for each selected skeletal mesh
+	TArray<TWeakObjectPtr<USkeleton>> Skeletons;
+	Skeletons.Reserve(Meshes.Num());
+	for (int32 i = 0; i < Meshes.Num(); i++)
+	{
+		Skeletons.Add(Meshes[i]->Skeleton);
+	}
+	
+	AnimationEditorUtils::FillCreateAssetMenu(MenuBuilder, Skeletons, FAnimAssetCreated::CreateSP(this, &FAssetTypeActions_SkeletalMesh::OnAssetCreated));
 
-	MenuBuilder.AddMenuEntry(
-		LOCTEXT("SkeletalMesh_NewPhysicsAsset", "Create Physics Asset"),
-		LOCTEXT("SkeletalMesh_NewPhysicsAssetTooltip", "Creates a new physics asset for each of the selected meshes."),
-		FSlateIcon(),
-		FUIAction(
-			FExecuteAction::CreateSP(this, &FAssetTypeActions_SkeletalMesh::ExecuteNewPhysicsAsset, Meshes),
-			FCanExecuteAction()
-			)
-		);
 }
 
 void FAssetTypeActions_SkeletalMesh::GetNonDestructibleActions( const TArray<TWeakObjectPtr<USkeletalMesh>>& Meshes, FMenuBuilder& MenuBuilder)
@@ -553,142 +544,14 @@ UThumbnailInfo* FAssetTypeActions_SkeletalMesh::GetThumbnailInfo(UObject* Asset)
 	return ThumbnailInfo;
 }
 
-void FAssetTypeActions_SkeletalMesh::ExecuteEdit(TArray<TWeakObjectPtr<USkeletalMesh>> Objects)
-{
-	for (auto ObjIt = Objects.CreateConstIterator(); ObjIt; ++ObjIt)
-	{
-		auto Object = (*ObjIt).Get();
-		if ( Object )
-		{
-			FAssetEditorManager::Get().OpenEditorForAsset(Object);
-		}
-	}
-}
-
-void FAssetTypeActions_SkeletalMesh::ExecuteNewAnimBlueprint(TArray<TWeakObjectPtr<USkeletalMesh>> Objects)
-{
-	const FString DefaultSuffix = TEXT("_AnimBlueprint");
-
-	if ( Objects.Num() == 1 )
-	{
-		auto Object = Objects[0].Get();
-
-		if ( Object )
-		{
-			// Determine an appropriate name for inline-rename
-			FString Name;
-			FString PackageName;
-			CreateUniqueAssetName(Object->GetOutermost()->GetName(), DefaultSuffix, PackageName, Name);
-
-			UAnimBlueprintFactory* Factory = ConstructObject<UAnimBlueprintFactory>(UAnimBlueprintFactory::StaticClass());
-			Factory->TargetSkeleton = Object->Skeleton;
-
-			FContentBrowserModule& ContentBrowserModule = FModuleManager::LoadModuleChecked<FContentBrowserModule>("ContentBrowser");
-			ContentBrowserModule.Get().CreateNewAsset(Name, FPackageName::GetLongPackagePath(PackageName), UAnimBlueprint::StaticClass(), Factory);
-		}
-	}
-	else
-	{
-		TArray<UObject*> ObjectsToSync;
-		for (auto ObjIt = Objects.CreateConstIterator(); ObjIt; ++ObjIt)
-		{
-			auto Object = (*ObjIt).Get();
-			if ( Object )
-			{
-				// Determine an appropriate name
-				FString Name;
-				FString PackageName;
-				CreateUniqueAssetName(Object->GetOutermost()->GetName(), DefaultSuffix, PackageName, Name);
-
-				// Create the anim blueprint factory used to generate the asset
-				UAnimBlueprintFactory* Factory = ConstructObject<UAnimBlueprintFactory>(UAnimBlueprintFactory::StaticClass());
-				Factory->TargetSkeleton = Object->Skeleton;
-
-				FAssetToolsModule& AssetToolsModule = FModuleManager::GetModuleChecked<FAssetToolsModule>("AssetTools");
-				UObject* NewAsset = AssetToolsModule.Get().CreateAsset(Name, FPackageName::GetLongPackagePath(PackageName), UAnimBlueprint::StaticClass(), Factory);
-
-				if ( NewAsset )
-				{
-					ObjectsToSync.Add(NewAsset);
-				}
-			}
-		}
-
-		if ( ObjectsToSync.Num() > 0 )
-		{
-			FAssetTools::Get().SyncBrowserToAssets(ObjectsToSync);
-		}
-	}
-}
-
-template <typename TFactory, typename T>
-void FAssetTypeActions_SkeletalMesh::ExecuteNewAnimAsset(TArray<TWeakObjectPtr<USkeletalMesh>> Objects, const FString InSuffix)
-{
-	if ( Objects.Num() == 1 )
-	{
-		auto Object = Objects[0].Get();
-
-		if ( Object )
-		{
-			// Determine an appropriate name for inline-rename
-			FString Name;
-			FString PackageName;
-			CreateUniqueAssetName(Object->GetOutermost()->GetName(), InSuffix, PackageName, Name);
-
-			TFactory* Factory = ConstructObject<TFactory>(TFactory::StaticClass());
-			Factory->TargetSkeleton = Object->Skeleton;
-
-			FContentBrowserModule& ContentBrowserModule = FModuleManager::LoadModuleChecked<FContentBrowserModule>("ContentBrowser");
-			ContentBrowserModule.Get().CreateNewAsset(Name, FPackageName::GetLongPackagePath(PackageName), T::StaticClass(), Factory);
-		}
-	}
-	else
-	{
-		CreateAnimationAssets(Objects, T::StaticClass(), InSuffix);
-	}
-}
-
-void FAssetTypeActions_SkeletalMesh::CreateAnimationAssets(const TArray<TWeakObjectPtr<USkeletalMesh>>& Meshes, TSubclassOf<UAnimationAsset> AssetClass, const FString& InSuffix)
-{
-	TArray<UObject*> ObjectsToSync;
-	for (auto SkelIt = Meshes.CreateConstIterator(); SkelIt; ++SkelIt)
-	{
-		USkeletalMesh* Mesh = (*SkelIt).Get();
-		if ( Mesh )
-		{
-			// Determine an appropriate name
-			FString Name;
-			FString PackageName;
-			CreateUniqueAssetName(Mesh->GetOutermost()->GetName(), InSuffix, PackageName, Name);
-
-			// Create the asset, and assign its skeleton
-			FAssetToolsModule& AssetToolsModule = FModuleManager::GetModuleChecked<FAssetToolsModule>("AssetTools");
-			UAnimationAsset* NewAsset = Cast<UAnimationAsset>(AssetToolsModule.Get().CreateAsset(Name, FPackageName::GetLongPackagePath(PackageName), AssetClass, NULL));
-
-			if ( NewAsset )
-			{
-				NewAsset->SetSkeleton(Mesh->Skeleton);
-				NewAsset->MarkPackageDirty();
-
-				ObjectsToSync.Add(NewAsset);
-			}
-		}
-	}
-
-	if ( ObjectsToSync.Num() > 0 )
-	{
-		FAssetTools::Get().SyncBrowserToAssets(ObjectsToSync);
-	}
-}
-
-void FAssetTypeActions_SkeletalMesh::ExecuteReimport(TArray<TWeakObjectPtr<USkeletalMesh>> Objects)
+void FAssetTypeActions_SkeletalMesh::GetResolvedSourceFilePaths(const TArray<UObject*>& TypeAssets, TArray<FString>& OutSourceFilePaths) const
 {
-	for (auto ObjIt = Objects.CreateConstIterator(); ObjIt; ++ObjIt)
+	for (auto& Asset : TypeAssets)
 	{
-		auto Object = (*ObjIt).Get();
-		if ( Object )
+		const auto SkeletalMesh = CastChecked<USkeletalMesh>(Asset);
+		if (SkeletalMesh->AssetImportData)
 		{
-			FReimportManager::Instance()->Reimport(Object, /*bAskForNewFileIfMissing=*/true);
+			OutSourceFilePaths.Add(FReimportManager::ResolveImportFilename(SkeletalMesh->AssetImportData->SourceFilePath, SkeletalMesh));
 		}
 	}
 }
@@ -710,149 +573,6 @@ void FAssetTypeActions_SkeletalMesh::GetLODMenu(class FMenuBuilder& MenuBuilder,
 	}
 }
 
-void FAssetTypeActions_SkeletalMesh::GetCreateMenu(class FMenuBuilder& MenuBuilder, TArray<TWeakObjectPtr<USkeletalMesh>> Objects)
-{
-	check(Objects.Num() > 0);
-
-	MenuBuilder.AddMenuEntry(
-		LOCTEXT("SkeletalMesh_NewAnimBlueprint", "Create Anim Blueprint"),
-		LOCTEXT("SkeletalMesh_NewAnimBlueprintTooltip", "Creates an Anim Blueprint using the skeleton of the selected mesh."),
-		FSlateIcon(),
-		FUIAction(
-			FExecuteAction::CreateSP( this, &FAssetTypeActions_SkeletalMesh::ExecuteNewAnimBlueprint, Objects ),
-			FCanExecuteAction()
-			)
-		);
-
-	MenuBuilder.AddSubMenu( 
-		LOCTEXT( "SkeletalMesh_NewAimOffset", "Create AimOffset" ), 
-		LOCTEXT("SkeletalMesh_NewAimOffsetTooltip", "Creates an aimoffset blendspace using the selected skeleton."),
-		FNewMenuDelegate::CreateSP( this, &FAssetTypeActions_SkeletalMesh::FillAimOffsetBlendSpaceMenu, Objects ) );
-
-	MenuBuilder.AddSubMenu( 
-		LOCTEXT( "SkeletalMesh_NewBlendspace", "Create BlendSpace" ), 
-		LOCTEXT( "SkeletalMesh_NewBlendspaceTooltip", "Creates a blendspace using the skeleton of the selected mesh." ),
-		FNewMenuDelegate::CreateSP( this, &FAssetTypeActions_SkeletalMesh::FillBlendSpaceMenu, Objects ) );
-
-	MenuBuilder.AddMenuEntry(
-		LOCTEXT("SkeletalMesh_NewAnimComposite", "Create AnimComposite"),
-		LOCTEXT("SkeletalMesh_NewAnimCompositeTooltip", "Creates an AnimComposite using the selected mesh's skeleton."),
-		FSlateIcon(),
-		FUIAction(
-			FExecuteAction::CreateSP( this, &FAssetTypeActions_SkeletalMesh::ExecuteNewAnimAsset<UAnimCompositeFactory, UAnimComposite>, Objects, FString("_Composite") ),
-			FCanExecuteAction()
-			)
-		);
-
-	MenuBuilder.AddMenuEntry(
-		LOCTEXT("SkeletalMesh_NewAnimMontage", "Create AnimMontage"),
-		LOCTEXT("SkeletalMesh_NewAnimMontageTooltip", "Creates an AnimMontage using the selected mesh's skeleton."),
-		FSlateIcon(),
-		FUIAction(
-			FExecuteAction::CreateSP( this, &FAssetTypeActions_SkeletalMesh::ExecuteNewAnimAsset<UAnimMontageFactory, UAnimMontage>, Objects, FString("_Montage") ),
-			FCanExecuteAction()
-			)
-		);
-}
-
-void FAssetTypeActions_SkeletalMesh::FillBlendSpaceMenu( FMenuBuilder& MenuBuilder, TArray<TWeakObjectPtr<USkeletalMesh>> Objects )
-{
-	MenuBuilder.AddMenuEntry(
-		LOCTEXT("SkeletalMesh_New1DBlendspace", "Create 1D BlendSpace"),
-		LOCTEXT("SkeletalMesh_New1DBlendspaceTooltip", "Creates a 1D blendspace using the skeleton of the selected mesh."),
-		FSlateIcon(),
-		FUIAction(
-		FExecuteAction::CreateSP( this, &FAssetTypeActions_SkeletalMesh::ExecuteNewAnimAsset<UBlendSpaceFactory1D, UBlendSpace1D>, Objects, FString("_BlendSpace1D") ),
-		FCanExecuteAction()
-		)
-		);
-
-	MenuBuilder.AddMenuEntry(
-		LOCTEXT("SkeletalMesh_New2DBlendspace", "Create 2D BlendSpace"),
-		LOCTEXT("SkeletalMesh_New2DBlendspaceTooltip", "Creates a 2D blendspace using the skeleton of the selected mesh."),
-		FSlateIcon(),
-		FUIAction(
-		FExecuteAction::CreateSP( this, &FAssetTypeActions_SkeletalMesh::ExecuteNewAnimAsset<UBlendSpaceFactoryNew, UBlendSpace>, Objects, FString("_BlendSpace2D") ),
-		FCanExecuteAction()
-		)
-		);
-}
-
-void FAssetTypeActions_SkeletalMesh::FillAimOffsetBlendSpaceMenu( FMenuBuilder& MenuBuilder, TArray<TWeakObjectPtr<USkeletalMesh>> Objects )
-{
-	MenuBuilder.AddMenuEntry(
-		LOCTEXT("SkeletalMesh_New1DAimOffset", "Create 1D AimOffset"),
-		LOCTEXT("SkeletalMesh_New1DAimOffsetTooltip", "Creates a 1D aimoffset blendspace using the selected skeleton."),
-		FSlateIcon(),
-		FUIAction(
-		FExecuteAction::CreateSP( this, &FAssetTypeActions_SkeletalMesh::ExecuteNewAnimAsset<UAimOffsetBlendSpaceFactory1D, UAimOffsetBlendSpace1D>, Objects, FString("_AimOffset1D") ),
-		FCanExecuteAction()
-		)
-		);
-
-	MenuBuilder.AddMenuEntry(
-		LOCTEXT("SkeletalMesh_New2DAimOffset", "Create 2D AimOffset"),
-		LOCTEXT("SkeletalMesh_New2DAimOffsetTooltip", "Creates a 2D aimoffset blendspace using the selected skeleton."),
-		FSlateIcon(),
-		FUIAction(
-		FExecuteAction::CreateSP( this, &FAssetTypeActions_SkeletalMesh::ExecuteNewAnimAsset<UAimOffsetBlendSpaceFactoryNew, UAimOffsetBlendSpace>, Objects, FString("_AimOffset2D") ),
-		FCanExecuteAction()
-		)
-		);
-}
-
-void FAssetTypeActions_SkeletalMesh::ExecuteFindInExplorer(TArray<TWeakObjectPtr<USkeletalMesh>> Objects)
-{
-	for (auto ObjIt = Objects.CreateConstIterator(); ObjIt; ++ObjIt)
-	{
-		auto Object = (*ObjIt).Get();
-		if ( Object && Object->AssetImportData )
-		{
-			const FString SourceFilePath = FReimportManager::ResolveImportFilename(Object->AssetImportData->SourceFilePath, Object);
-			if ( SourceFilePath.Len() && IFileManager::Get().FileSize( *SourceFilePath ) != INDEX_NONE )
-			{
-				FPlatformProcess::ExploreFolder( *FPaths::GetPath(SourceFilePath) );
-			}
-		}
-	}
-}
-
-void FAssetTypeActions_SkeletalMesh::ExecuteOpenInExternalEditor(TArray<TWeakObjectPtr<USkeletalMesh>> Objects)
-{
-	for (auto ObjIt = Objects.CreateConstIterator(); ObjIt; ++ObjIt)
-	{
-		auto Object = (*ObjIt).Get();
-		if ( Object && Object->AssetImportData )
-		{
-			const FString SourceFilePath = FReimportManager::ResolveImportFilename(Object->AssetImportData->SourceFilePath, Object);
-			if ( SourceFilePath.Len() && IFileManager::Get().FileSize( *SourceFilePath ) != INDEX_NONE )
-			{
-				FPlatformProcess::LaunchFileInDefaultExternalApplication( *SourceFilePath, NULL, ELaunchVerb::Edit );
-			}
-		}
-	}
-}
-
-bool FAssetTypeActions_SkeletalMesh::CanExecuteSourceCommands(TArray<TWeakObjectPtr<USkeletalMesh>> Objects) const
-{
-	bool bHaveSourceAsset = false;
-	for (auto ObjIt = Objects.CreateConstIterator(); ObjIt; ++ObjIt)
-	{
-		auto Object = (*ObjIt).Get();
-		if ( Object && Object->AssetImportData )
-		{
-			const FString& SourceFilePath = FReimportManager::ResolveImportFilename(Object->AssetImportData->SourceFilePath, Object);
-
-			if ( SourceFilePath.Len() && IFileManager::Get().FileSize(*SourceFilePath) != INDEX_NONE )
-			{
-				return true;
-			}
-		}
-	}
-
-	return false;
-}
-
 void FAssetTypeActions_SkeletalMesh::ExecuteNewPhysicsAsset(TArray<TWeakObjectPtr<USkeletalMesh>> Objects)
 {
 	for (auto ObjIt = Objects.CreateConstIterator(); ObjIt; ++ObjIt)
@@ -955,34 +675,14 @@ void FAssetTypeActions_SkeletalMesh::ExecuteFindSkeleton(TArray<TWeakObjectPtr<U
 	}
 }
 
-void FAssetTypeActions_SkeletalMesh::FillSourceMenu(FMenuBuilder& MenuBuilder, const TArray<TWeakObjectPtr<USkeletalMesh>> Meshes) const
-{
-	MenuBuilder.AddMenuEntry(
-		LOCTEXT("SkeletalMesh_FindInExplorer", "Find Source"),
-		LOCTEXT("SkeletalMesh_FindInExplorerTooltip", "Opens explorer at the location of this asset."),
-		FSlateIcon(),
-		FUIAction(
-			FExecuteAction::CreateSP(this, &FAssetTypeActions_SkeletalMesh::ExecuteFindInExplorer, Meshes),
-			FCanExecuteAction::CreateSP(this, &FAssetTypeActions_SkeletalMesh::CanExecuteSourceCommands, Meshes)
-			)
-		);
 
-	MenuBuilder.AddMenuEntry(
-		LOCTEXT("SkeletalMesh_OpenInExternalEditor", "Open Source"),
-		LOCTEXT("SkeletalMesh_OpenInExternalEditorTooltip", "Opens the selected asset in an external editor."),
-		FSlateIcon(),
-		FUIAction(
-			FExecuteAction::CreateSP(this, &FAssetTypeActions_SkeletalMesh::ExecuteOpenInExternalEditor, Meshes),
-			FCanExecuteAction::CreateSP(this, &FAssetTypeActions_SkeletalMesh::CanExecuteSourceCommands, Meshes)
-			)
-		);
-}
 void FAssetTypeActions_SkeletalMesh::FillSkeletonMenu(FMenuBuilder& MenuBuilder, const TArray<TWeakObjectPtr<USkeletalMesh>> Meshes) const
 {
+	MenuBuilder.BeginSection("SkeletonMenu", LOCTEXT("SkeletonMenuHeading", "Skeleton"));
 	MenuBuilder.AddMenuEntry(
 		LOCTEXT("SkeletalMesh_NewSkeleton", "Create Skeleton"),
 		LOCTEXT("SkeletalMesh_NewSkeletonTooltip", "Creates a new skeleton for each of the selected meshes."),
-		FSlateIcon(),
+		FSlateIcon(), //@todo: icon
 		FUIAction(
 			FExecuteAction::CreateSP(this, &FAssetTypeActions_SkeletalMesh::ExecuteNewSkeleton, Meshes),
 			FCanExecuteAction()
@@ -992,7 +692,7 @@ void FAssetTypeActions_SkeletalMesh::FillSkeletonMenu(FMenuBuilder& MenuBuilder,
 	MenuBuilder.AddMenuEntry(
 		LOCTEXT("SkeletalMesh_AssignSkeleton", "Assign Skeleton"),
 		LOCTEXT("SkeletalMesh_AssignSkeletonTooltip", "Assigns a skeleton to the selected meshes."),
-		FSlateIcon(),
+		FSlateIcon(), //@todo: icon
 		FUIAction(
 			FExecuteAction::CreateSP(this, &FAssetTypeActions_SkeletalMesh::ExecuteAssignSkeleton, Meshes),
 			FCanExecuteAction()
@@ -1002,7 +702,7 @@ void FAssetTypeActions_SkeletalMesh::FillSkeletonMenu(FMenuBuilder& MenuBuilder,
 	MenuBuilder.AddMenuEntry(
 		LOCTEXT("SkeletalMesh_FindSkeleton", "Find Skeleton"),
 		LOCTEXT("SkeletalMesh_FindSkeletonTooltip", "Finds the skeleton used by the selected meshes in the content browser."),
-		FSlateIcon(),
+		FSlateIcon(), //@todo: icon
 		FUIAction(
 			FExecuteAction::CreateSP(this, &FAssetTypeActions_SkeletalMesh::ExecuteFindSkeleton, Meshes),
 			FCanExecuteAction()
@@ -1129,4 +829,13 @@ void FAssetTypeActions_SkeletalMesh::AssignSkeletonToMesh(USkeletalMesh* SkelMes
 	}
 }
 
+void FAssetTypeActions_SkeletalMesh::OnAssetCreated(TArray<UObject*> NewAssets) const
+{
+	if (NewAssets.Num() > 1)
+	{
+		FAssetTools::Get().SyncBrowserToAssets(NewAssets);
+	}
+
+}
+
 #undef LOCTEXT_NAMESPACE
diff --git a/Engine/Source/Developer/AssetTools/Private/AssetTypeActions/AssetTypeActions_SkeletalMesh.h b/Engine/Source/Developer/AssetTools/Private/AssetTypeActions/AssetTypeActions_SkeletalMesh.h
--- a/Engine/Source/Developer/AssetTools/Private/AssetTypeActions/AssetTypeActions_SkeletalMesh.h
+++ b/Engine/Source/Developer/AssetTools/Private/AssetTypeActions/AssetTypeActions_SkeletalMesh.h
@@ -14,52 +14,20 @@ public:
 	virtual void OpenAssetEditor( const TArray<UObject*>& InObjects, TSharedPtr<class IToolkitHost> EditWithinLevelEditor = TSharedPtr<IToolkitHost>() ) override;
 	virtual uint32 GetCategories() override { return EAssetTypeCategories::Basic | EAssetTypeCategories::Animation; }
 	virtual class UThumbnailInfo* GetThumbnailInfo(UObject* Asset) const override;
+	virtual bool IsImportedAsset() const override { return true; }
+	virtual void GetResolvedSourceFilePaths(const TArray<UObject*>& TypeAssets, TArray<FString>& OutSourceFilePaths) const override;
 
 protected:
 	/** Gets additional actions that do not apply to destructible meshes */
 	virtual void GetNonDestructibleActions( const TArray<TWeakObjectPtr<USkeletalMesh>>& Meshes, FMenuBuilder& MenuBuilder);
 
 private:
-	/** Handler for when Edit is selected */
-	void ExecuteEdit(TArray<TWeakObjectPtr<USkeletalMesh>> Objects);
-
-	/** Handler for when NewAnimBlueprint is selected */
-	void ExecuteNewAnimBlueprint(TArray<TWeakObjectPtr<USkeletalMesh>> Objects);
-
-	/** Creates a new object of type T from factory TFactory */
-	template <typename TFactory, typename T>
-	void ExecuteNewAnimAsset(TArray<TWeakObjectPtr<USkeletalMesh>> Objects, const FString InSuffix);
-
-	/** Creates animation assets using the BaseName+Suffix */
-	void CreateAnimationAssets(const TArray<TWeakObjectPtr<USkeletalMesh>>& Skeletons, TSubclassOf<UAnimationAsset> AssetClass, const FString& InPrefix);
-
-	/** Handler for when Reimport is selected */
-	void ExecuteReimport(TArray<TWeakObjectPtr<USkeletalMesh>> Objects);
-
 	/** Handler for when skeletal mesh LOD import is selected */
 	void LODImport(TArray<TWeakObjectPtr<USkeletalMesh>> Objects);
 
 	/** Handler for when skeletal mesh LOD sub menu is opened */
 	void GetLODMenu(class FMenuBuilder& MenuBuilder,TArray<TWeakObjectPtr<USkeletalMesh>> Objects);
 
-	/** Handler for when create asset sub menu is opened */
-	void GetCreateMenu(class FMenuBuilder& MenuBuilder,TArray<TWeakObjectPtr<USkeletalMesh>> Objects);
-
-	/** Handler for the blend space sub menu */
-	void FillBlendSpaceMenu( FMenuBuilder& MenuBuilder, TArray<TWeakObjectPtr<USkeletalMesh>> Objects );
-
-	/** Handler for the blend space sub menu */
-	void FillAimOffsetBlendSpaceMenu( FMenuBuilder& MenuBuilder, TArray<TWeakObjectPtr<USkeletalMesh>> Objects );
-
-	/** Handler for when FindInExplorer is selected */
-	void ExecuteFindInExplorer(TArray<TWeakObjectPtr<USkeletalMesh>> Objects);
-
-	/** Handler for when OpenInExternalEditor is selected */
-	void ExecuteOpenInExternalEditor(TArray<TWeakObjectPtr<USkeletalMesh>> Objects);
-
-	/** Returns true to allow execution of source file commands */
-	bool CanExecuteSourceCommands(TArray<TWeakObjectPtr<USkeletalMesh>> Objects) const;
-
 	/** Handler for when NewPhysicsAsset is selected */
 	void ExecuteNewPhysicsAsset(TArray<TWeakObjectPtr<USkeletalMesh>> Objects);
 	
@@ -80,6 +48,8 @@ private:
 	/** Assigns a skeleton to the mesh */
 	void AssignSkeletonToMesh(USkeletalMesh* SkelMesh) const;
 
+	void OnAssetCreated(TArray<UObject*> NewAssets) const;
+
 	void FillSourceMenu(FMenuBuilder& MenuBuilder, const TArray<TWeakObjectPtr<USkeletalMesh>> Meshes) const;
 	void FillSkeletonMenu(FMenuBuilder& MenuBuilder, const TArray<TWeakObjectPtr<USkeletalMesh>> Meshes) const;
 	void FillCreateMenu(FMenuBuilder& MenuBuilder, const TArray<TWeakObjectPtr<USkeletalMesh>> Meshes) const;
diff --git a/Engine/Source/Developer/AssetTools/Private/AssetTypeActions/AssetTypeActions_Skeleton.cpp b/Engine/Source/Developer/AssetTools/Private/AssetTypeActions/AssetTypeActions_Skeleton.cpp
--- a/Engine/Source/Developer/AssetTools/Private/AssetTypeActions/AssetTypeActions_Skeleton.cpp
+++ b/Engine/Source/Developer/AssetTools/Private/AssetTypeActions/AssetTypeActions_Skeleton.cpp
@@ -401,21 +401,14 @@ void FAssetTypeActions_Skeleton::GetActions( const TArray<UObject*>& InObjects,
 {
 	auto Skeletons = GetTypedWeakObjectPtrs<USkeleton>(InObjects);
 
-	MenuBuilder.AddMenuEntry(
-		LOCTEXT("Skeleton_Edit", "Edit"),
-		LOCTEXT("Skeleton_EditTooltip", "Opens the selected skeleton in the persona."),
-		FSlateIcon(),
-		FUIAction(
-			FExecuteAction::CreateSP( this, &FAssetTypeActions_Skeleton::ExecuteEdit, Skeletons ),
-			FCanExecuteAction()
-			)
-		);
-
 	// create menu
 	MenuBuilder.AddSubMenu(
 			LOCTEXT("CreateSkeletonSubmenu", "Create"),
 			LOCTEXT("CreateSkeletonSubmenu_ToolTip", "Create assets for this skeleton"),
-			FNewMenuDelegate::CreateSP(this, &FAssetTypeActions_Skeleton::FillCreateMenu, Skeletons));
+			FNewMenuDelegate::CreateSP(this, &FAssetTypeActions_Skeleton::FillCreateMenu, Skeletons),
+			false, 
+			FSlateIcon(FEditorStyle::GetStyleSetName(), "Persona.AssetActions.CreateAnimAsset")
+			);
 
 	// only show if one is selected. It won't work since I changed the window to be normal window
 	if (Skeletons.Num() == 1)
@@ -447,35 +440,25 @@ void FAssetTypeActions_Skeleton::GetActions( const TArray<UObject*>& InObjects,
 
 void FAssetTypeActions_Skeleton::FillCreateMenu(FMenuBuilder& MenuBuilder, TArray<TWeakObjectPtr<USkeleton>> Skeletons) const
 {
-	MenuBuilder.AddMenuEntry(
-		LOCTEXT("Skeleton_NewAnimBlueprint", "Create Anim Blueprint"),
-		LOCTEXT("Skeleton_NewAnimBlueprintTooltip", "Creates an Anim Blueprint using the selected skeleton."),
-		FSlateIcon(),
-		FUIAction(
-			FExecuteAction::CreateSP(this, &FAssetTypeActions_Skeleton::ExecuteNewAnimBlueprint, Skeletons),
-			FCanExecuteAction()
-			)
-		);
-
-	FAnimAssetCreated OnAssetCreated;
-	OnAssetCreated.BindSP(this, &FAssetTypeActions_Skeleton::OnAssetCreated);
-	AnimationEditorUtils::FillCreateAssetMenu(MenuBuilder, Skeletons, /*OnAssetCreated*/FAnimAssetCreated::CreateSP(this, &FAssetTypeActions_Skeleton::OnAssetCreated));
-
-	MenuBuilder.AddMenuSeparator();
-
 	// create rig
-	if(Skeletons.Num() == 1)
+	if (Skeletons.Num() == 1)
 	{
-		MenuBuilder.AddMenuEntry(
-			LOCTEXT("Skeleton_CreateRig", "Create Rig"),
-			LOCTEXT("Skeleton_CreateRigTooltip", "Create Rig from this skeleton."),
-			FSlateIcon(),
-			FUIAction(
+		MenuBuilder.BeginSection("CreateRig", LOCTEXT("CreateRigMenuHeading", "Rig"));
+		{
+			MenuBuilder.AddMenuEntry(
+				LOCTEXT("Skeleton_CreateRig", "Create Rig"),
+				LOCTEXT("Skeleton_CreateRigTooltip", "Create Rig from this skeleton."),
+				FSlateIcon(),
+				FUIAction(
 				FExecuteAction::CreateSP(this, &FAssetTypeActions_Skeleton::ExecuteCreateRig, Skeletons),
 				FCanExecuteAction()
 				)
-			);
+				);
+		}
+		MenuBuilder.EndSection();
 	}
+
+	AnimationEditorUtils::FillCreateAssetMenu(MenuBuilder, Skeletons, FAnimAssetCreated::CreateSP(this, &FAssetTypeActions_Skeleton::OnAssetCreated));	
 }
 
 void FAssetTypeActions_Skeleton::OpenAssetEditor( const TArray<UObject*>& InObjects, TSharedPtr<IToolkitHost> EditWithinLevelEditor )
@@ -493,74 +476,6 @@ void FAssetTypeActions_Skeleton::OpenAssetEditor( const TArray<UObject*>& InObje
 	}
 }
 
-void FAssetTypeActions_Skeleton::ExecuteEdit(TArray<TWeakObjectPtr<USkeleton>> Objects)
-{
-	for (auto ObjIt = Objects.CreateConstIterator(); ObjIt; ++ObjIt)
-	{
-		auto Object = (*ObjIt).Get();
-		if ( Object )
-		{
-			FAssetEditorManager::Get().OpenEditorForAsset(Object);
-		}
-	}
-}
-
-void FAssetTypeActions_Skeleton::ExecuteNewAnimBlueprint(TArray<TWeakObjectPtr<USkeleton>> Objects)
-{
-	const FString DefaultSuffix = TEXT("_AnimBlueprint");
-
-	if ( Objects.Num() == 1 )
-	{
-		auto Object = Objects[0].Get();
-
-		if ( Object )
-		{
-			// Determine an appropriate name for inline-rename
-			FString Name;
-			FString PackageName;
-			CreateUniqueAssetName(Object->GetOutermost()->GetName(), DefaultSuffix, PackageName, Name);
-
-			UAnimBlueprintFactory* Factory = ConstructObject<UAnimBlueprintFactory>(UAnimBlueprintFactory::StaticClass());
-			Factory->TargetSkeleton = Object;
-
-			FContentBrowserModule& ContentBrowserModule = FModuleManager::LoadModuleChecked<FContentBrowserModule>("ContentBrowser");
-			ContentBrowserModule.Get().CreateNewAsset(Name, FPackageName::GetLongPackagePath(PackageName), UAnimBlueprint::StaticClass(), Factory);
-		}
-	}
-	else
-	{
-		TArray<UObject*> ObjectsToSync;
-		for (auto ObjIt = Objects.CreateConstIterator(); ObjIt; ++ObjIt)
-		{
-			auto Object = (*ObjIt).Get();
-			if ( Object )
-			{
-				// Determine an appropriate name
-				FString Name;
-				FString PackageName;
-				CreateUniqueAssetName(Object->GetOutermost()->GetName(), DefaultSuffix, PackageName, Name);
-
-				// Create the anim blueprint factory used to generate the asset
-				UAnimBlueprintFactory* Factory = ConstructObject<UAnimBlueprintFactory>(UAnimBlueprintFactory::StaticClass());
-				Factory->TargetSkeleton = Object;
-
-				FAssetToolsModule& AssetToolsModule = FModuleManager::GetModuleChecked<FAssetToolsModule>("AssetTools");
-				UObject* NewAsset = AssetToolsModule.Get().CreateAsset(Name, FPackageName::GetLongPackagePath(PackageName), UAnimBlueprint::StaticClass(), Factory);
-
-				if ( NewAsset )
-				{
-					ObjectsToSync.Add(NewAsset);
-				}
-			}
-		}
-
-		if ( ObjectsToSync.Num() > 0 )
-		{
-			FAssetTools::Get().SyncBrowserToAssets(ObjectsToSync);
-		}
-	}
-}
-
 /** Handler for when Create Rig is selected */
 void FAssetTypeActions_Skeleton::ExecuteCreateRig(TArray<TWeakObjectPtr<USkeleton>> Skeletons)
 {
@@ -602,13 +517,6 @@ void FAssetTypeActions_Skeleton::CreateRig(const TWeakObjectPtr<USkeleton> Skele
 	}
 }
 
-void FAssetTypeActions_Skeleton::CreateAnimationAssets(const TArray<TWeakObjectPtr<USkeleton>>& Skeletons, TSubclassOf<UAnimationAsset> AssetClass, const FString& InPrefix)
-{
-	FAnimAssetCreated OnAssetCreated;
-	OnAssetCreated.BindSP(this, &FAssetTypeActions_Skeleton::OnAssetCreated);
-	AnimationEditorUtils::CreateAnimationAssets(Skeletons, AssetClass, InPrefix, /*OnAssetCreated*/FAnimAssetCreated::CreateSP(this, &FAssetTypeActions_Skeleton::OnAssetCreated));
-}
-
 void FAssetTypeActions_Skeleton::RetargetAnimationHandler(USkeleton* OldSkeleton, USkeleton* NewSkeleton, bool bRemapReferencedAssets, bool bConvertSpaces)
 {
 	// find all assets who references old skeleton
diff --git a/Engine/Source/Developer/AssetTools/Private/AssetTypeActions/AssetTypeActions_Skeleton.h b/Engine/Source/Developer/AssetTools/Private/AssetTypeActions/AssetTypeActions_Skeleton.h
--- a/Engine/Source/Developer/AssetTools/Private/AssetTypeActions/AssetTypeActions_Skeleton.h
+++ b/Engine/Source/Developer/AssetTools/Private/AssetTypeActions/AssetTypeActions_Skeleton.h
@@ -38,12 +38,6 @@ public:
 	virtual uint32 GetCategories() override { return EAssetTypeCategories::Animation; }
 
 private:
-	/** Handler for when Edit is selected */
-	void ExecuteEdit(TArray<TWeakObjectPtr<USkeleton>> Objects);
-
-	/** Handler for when NewAnimBlueprint is selected */
-	void ExecuteNewAnimBlueprint(TArray<TWeakObjectPtr<USkeleton>> Objects);
-
 	/** Handler for when Skeleton Retarget is selected */
 	void ExecuteRetargetSkeleton(TArray<TWeakObjectPtr<USkeleton>> Skeletons);
 
@@ -51,9 +45,6 @@ private:
 	void ExecuteCreateRig(TArray<TWeakObjectPtr<USkeleton>> Skeletons);
 
 private: // Helper functions
-	/** Creates animation assets using the BaseName+Suffix */
-	void CreateAnimationAssets(const TArray<TWeakObjectPtr<USkeleton>>& Skeletons, TSubclassOf<UAnimationAsset> AssetClass, const FString& InPrefix);
-
 	/** Creates animation assets using the BaseName+Suffix */
 	void CreateRig(const TWeakObjectPtr<USkeleton> Skeleton);
 
diff --git a/Engine/Source/Developer/AssetTools/Private/AssetTypeActions/AssetTypeActions_SlateWidgetStyle.cpp b/Engine/Source/Developer/AssetTools/Private/AssetTypeActions/AssetTypeActions_SlateWidgetStyle.cpp
--- a/Engine/Source/Developer/AssetTools/Private/AssetTypeActions/AssetTypeActions_SlateWidgetStyle.cpp
+++ b/Engine/Source/Developer/AssetTools/Private/AssetTypeActions/AssetTypeActions_SlateWidgetStyle.cpp
@@ -10,16 +10,6 @@
 void FAssetTypeActions_SlateWidgetStyle::GetActions( const TArray<UObject*>& InObjects, FMenuBuilder& MenuBuilder )
 {
 	auto Styles = GetTypedWeakObjectPtrs< USlateWidgetStyleAsset >(InObjects);
-
-	MenuBuilder.AddMenuEntry(
-		LOCTEXT("SlateWidgetStyle_Edit", "Edit"),
-		LOCTEXT("SlateWidgetStyle_EditTooltip", "Opens the selected styles in the style editor."),
-		FSlateIcon(),
-		FUIAction(
-		FExecuteAction::CreateSP( this, &FAssetTypeActions_SlateWidgetStyle::ExecuteEdit, Styles ),
-		FCanExecuteAction()
-		)
-	);
 }
 
 void FAssetTypeActions_SlateWidgetStyle::OpenAssetEditor( const TArray<UObject*>& Objects, TSharedPtr<IToolkitHost> EditWithinLevelEditor )
@@ -44,15 +34,4 @@ void FAssetTypeActions_SlateWidgetStyle::OpenAssetEditor( const TArray<UObject*>
 	FSimpleAssetEditor::CreateEditor(EToolkitMode::Standalone, EditWithinLevelEditor, Objects, FSimpleAssetEditor::FGetDetailsViewObjects::CreateStatic(&Local::GetSubObjects));
 }
 
-void FAssetTypeActions_SlateWidgetStyle::ExecuteEdit(TArray<TWeakObjectPtr< USlateWidgetStyleAsset >> Styles)
-{
-	TArray<UObject*> Objects;
-	for( auto StyleIter = Styles.CreateConstIterator(); StyleIter; ++StyleIter )
-	{
-		Objects.Add( StyleIter->Get() );
-	}
-
-	OpenAssetEditor( Objects );
-}
-
 #undef LOCTEXT_NAMESPACE
\ No newline at end of file
diff --git a/Engine/Source/Developer/AssetTools/Private/AssetTypeActions/AssetTypeActions_SlateWidgetStyle.h b/Engine/Source/Developer/AssetTools/Private/AssetTypeActions/AssetTypeActions_SlateWidgetStyle.h
--- a/Engine/Source/Developer/AssetTools/Private/AssetTypeActions/AssetTypeActions_SlateWidgetStyle.h
+++ b/Engine/Source/Developer/AssetTools/Private/AssetTypeActions/AssetTypeActions_SlateWidgetStyle.h
@@ -14,8 +14,4 @@ public:
 	virtual void GetActions( const TArray<UObject*>& InObjects, FMenuBuilder& MenuBuilder ) override;
 
 	virtual void OpenAssetEditor( const TArray<UObject*>& InObjects, TSharedPtr<class IToolkitHost> EditWithinLevelEditor = TSharedPtr<IToolkitHost>() ) override;
-
-private:
-	/** Handler for when Edit is selected */
-	void ExecuteEdit(TArray<TWeakObjectPtr< USlateWidgetStyleAsset >> Styles);
 };
\ No newline at end of file
diff --git a/Engine/Source/Developer/AssetTools/Private/AssetTypeActions/AssetTypeActions_SoundBase.cpp b/Engine/Source/Developer/AssetTools/Private/AssetTypeActions/AssetTypeActions_SoundBase.cpp
--- a/Engine/Source/Developer/AssetTools/Private/AssetTypeActions/AssetTypeActions_SoundBase.cpp
+++ b/Engine/Source/Developer/AssetTools/Private/AssetTypeActions/AssetTypeActions_SoundBase.cpp
@@ -19,7 +19,7 @@ void FAssetTypeActions_SoundBase::GetActions( const TArray<UObject*>& InObjects,
 	MenuBuilder.AddMenuEntry(
 		LOCTEXT("Sound_PlaySound", "Play"),
 		LOCTEXT("Sound_PlaySoundTooltip", "Plays the selected sound."),
-		FSlateIcon(),
+		FSlateIcon(FEditorStyle::GetStyleSetName(), "MediaAsset.AssetActions.Play"),
 		FUIAction(
 			FExecuteAction::CreateSP( this, &FAssetTypeActions_SoundBase::ExecutePlaySound, Sounds ),
 			FCanExecuteAction::CreateSP( this, &FAssetTypeActions_SoundBase::CanExecutePlayCommand, Sounds )
@@ -29,7 +29,7 @@ void FAssetTypeActions_SoundBase::GetActions( const TArray<UObject*>& InObjects,
 	MenuBuilder.AddMenuEntry(
 		LOCTEXT("Sound_StopSound", "Stop"),
 		LOCTEXT("Sound_StopSoundTooltip", "Stops the selected sounds."),
-		FSlateIcon(),
+		FSlateIcon(FEditorStyle::GetStyleSetName(), "MediaAsset.AssetActions.Stop"),
 		FUIAction(
 			FExecuteAction::CreateSP( this, &FAssetTypeActions_SoundBase::ExecuteStopSound, Sounds ),
 			FCanExecuteAction()
diff --git a/Engine/Source/Developer/AssetTools/Private/AssetTypeActions/AssetTypeActions_SoundClass.cpp b/Engine/Source/Developer/AssetTools/Private/AssetTypeActions/AssetTypeActions_SoundClass.cpp
--- a/Engine/Source/Developer/AssetTools/Private/AssetTypeActions/AssetTypeActions_SoundClass.cpp
+++ b/Engine/Source/Developer/AssetTools/Private/AssetTypeActions/AssetTypeActions_SoundClass.cpp
@@ -11,21 +11,6 @@ UClass* FAssetTypeActions_SoundClass::GetSupportedClass() const
 	return USoundClass::StaticClass();
 }
 
-void FAssetTypeActions_SoundClass::GetActions( const TArray<UObject*>& InObjects, FMenuBuilder& MenuBuilder )
-{
-	auto Classes = GetTypedWeakObjectPtrs<USoundClass>(InObjects);
-
-	MenuBuilder.AddMenuEntry(
-		LOCTEXT("SoundClass_Edit", "Edit"),
-		LOCTEXT("SoundClass_EditTooltip", "Opens the selected sound classes in the sound class editor."),
-		FSlateIcon(),
-		FUIAction(
-			FExecuteAction::CreateSP( this, &FAssetTypeActions_SoundClass::ExecuteEdit, Classes ),
-			FCanExecuteAction()
-			)
-		);
-}
-
 void FAssetTypeActions_SoundClass::OpenAssetEditor( const TArray<UObject*>& InObjects, TSharedPtr<IToolkitHost> EditWithinLevelEditor )
 {
 	EToolkitMode::Type Mode = EditWithinLevelEditor.IsValid() ? EToolkitMode::WorldCentric : EToolkitMode::Standalone;
@@ -41,16 +26,4 @@ void FAssetTypeActions_SoundClass::OpenAssetEditor( const TArray<UObject*>& InOb
 	}
 }
 
-void FAssetTypeActions_SoundClass::ExecuteEdit(TArray<TWeakObjectPtr<USoundClass>> Objects)
-{
-	for (auto ObjIt = Objects.CreateConstIterator(); ObjIt; ++ObjIt)
-	{
-		auto Object = (*ObjIt).Get();
-		if ( Object )
-		{
-			FAssetEditorManager::Get().OpenEditorForAsset(Object);
-		}
-	}
-}
-
 #undef LOCTEXT_NAMESPACE
\ No newline at end of file
diff --git a/Engine/Source/Developer/AssetTools/Private/AssetTypeActions/AssetTypeActions_SoundClass.h b/Engine/Source/Developer/AssetTools/Private/AssetTypeActions/AssetTypeActions_SoundClass.h
--- a/Engine/Source/Developer/AssetTools/Private/AssetTypeActions/AssetTypeActions_SoundClass.h
+++ b/Engine/Source/Developer/AssetTools/Private/AssetTypeActions/AssetTypeActions_SoundClass.h
@@ -9,12 +9,6 @@ public:
 	virtual FText GetName() const override { return NSLOCTEXT("AssetTypeActions", "AssetTypeActions_SoundClass", "Sound Class"); }
 	virtual FColor GetTypeColor() const override { return FColor(255, 175, 0); }
 	virtual UClass* GetSupportedClass() const override;
-	virtual bool HasActions ( const TArray<UObject*>& InObjects ) const override { return true; }
-	virtual void GetActions( const TArray<UObject*>& InObjects, FMenuBuilder& MenuBuilder ) override;
 	virtual void OpenAssetEditor( const TArray<UObject*>& InObjects, TSharedPtr<class IToolkitHost> EditWithinLevelEditor = TSharedPtr<IToolkitHost>() ) override;
 	virtual uint32 GetCategories() override { return EAssetTypeCategories::Sounds; }
-
-private:
-	/** Handler for when Edit is selected */
-	void ExecuteEdit(TArray<TWeakObjectPtr<class USoundClass>> Objects);
 };
\ No newline at end of file
diff --git a/Engine/Source/Developer/AssetTools/Private/AssetTypeActions/AssetTypeActions_SoundCue.cpp b/Engine/Source/Developer/AssetTools/Private/AssetTypeActions/AssetTypeActions_SoundCue.cpp
--- a/Engine/Source/Developer/AssetTools/Private/AssetTypeActions/AssetTypeActions_SoundCue.cpp
+++ b/Engine/Source/Developer/AssetTools/Private/AssetTypeActions/AssetTypeActions_SoundCue.cpp
@@ -16,17 +16,6 @@ UClass* FAssetTypeActions_SoundCue::GetSupportedClass() const
 void FAssetTypeActions_SoundCue::GetActions( const TArray<UObject*>& InObjects, FMenuBuilder& MenuBuilder )
 {
 	auto SoundCues = GetTypedWeakObjectPtrs<USoundCue>(InObjects);
-
-	MenuBuilder.AddMenuEntry(
-		LOCTEXT("SoundCue_Edit", "Edit"),
-		LOCTEXT("SoundCue_EditTooltip", "Opens the selected sound cues in the sound cue editor."),
-		FSlateIcon(),
-		FUIAction(
-			FExecuteAction::CreateSP( this, &FAssetTypeActions_SoundCue::ExecuteEdit, SoundCues ),
-			FCanExecuteAction()
-			)
-		);
-
 	FAssetTypeActions_SoundBase::GetActions(InObjects, MenuBuilder);
 
 	MenuBuilder.AddMenuEntry(
@@ -55,18 +44,6 @@ void FAssetTypeActions_SoundCue::OpenAssetEditor( const TArray<UObject*>& InObje
 	}
 }
 
-void FAssetTypeActions_SoundCue::ExecuteEdit(TArray<TWeakObjectPtr<USoundCue>> Objects)
-{
-	for (auto ObjIt = Objects.CreateConstIterator(); ObjIt; ++ObjIt)
-	{
-		auto Object = (*ObjIt).Get();
-		if ( Object )
-		{
-			FAssetEditorManager::Get().OpenEditorForAsset(Object);
-		}
-	}
-}
-
 bool FAssetTypeActions_SoundCue::CanExecuteConsolidateCommand(TArray<TWeakObjectPtr<USoundCue>> Objects) const
 {
 	return Objects.Num() > 1;
diff --git a/Engine/Source/Developer/AssetTools/Private/AssetTypeActions/AssetTypeActions_SoundCue.h b/Engine/Source/Developer/AssetTools/Private/AssetTypeActions/AssetTypeActions_SoundCue.h
--- a/Engine/Source/Developer/AssetTools/Private/AssetTypeActions/AssetTypeActions_SoundCue.h
+++ b/Engine/Source/Developer/AssetTools/Private/AssetTypeActions/AssetTypeActions_SoundCue.h
@@ -14,13 +14,9 @@ public:
 	virtual bool CanFilter() override { return true; }
 
 private:
-	/** Handler for when Edit is selected */
-	void ExecuteEdit(TArray<TWeakObjectPtr<USoundCue>> Objects);
-
 	/** Take selected SoundCues and combine, as much as possible, them to using shared attenuation settings */
 	void ExecuteConsolidateAttenuation(TArray<TWeakObjectPtr<USoundCue>> Objects);
 
 	/** Returns true if more than one cue is selected to consolidate */
 	bool CanExecuteConsolidateCommand(TArray<TWeakObjectPtr<USoundCue>> Objects) const;
-
 };
\ No newline at end of file
diff --git a/Engine/Source/Developer/AssetTools/Private/AssetTypeActions/AssetTypeActions_SoundWave.cpp b/Engine/Source/Developer/AssetTools/Private/AssetTypeActions/AssetTypeActions_SoundWave.cpp
--- a/Engine/Source/Developer/AssetTools/Private/AssetTypeActions/AssetTypeActions_SoundWave.cpp
+++ b/Engine/Source/Developer/AssetTools/Private/AssetTypeActions/AssetTypeActions_SoundWave.cpp
@@ -19,55 +19,24 @@ void FAssetTypeActions_SoundWave::GetActions( const TArray<UObject*>& InObjects,
 
 	auto SoundNodes = GetTypedWeakObjectPtrs<USoundWave>(InObjects);
 
-	/*MenuBuilder.AddMenuEntry(
-		LOCTEXT("SoundWave_CompressionPreview", "Compression Preview"),
-		LOCTEXT("SoundWave_CompressionPreviewTooltip", "Opens the selected sound node waves in the compression preview window."),
-		FSlateIcon(),
-		FUIAction(
-			FExecuteAction::CreateSP( this, &FAssetTypeActions_SoundWave::ExecuteCompressionPreview, SoundNodes ),
-			FCanExecuteAction()
-			)
-		);*/
-
-	MenuBuilder.AddMenuEntry(
-		LOCTEXT("SoundWave_Reimport", "Reimport"),
-		LOCTEXT("SoundWave_ReimportTooltip", "Reimports the selected waves from file."),
-		FSlateIcon(),
-		FUIAction(
-			FExecuteAction::CreateSP( this, &FAssetTypeActions_SoundWave::ExecuteReimport, SoundNodes ),
-			FCanExecuteAction()
-			)
-		);
-
-	MenuBuilder.AddMenuEntry(
-		LOCTEXT("SoundWave_FindInExplorer", "Find Source"),
-		LOCTEXT("SoundWave_FindInExplorerTooltip", "Opens explorer at the location of this asset."),
-		FSlateIcon(),
-		FUIAction(
-			FExecuteAction::CreateSP( this, &FAssetTypeActions_SoundWave::ExecuteFindInExplorer, SoundNodes ),
-			FCanExecuteAction::CreateSP( this, &FAssetTypeActions_SoundWave::CanExecuteSourceCommands, SoundNodes )
-			)
-		);
-
-	MenuBuilder.AddMenuEntry(
-		LOCTEXT("SoundWave_OpenInExternalEditor", "Open Source"),
-		LOCTEXT("SoundWave_OpenInExternalEditorTooltip", "Opens the selected asset in an external editor."),
-		FSlateIcon(),
-		FUIAction(
-			FExecuteAction::CreateSP( this, &FAssetTypeActions_SoundWave::ExecuteOpenInExternalEditor, SoundNodes ),
-			FCanExecuteAction::CreateSP( this, &FAssetTypeActions_SoundWave::CanExecuteSourceCommands, SoundNodes )
-			)
-		);
-
 	MenuBuilder.AddMenuEntry(
 		LOCTEXT("SoundWave_CreateCue", "Create Cue"),
-		LOCTEXT("SoundWave_CreateCueTooltip", "Creates a sound cue ."),
-		FSlateIcon(),
+		LOCTEXT("SoundWave_CreateCueTooltip", "Creates a sound cue using this sound wave."),
+		FSlateIcon(FEditorStyle::GetStyleSetName(), "ClassIcon.SoundCue"),
 		FUIAction(
 		FExecuteAction::CreateSP( this, &FAssetTypeActions_SoundWave::ExecuteCreateSoundCue, SoundNodes ),
 		FCanExecuteAction()
-		)
-	);
+			)
+		);
+}
+
+void FAssetTypeActions_SoundWave::GetResolvedSourceFilePaths(const TArray<UObject*>& TypeAssets, TArray<FString>& OutSourceFilePaths) const
+{
+	for (auto& Asset : TypeAssets)
+	{
+		const auto SoundWave = CastChecked<USoundWave>(Asset);
+		OutSourceFilePaths.Add(FReimportManager::ResolveImportFilename(SoundWave->SourceFilePath, SoundWave));
+	}
 }
 
 void FAssetTypeActions_SoundWave::OpenAssetEditor( const TArray<UObject*>& InObjects, TSharedPtr<IToolkitHost> EditWithinLevelEditor )
@@ -125,70 +94,6 @@ void FAssetTypeActions_SoundWave::ExecuteCompressionPreview(TArray<TWeakObjectPt
 	OpenAssetEditor(ObjectsToEdit);*/
 }
 
-void FAssetTypeActions_SoundWave::ExecuteReimport(TArray<TWeakObjectPtr<USoundWave>> Objects)
-{
-	for (auto ObjIt = Objects.CreateConstIterator(); ObjIt; ++ObjIt)
-	{
-		USoundWave* SoundWave = (*ObjIt).Get();
-		if ( SoundWave )
-		{
-			FReimportManager::Instance()->Reimport(SoundWave, /*bAskForNewFileIfMissing=*/true);
-		}
-	}
-}
-
-void FAssetTypeActions_SoundWave::ExecuteFindInExplorer(TArray<TWeakObjectPtr<USoundWave>> Objects)
-{
-	for (auto ObjIt = Objects.CreateConstIterator(); ObjIt; ++ObjIt)
-	{
-		auto Object = (*ObjIt).Get();
-		if ( Object )
-		{
-			const FString SourceFilePath = FReimportManager::ResolveImportFilename(Object->SourceFilePath, Object);
-			if ( SourceFilePath.Len() && IFileManager::Get().FileSize( *SourceFilePath ) != INDEX_NONE )
-			{
-				FPlatformProcess::ExploreFolder( *FPaths::GetPath(SourceFilePath) );
-			}
-		}
-	}
-}
-
-void FAssetTypeActions_SoundWave::ExecuteOpenInExternalEditor(TArray<TWeakObjectPtr<USoundWave>> Objects)
-{
-	for (auto ObjIt = Objects.CreateConstIterator(); ObjIt; ++ObjIt)
-	{
-		auto Object = (*ObjIt).Get();
-		if ( Object )
-		{
-			const FString SourceFilePath = FReimportManager::ResolveImportFilename(Object->SourceFilePath, Object);
-			if ( SourceFilePath.Len() && IFileManager::Get().FileSize( *SourceFilePath ) != INDEX_NONE )
-			{
-				FPlatformProcess::LaunchFileInDefaultExternalApplication( *SourceFilePath, NULL, ELaunchVerb::Edit );
-			}
-		}
-	}
-}
-
-bool FAssetTypeActions_SoundWave::CanExecuteSourceCommands(TArray<TWeakObjectPtr<USoundWave>> Objects) const
-{
-	bool bHaveSourceAsset = false;
-	for (auto ObjIt = Objects.CreateConstIterator(); ObjIt; ++ObjIt)
-	{
-		auto Object = (*ObjIt).Get();
-		if ( Object )
-		{
-			const FString& SourceFilePath = FReimportManager::ResolveImportFilename(Object->SourceFilePath, Object);
-
-			if ( SourceFilePath.Len() && IFileManager::Get().FileSize(*SourceFilePath) != INDEX_NONE )
-			{
-				return true;
-			}
-		}
-	}
-
-	return false;
-}
-
 void FAssetTypeActions_SoundWave::ExecuteCreateSoundCue(TArray<TWeakObjectPtr<USoundWave>> Objects)
 {
 	const FString DefaultSuffix = TEXT("_Cue");
diff --git a/Engine/Source/Developer/AssetTools/Private/AssetTypeActions/AssetTypeActions_SoundWave.h b/Engine/Source/Developer/AssetTools/Private/AssetTypeActions/AssetTypeActions_SoundWave.h
--- a/Engine/Source/Developer/AssetTools/Private/AssetTypeActions/AssetTypeActions_SoundWave.h
+++ b/Engine/Source/Developer/AssetTools/Private/AssetTypeActions/AssetTypeActions_SoundWave.h
@@ -12,24 +12,14 @@ public:
 	virtual void GetActions( const TArray<UObject*>& InObjects, FMenuBuilder& MenuBuilder ) override;
 	virtual void OpenAssetEditor( const TArray<UObject*>& InObjects, TSharedPtr<class IToolkitHost> EditWithinLevelEditor = TSharedPtr<IToolkitHost>() ) override;
 	virtual bool CanFilter() override { return true; }
-
+	virtual bool IsImportedAsset() const override { return true; }
+	virtual void GetResolvedSourceFilePaths(const TArray<UObject*>& TypeAssets, TArray<FString>& OutSourceFilePaths) const override;
+	
 private:
 
 	/** Handler for when CompressionPreview is selected */
 	void ExecuteCompressionPreview(TArray<TWeakObjectPtr<USoundWave>> Objects);
 
-	/** Handler for when Reimport is selected */
-	void ExecuteReimport(TArray<TWeakObjectPtr<USoundWave>> Objects);
-
-	/** Handler for when FindInExplorer is selected */
-	void ExecuteFindInExplorer(TArray<TWeakObjectPtr<USoundWave>> Objects);
-
-	/** Handler for when OpenInExternalEditor is selected */
-	void ExecuteOpenInExternalEditor(TArray<TWeakObjectPtr<USoundWave>> Objects);
-
 	/** Creates a SoundCue of the same name for the sound, if one does not already exist */
 	void ExecuteCreateSoundCue(TArray<TWeakObjectPtr<USoundWave>> Objects);
-
-	/** Returns true to allow execution of source file commands */
-	bool CanExecuteSourceCommands(TArray<TWeakObjectPtr<USoundWave>> Objects) const;
 };
\ No newline at end of file
diff --git a/Engine/Source/Developer/AssetTools/Private/AssetTypeActions/AssetTypeActions_StaticMesh.cpp b/Engine/Source/Developer/AssetTools/Private/AssetTypeActions/AssetTypeActions_StaticMesh.cpp
--- a/Engine/Source/Developer/AssetTools/Private/AssetTypeActions/AssetTypeActions_StaticMesh.cpp
+++ b/Engine/Source/Developer/AssetTools/Private/AssetTypeActions/AssetTypeActions_StaticMesh.cpp
@@ -19,60 +19,20 @@ void FAssetTypeActions_StaticMesh::GetActions( const TArray<UObject*>& InObjects
 	auto Meshes = GetTypedWeakObjectPtrs<UStaticMesh>(InObjects);
 
 	MenuBuilder.AddMenuEntry(
-		NSLOCTEXT("AssetTypeActions_StaticMesh", "StaticMesh_Edit", "Edit"),
-		NSLOCTEXT("AssetTypeActions_StaticMesh", "StaticMesh_EditTooltip", "Opens the selected meshes in the static mesh editor."),
-		FSlateIcon(),
-		FUIAction(
-			FExecuteAction::CreateSP( this, &FAssetTypeActions_StaticMesh::ExecuteEdit, Meshes ),
-			FCanExecuteAction()
-			)
-		);
-
-	MenuBuilder.AddMenuEntry(
-		NSLOCTEXT("AssetTypeActions_StaticMesh", "StaticMesh_Reimport", "Reimport"),
-		NSLOCTEXT("AssetTypeActions_StaticMesh", "StaticMesh_ReimportTooltip", "Reimports the selected meshes from file."),
-		FSlateIcon(),
+		NSLOCTEXT("AssetTypeActions_StaticMesh", "ObjectContext_CreateDestructibleMesh", "Create Destructible Mesh"),
+		NSLOCTEXT("AssetTypeActions_StaticMesh", "ObjectContext_CreateDestructibleMeshTooltip", "Creates a DestructibleMesh from the StaticMesh and opens it in the DestructibleMesh editor."),
+		FSlateIcon(FEditorStyle::GetStyleSetName(), "ClassIcon.DestructibleComponent"),
 		FUIAction(
-			FExecuteAction::CreateSP( this, &FAssetTypeActions_StaticMesh::ExecuteReimport, Meshes ),
+			FExecuteAction::CreateSP( this, &FAssetTypeActions_StaticMesh::ExecuteCreateDestructibleMesh, Meshes ),
 			FCanExecuteAction()
 			)
 		);
 
 	MenuBuilder.AddSubMenu(
-		NSLOCTEXT("AssetTypeActions_StaticMesh", "StaticMesh_ImportLOD", "ImportLOD"),
+		NSLOCTEXT("AssetTypeActions_StaticMesh", "StaticMesh_ImportLOD", "Import LOD"),
 		NSLOCTEXT("AssetTypeActions_StaticMesh", "StaticMesh_ImportLODtooltip", "Imports meshes into the LODs"),
 		FNewMenuDelegate::CreateSP( this, &FAssetTypeActions_StaticMesh::GetImportLODMenu, Meshes )
 	);
-
-	MenuBuilder.AddMenuEntry(
-		NSLOCTEXT("AssetTypeActions_StaticMesh", "StaticMesh_FindInExplorer", "Find Source"),
-		NSLOCTEXT("AssetTypeActions_StaticMesh", "StaticMesh_FindInExplorerTooltip", "Opens explorer at the location of this asset."),
-		FSlateIcon(),
-		FUIAction(
-			FExecuteAction::CreateSP( this, &FAssetTypeActions_StaticMesh::ExecuteFindInExplorer, Meshes ),
-			FCanExecuteAction::CreateSP( this, &FAssetTypeActions_StaticMesh::CanExecuteSourceCommands, Meshes )
-			)
-		);
-
-	MenuBuilder.AddMenuEntry(
-		NSLOCTEXT("AssetTypeActions_StaticMesh", "StaticMesh_OpenInExternalEditor", "Open Source"),
-		NSLOCTEXT("AssetTypeActions_StaticMesh", "StaticMesh_OpenInExternalEditorTooltip", "Opens the selected asset in an external editor."),
-		FSlateIcon(),
-		FUIAction(
-			FExecuteAction::CreateSP( this, &FAssetTypeActions_StaticMesh::ExecuteOpenInExternalEditor, Meshes ),
-			FCanExecuteAction::CreateSP( this, &FAssetTypeActions_StaticMesh::CanExecuteSourceCommands, Meshes )
-			)
-		);
-
-	MenuBuilder.AddMenuEntry(
-		NSLOCTEXT("AssetTypeActions_StaticMesh", "ObjectContext_CreateDestructibleMesh", "Create Destructible Mesh"),
-		NSLOCTEXT("AssetTypeActions_StaticMesh", "ObjectContext_CreateDestructibleMeshTooltip", "Creates a DestructibleMesh from the StaticMesh and opens it in the DestructibleMesh editor."),
-		FSlateIcon(),
-		FUIAction(
-			FExecuteAction::CreateSP( this, &FAssetTypeActions_StaticMesh::ExecuteCreateDestructibleMesh, Meshes ),
-			FCanExecuteAction()
-			)
-		);
 }
 
 void FAssetTypeActions_StaticMesh::OpenAssetEditor( const TArray<UObject*>& InObjects, TSharedPtr<IToolkitHost> EditWithinLevelEditor )
@@ -103,26 +63,14 @@ UThumbnailInfo* FAssetTypeActions_StaticMesh::GetThumbnailInfo(UObject* Asset) c
 	return ThumbnailInfo;
 }
 
-void FAssetTypeActions_StaticMesh::ExecuteEdit(TArray<TWeakObjectPtr<UStaticMesh>> Objects)
-{
-	for (auto ObjIt = Objects.CreateConstIterator(); ObjIt; ++ObjIt)
-	{
-		auto Object = (*ObjIt).Get();
-		if ( Object )
-		{
-			FAssetEditorManager::Get().OpenEditorForAsset(Object);
-		}
-	}
-}
-
-void FAssetTypeActions_StaticMesh::ExecuteReimport(TArray<TWeakObjectPtr<UStaticMesh>> Objects)
+void FAssetTypeActions_StaticMesh::GetResolvedSourceFilePaths(const TArray<UObject*>& TypeAssets, TArray<FString>& OutSourceFilePaths) const
 {
-	for (auto ObjIt = Objects.CreateConstIterator(); ObjIt; ++ObjIt)
+	for (auto& Asset : TypeAssets)
 	{
-		auto Object = (*ObjIt).Get();
-		if ( Object )
+		const auto StaticMesh = CastChecked<UStaticMesh>(Asset);
+		if (StaticMesh->AssetImportData)
 		{
-			FReimportManager::Instance()->Reimport(Object, /*bAskForNewFileIfMissing=*/true);
+			OutSourceFilePaths.Add(FReimportManager::ResolveImportFilename(StaticMesh->AssetImportData->SourceFilePath, StaticMesh));
 		}
 	}
 }
@@ -140,7 +88,7 @@ void FAssetTypeActions_StaticMesh::GetImportLODMenu(class FMenuBuilder& MenuBuil
 		FText ToolTip = NSLOCTEXT("AssetTypeActions_StaticMesh", "ReimportTip", "Reimport over existing LOD");
 		if(LOD == First->GetNumLODs())
 		{
-			Description = FText::Format( NSLOCTEXT("AssetTypeActions_StaticMesh", "Import LOD (number)", "Import LOD {0}"), LODText );
+			Description = FText::Format( NSLOCTEXT("AssetTypeActions_StaticMesh", "LOD (number)", "LOD {0}"), LODText );
 			ToolTip = NSLOCTEXT("AssetTypeActions_StaticMesh", "NewImportTip", "Import new LOD");
 		}
 
@@ -149,38 +97,6 @@ void FAssetTypeActions_StaticMesh::GetImportLODMenu(class FMenuBuilder& MenuBuil
 	}
 }
 
-void FAssetTypeActions_StaticMesh::ExecuteFindInExplorer(TArray<TWeakObjectPtr<UStaticMesh>> Objects)
-{
-	for (auto ObjIt = Objects.CreateConstIterator(); ObjIt; ++ObjIt)
-	{
-		auto Object = (*ObjIt).Get();
-		if ( Object && Object->AssetImportData )
-		{
-			const FString SourceFilePath = FReimportManager::ResolveImportFilename(Object->AssetImportData->SourceFilePath, Object);
-			if ( SourceFilePath.Len() && IFileManager::Get().FileSize( *SourceFilePath ) != INDEX_NONE )
-			{
-				FPlatformProcess::ExploreFolder( *FPaths::GetPath(SourceFilePath) );
-			}
-		}
-	}
-}
-
-void FAssetTypeActions_StaticMesh::ExecuteOpenInExternalEditor(TArray<TWeakObjectPtr<UStaticMesh>> Objects)
-{
-	for (auto ObjIt = Objects.CreateConstIterator(); ObjIt; ++ObjIt)
-	{
-		auto Object = (*ObjIt).Get();
-		if ( Object && Object->AssetImportData )
-		{
-			const FString SourceFilePath = FReimportManager::ResolveImportFilename(Object->AssetImportData->SourceFilePath, Object);
-			if ( SourceFilePath.Len() && IFileManager::Get().FileSize( *SourceFilePath ) != INDEX_NONE )
-			{
-				FPlatformProcess::LaunchFileInDefaultExternalApplication( *SourceFilePath, NULL, ELaunchVerb::Edit );
-			}
-		}
-	}
-}
-
 void FAssetTypeActions_StaticMesh::ExecuteCreateDestructibleMesh(TArray<TWeakObjectPtr<UStaticMesh>> Objects)
 {
 	TArray< UObject* > Assets;
@@ -210,24 +126,4 @@ void FAssetTypeActions_StaticMesh::ExecuteCreateDestructibleMesh(TArray<TWeakObj
 	}
 }
 
-bool FAssetTypeActions_StaticMesh::CanExecuteSourceCommands(TArray<TWeakObjectPtr<UStaticMesh>> Objects) const
-{
-	bool bHaveSourceAsset = false;
-	for (auto ObjIt = Objects.CreateConstIterator(); ObjIt; ++ObjIt)
-	{
-		auto Object = (*ObjIt).Get();
-		if ( Object && Object->AssetImportData )
-		{
-			const FString& SourceFilePath = FReimportManager::ResolveImportFilename(Object->AssetImportData->SourceFilePath, Object);
-
-			if ( SourceFilePath.Len() && IFileManager::Get().FileSize(*SourceFilePath) != INDEX_NONE )
-			{
-				return true;
-			}
-		}
-	}
-
-	return false;
-}
-
 #undef LOCTEXT_NAMESPACE
diff --git a/Engine/Source/Developer/AssetTools/Private/AssetTypeActions/AssetTypeActions_StaticMesh.h b/Engine/Source/Developer/AssetTools/Private/AssetTypeActions/AssetTypeActions_StaticMesh.h
--- a/Engine/Source/Developer/AssetTools/Private/AssetTypeActions/AssetTypeActions_StaticMesh.h
+++ b/Engine/Source/Developer/AssetTools/Private/AssetTypeActions/AssetTypeActions_StaticMesh.h
@@ -9,31 +9,19 @@ public:
 	virtual FText GetName() const override { return NSLOCTEXT("AssetTypeActions", "AssetTypeActions_StaticMesh", "Static Mesh"); }
 	virtual FColor GetTypeColor() const override { return FColor(0, 255, 255); }
 	virtual UClass* GetSupportedClass() const override { return UStaticMesh::StaticClass(); }
-	virtual bool HasActions ( const TArray<UObject*>& InObjects ) const override { return true; }
+	virtual bool HasActions( const TArray<UObject*>& InObjects ) const override { return true; }
 	virtual void GetActions( const TArray<UObject*>& InObjects, FMenuBuilder& MenuBuilder ) override;
 	virtual void OpenAssetEditor( const TArray<UObject*>& InObjects, TSharedPtr<class IToolkitHost> EditWithinLevelEditor = TSharedPtr<IToolkitHost>() ) override;
 	virtual uint32 GetCategories() override { return EAssetTypeCategories::Basic; }
 	virtual class UThumbnailInfo* GetThumbnailInfo(UObject* Asset) const override;
+	virtual bool IsImportedAsset() const override { return true; }
+	virtual void GetResolvedSourceFilePaths(const TArray<UObject*>& TypeAssets, TArray<FString>& OutSourceFilePaths) const override;
+	// End IAssetTypeActions
 
 private:
-	/** Handler for when Edit is selected */
-	void ExecuteEdit(TArray<TWeakObjectPtr<UStaticMesh>> Objects);
-
-	/** Handler for when Reimport is selected */
-	void ExecuteReimport(TArray<TWeakObjectPtr<UStaticMesh>> Objects);
-
-	/** Handler for when FindInExplorer is selected */
-	void ExecuteFindInExplorer(TArray<TWeakObjectPtr<UStaticMesh>> Objects);
-
-	/** Handler for when OpenInExternalEditor is selected */
-	void ExecuteOpenInExternalEditor(TArray<TWeakObjectPtr<UStaticMesh>> Objects);
-
 	/** Handler for when CreateDestructibleMesh is selected */
 	void ExecuteCreateDestructibleMesh(TArray<TWeakObjectPtr<UStaticMesh>> Objects);
 
-	/** Returns true to allow execution of source file commands */
-	bool CanExecuteSourceCommands(TArray<TWeakObjectPtr<UStaticMesh>> Objects) const;
-
 	/** Handler to provide the list of LODs that can be imported or reimported */
 	void GetImportLODMenu(class FMenuBuilder& MenuBuilder,TArray<TWeakObjectPtr<UStaticMesh>> Objects);
 };
\ No newline at end of file
diff --git a/Engine/Source/Developer/AssetTools/Private/AssetTypeActions/AssetTypeActions_Texture.cpp b/Engine/Source/Developer/AssetTools/Private/AssetTypeActions/AssetTypeActions_Texture.cpp
--- a/Engine/Source/Developer/AssetTools/Private/AssetTypeActions/AssetTypeActions_Texture.cpp
+++ b/Engine/Source/Developer/AssetTools/Private/AssetTypeActions/AssetTypeActions_Texture.cpp
@@ -13,56 +13,16 @@ void FAssetTypeActions_Texture::GetActions( const TArray<UObject*>& InObjects, F
 {
 	auto Textures = GetTypedWeakObjectPtrs<UTexture>(InObjects);
 
-	MenuBuilder.AddMenuEntry(
-		LOCTEXT("Texture_Reimport", "Reimport"),
-		LOCTEXT("Texture_ReimportTooltip", "Reimports the selected textures from file."),
-		FSlateIcon(),
-		FUIAction(
-			FExecuteAction::CreateSP( this, &FAssetTypeActions_Texture::ExecuteReimport, Textures ),
-			FCanExecuteAction()
-			)
-		);
-
-	MenuBuilder.AddMenuEntry(
-		LOCTEXT("Texture_Edit", "Edit"),
-		LOCTEXT("Texture_EditTooltip", "Opens the selected textures in the texture viewer."),
-		FSlateIcon(),
-		FUIAction(
-			FExecuteAction::CreateSP( this, &FAssetTypeActions_Texture::ExecuteEdit, Textures ),
-			FCanExecuteAction()
-			)
-		);
-
 	MenuBuilder.AddMenuEntry(
 		LOCTEXT("Texture_CreateMaterial", "Create Material"),
 		LOCTEXT("Texture_CreateMaterialTooltip", "Creates a new material using this texture."),
-		FSlateIcon(),
+		FSlateIcon(FEditorStyle::GetStyleSetName(), "ClassIcon.Material"),
 		FUIAction(
 			FExecuteAction::CreateSP( this, &FAssetTypeActions_Texture::ExecuteCreateMaterial, Textures ),
 			FCanExecuteAction()
 			)
 		);
 
-	MenuBuilder.AddMenuEntry(
-		LOCTEXT("Texture_FindInExplorer", "Find Source"),
-		LOCTEXT("Texture_FindInExplorerTooltip", "Opens explorer at the location of this asset."),
-		FSlateIcon(),
-		FUIAction(
-			FExecuteAction::CreateSP( this, &FAssetTypeActions_Texture::ExecuteFindInExplorer, Textures ),
-			FCanExecuteAction::CreateSP( this, &FAssetTypeActions_Texture::CanExecuteSourceCommands, Textures )
-			)
-		);
-
-	MenuBuilder.AddMenuEntry(
-		LOCTEXT("Texture_OpenInExternalEditor", "Open Source"),
-		LOCTEXT("Texture_OpenInExternalEditorTooltip", "Opens the selected asset in an external editor."),
-		FSlateIcon(),
-		FUIAction(
-			FExecuteAction::CreateSP( this, &FAssetTypeActions_Texture::ExecuteOpenInExternalEditor, Textures ),
-			FCanExecuteAction::CreateSP( this, &FAssetTypeActions_Texture::CanExecuteSourceCommands, Textures )
-			)
-		);
-
 	// @todo AssetTypeActions Implement FindMaterials using the asset registry.
 	/*
 	if ( InObjects.Num() == 1 )
@@ -80,6 +40,15 @@ void FAssetTypeActions_Texture::GetActions( const TArray<UObject*>& InObjects, F
 	*/
 }
 
+void FAssetTypeActions_Texture::GetResolvedSourceFilePaths(const TArray<UObject*>& TypeAssets, TArray<FString>& OutSourceFilePaths) const
+{
+	for (auto& Asset : TypeAssets)
+	{
+		UTexture* Texture = CastChecked<UTexture>(Asset);
+		OutSourceFilePaths.Add(FReimportManager::ResolveImportFilename(Texture->SourceFilePath, Texture));
+	}
+}
+
 void FAssetTypeActions_Texture::OpenAssetEditor( const TArray<UObject*>& InObjects, TSharedPtr<IToolkitHost> EditWithinLevelEditor )
 {
 	EToolkitMode::Type Mode = EditWithinLevelEditor.IsValid() ? EToolkitMode::WorldCentric : EToolkitMode::Standalone;
@@ -95,30 +64,6 @@ void FAssetTypeActions_Texture::OpenAssetEditor( const TArray<UObject*>& InObjec
 	}
 }
 
-void FAssetTypeActions_Texture::ExecuteReimport(TArray<TWeakObjectPtr<UTexture>> Objects)
-{
-	for (auto ObjIt = Objects.CreateConstIterator(); ObjIt; ++ObjIt)
-	{
-		auto Object = (*ObjIt).Get();
-		if ( Object )
-		{
-			FReimportManager::Instance()->Reimport(Object, /*bAskForNewFileIfMissing=*/true);
-		}
-	}
-}
-
-void FAssetTypeActions_Texture::ExecuteEdit(TArray<TWeakObjectPtr<UTexture>> Objects)
-{
-	for (auto ObjIt = Objects.CreateConstIterator(); ObjIt; ++ObjIt)
-	{
-		auto Object = (*ObjIt).Get();
-		if ( Object )
-		{
-			FAssetEditorManager::Get().OpenEditorForAsset(Object);
-		}
-	}
-}
-
 void FAssetTypeActions_Texture::ExecuteCreateMaterial(TArray<TWeakObjectPtr<UTexture>> Objects)
 {
 	const FString DefaultSuffix = TEXT("_Mat");
@@ -176,38 +121,6 @@ void FAssetTypeActions_Texture::ExecuteCreateMaterial(TArray<TWeakObjectPtr<UTex
 	}
 }
 
-void FAssetTypeActions_Texture::ExecuteFindInExplorer(TArray<TWeakObjectPtr<UTexture>> Objects)
-{
-	for (auto ObjIt = Objects.CreateConstIterator(); ObjIt; ++ObjIt)
-	{
-		auto Object = (*ObjIt).Get();
-		if ( Object )
-		{
-			const FString SourceFilePath = FReimportManager::ResolveImportFilename(Object->SourceFilePath, Object);
-			if ( SourceFilePath.Len() && IFileManager::Get().FileSize( *SourceFilePath ) != INDEX_NONE )
-			{
-				FPlatformProcess::ExploreFolder( *FPaths::GetPath(SourceFilePath) );
-			}
-		}
-	}
-}
-
-void FAssetTypeActions_Texture::ExecuteOpenInExternalEditor(TArray<TWeakObjectPtr<UTexture>> Objects)
-{
-	for (auto ObjIt = Objects.CreateConstIterator(); ObjIt; ++ObjIt)
-	{
-		auto Object = (*ObjIt).Get();
-		if ( Object )
-		{
-			const FString SourceFilePath = FReimportManager::ResolveImportFilename(Object->SourceFilePath, Object);
-			if ( SourceFilePath.Len() && IFileManager::Get().FileSize( *SourceFilePath ) != INDEX_NONE )
-			{
-				FPlatformProcess::LaunchFileInDefaultExternalApplication( *SourceFilePath, NULL, ELaunchVerb::Edit );
-			}
-		}
-	}
-}
-
 void FAssetTypeActions_Texture::ExecuteFindMaterials(TWeakObjectPtr<UTexture> Object)
 {
 	auto Texture = Object.Get();
@@ -217,24 +130,4 @@ void FAssetTypeActions_Texture::ExecuteFindMaterials(TWeakObjectPtr<UTexture> Ob
 	}
 }
 
-bool FAssetTypeActions_Texture::CanExecuteSourceCommands(TArray<TWeakObjectPtr<UTexture>> Objects) const
-{
-	bool bHaveSourceAsset = false;
-	for (auto ObjIt = Objects.CreateConstIterator(); ObjIt; ++ObjIt)
-	{
-		auto Object = (*ObjIt).Get();
-		if ( Object )
-		{
-			const FString& SourceFilePath = FReimportManager::ResolveImportFilename(Object->SourceFilePath, Object);
-
-			if ( SourceFilePath.Len() && IFileManager::Get().FileSize(*SourceFilePath) != INDEX_NONE )
-			{
-				return true;
-			}
-		}
-	}
-
-	return false;
-}
-
 #undef LOCTEXT_NAMESPACE
\ No newline at end of file
diff --git a/Engine/Source/Developer/AssetTools/Private/AssetTypeActions/AssetTypeActions_Texture.h b/Engine/Source/Developer/AssetTools/Private/AssetTypeActions/AssetTypeActions_Texture.h
--- a/Engine/Source/Developer/AssetTools/Private/AssetTypeActions/AssetTypeActions_Texture.h
+++ b/Engine/Source/Developer/AssetTools/Private/AssetTypeActions/AssetTypeActions_Texture.h
@@ -14,26 +14,13 @@ public:
 	virtual void OpenAssetEditor( const TArray<UObject*>& InObjects, TSharedPtr<class IToolkitHost> EditWithinLevelEditor = TSharedPtr<IToolkitHost>() ) override;
 	virtual bool CanFilter() override { return false; }
 	virtual uint32 GetCategories() override { return EAssetTypeCategories::MaterialsAndTextures; }
+	virtual bool IsImportedAsset() const override { return true; }
+	virtual void GetResolvedSourceFilePaths(const TArray<UObject*>& TypeAssets, TArray<FString>& OutSourceFilePaths) const override;
 
 private:
-	/** Handler for when Reimport is selected */
-	void ExecuteReimport(TArray<TWeakObjectPtr<UTexture>> Objects);
-
-	/** Handler for when Edit is selected */
-	void ExecuteEdit(TArray<TWeakObjectPtr<UTexture>> Objects);
-
 	/** Handler for when CreateMaterial is selected */
 	void ExecuteCreateMaterial(TArray<TWeakObjectPtr<UTexture>> Objects);
 
-	/** Handler for when FindInExplorer is selected */
-	void ExecuteFindInExplorer(TArray<TWeakObjectPtr<UTexture>> Objects);
-
-	/** Handler for when OpenInExternalEditor is selected */
-	void ExecuteOpenInExternalEditor(TArray<TWeakObjectPtr<UTexture>> Objects);
-
 	/** Handler for when FindMaterials is selected */
 	void ExecuteFindMaterials(TWeakObjectPtr<UTexture> Object);
-
-	/** Returns true to allow execution of source file commands */
-	bool CanExecuteSourceCommands(TArray<TWeakObjectPtr<UTexture>> Objects) const;
 };
\ No newline at end of file
diff --git a/Engine/Source/Developer/AssetTools/Private/AssetTypeActions/AssetTypeActions_Texture2D.cpp b/Engine/Source/Developer/AssetTools/Private/AssetTypeActions/AssetTypeActions_Texture2D.cpp
--- a/Engine/Source/Developer/AssetTools/Private/AssetTypeActions/AssetTypeActions_Texture2D.cpp
+++ b/Engine/Source/Developer/AssetTools/Private/AssetTypeActions/AssetTypeActions_Texture2D.cpp
@@ -15,7 +15,7 @@ void FAssetTypeActions_Texture2D::GetActions( const TArray<UObject*>& InObjects,
 	MenuBuilder.AddMenuEntry(
 		LOCTEXT("Texture2D_CreateSlateBrush", "Create Slate Brush"),
 		LOCTEXT("Texture2D_CreateSlateBrushToolTip", "Creates a new slate brush using this texture."),
-		FSlateIcon(),
+		FSlateIcon(FEditorStyle::GetStyleSetName(), "ClassIcon.SlateBrushAsset"),
 		FUIAction(
 			FExecuteAction::CreateSP( this, &FAssetTypeActions_Texture2D::ExecuteCreateSlateBrush, Textures ),
 			FCanExecuteAction()
diff --git a/Engine/Source/Developer/AssetTools/Private/AssetTypeActions/AssetTypeActions_TextureRenderTarget.cpp b/Engine/Source/Developer/AssetTools/Private/AssetTypeActions/AssetTypeActions_TextureRenderTarget.cpp
--- a/Engine/Source/Developer/AssetTools/Private/AssetTypeActions/AssetTypeActions_TextureRenderTarget.cpp
+++ b/Engine/Source/Developer/AssetTools/Private/AssetTypeActions/AssetTypeActions_TextureRenderTarget.cpp
@@ -14,7 +14,7 @@ void FAssetTypeActions_TextureRenderTarget::GetActions( const TArray<UObject*>&
 	MenuBuilder.AddMenuEntry(
 		LOCTEXT("TextureRenderTarget_CreateStatic", "Create Static Texture"),
 		LOCTEXT("TextureRenderTarget_CreateStaticTooltip", "Creates a static texture from the selected render targets."),
-		FSlateIcon(),
+		FSlateIcon(FEditorStyle::GetStyleSetName(), "ClassIcon.Texture2D"),
 		FUIAction(
 			FExecuteAction::CreateSP( this, &FAssetTypeActions_TextureRenderTarget::ExecuteCreateStatic, RenderTargets ),
 			FCanExecuteAction()
diff --git a/Engine/Source/Developer/AssetTools/Private/AssetTypeActions/AssetTypeActions_VectorFieldStatic.cpp b/Engine/Source/Developer/AssetTools/Private/AssetTypeActions/AssetTypeActions_VectorFieldStatic.cpp
--- a/Engine/Source/Developer/AssetTools/Private/AssetTypeActions/AssetTypeActions_VectorFieldStatic.cpp
+++ b/Engine/Source/Developer/AssetTools/Private/AssetTypeActions/AssetTypeActions_VectorFieldStatic.cpp
@@ -2,105 +2,11 @@
 
 #include "AssetToolsPrivatePCH.h"
 
-#define LOCTEXT_NAMESPACE "AssetTypeActions"
-
-void FAssetTypeActions_VectorFieldStatic::GetActions( const TArray<UObject*>& InObjects, FMenuBuilder& MenuBuilder )
-{
-	auto Fields = GetTypedWeakObjectPtrs<UVectorFieldStatic>(InObjects);
-
-	MenuBuilder.AddMenuEntry(
-		LOCTEXT("VectorFieldStatic_Reimport", "Reimport"),
-		LOCTEXT("VectorFieldStatic_ReimportTooltip", "Reimports the selected meshes from file."),
-		FSlateIcon(),
-		FUIAction(
-			FExecuteAction::CreateSP( this, &FAssetTypeActions_VectorFieldStatic::ExecuteReimport, Fields ),
-			FCanExecuteAction()
-			)
-		);
-
-	MenuBuilder.AddMenuEntry(
-		LOCTEXT("VectorFieldStatic_FindInExplorer", "Find Source"),
-		LOCTEXT("VectorFieldStatic_FindInExplorerTooltip", "Opens explorer at the location of this asset."),
-		FSlateIcon(),
-		FUIAction(
-			FExecuteAction::CreateSP( this, &FAssetTypeActions_VectorFieldStatic::ExecuteFindInExplorer, Fields ),
-			FCanExecuteAction::CreateSP( this, &FAssetTypeActions_VectorFieldStatic::CanExecuteSourceCommands, Fields )
-			)
-		);
-
-	MenuBuilder.AddMenuEntry(
-		LOCTEXT("VectorFieldStatic_OpenInExternalEditor", "Open Source"),
-		LOCTEXT("VectorFieldStatic_OpenInExternalEditorTooltip", "Opens the selected asset in an external editor."),
-		FSlateIcon(),
-		FUIAction(
-			FExecuteAction::CreateSP( this, &FAssetTypeActions_VectorFieldStatic::ExecuteOpenInExternalEditor, Fields ),
-			FCanExecuteAction::CreateSP( this, &FAssetTypeActions_VectorFieldStatic::CanExecuteSourceCommands, Fields )
-			)
-		);
-}
-
-void FAssetTypeActions_VectorFieldStatic::ExecuteReimport(TArray<TWeakObjectPtr<UVectorFieldStatic>> Objects)
+void FAssetTypeActions_VectorFieldStatic::GetResolvedSourceFilePaths(const TArray<UObject*>& TypeAssets, TArray<FString>& OutSourceFilePaths) const
 {
-	for (auto ObjIt = Objects.CreateConstIterator(); ObjIt; ++ObjIt)
+	for (auto& Asset : TypeAssets)
 	{
-		auto Object = (*ObjIt).Get();
-		if ( Object )
-		{
-			FReimportManager::Instance()->Reimport(Object, /*bAskForNewFileIfMissing=*/true);
-		}
+		const auto StaticVectorField = CastChecked<UVectorFieldStatic>(Asset);
+		OutSourceFilePaths.Add(FReimportManager::ResolveImportFilename(StaticVectorField->SourceFilePath, StaticVectorField));
 	}
-}
-
-void FAssetTypeActions_VectorFieldStatic::ExecuteFindInExplorer(TArray<TWeakObjectPtr<UVectorFieldStatic>> Objects)
-{
-	for (auto ObjIt = Objects.CreateConstIterator(); ObjIt; ++ObjIt)
-	{
-		auto Object = (*ObjIt).Get();
-		if ( Object )
-		{
-			const FString SourceFilePath = FReimportManager::ResolveImportFilename(Object->SourceFilePath, Object);
-			if ( SourceFilePath.Len() && IFileManager::Get().FileSize( *SourceFilePath ) != INDEX_NONE )
-			{
-				FPlatformProcess::ExploreFolder( *FPaths::GetPath(SourceFilePath) );
-			}
-		}
-	}
-}
-
-void FAssetTypeActions_VectorFieldStatic::ExecuteOpenInExternalEditor(TArray<TWeakObjectPtr<UVectorFieldStatic>> Objects)
-{
-	for (auto ObjIt = Objects.CreateConstIterator(); ObjIt; ++ObjIt)
-	{
-		auto Object = (*ObjIt).Get();
-		if ( Object )
-		{
-			const FString SourceFilePath = FReimportManager::ResolveImportFilename(Object->SourceFilePath, Object);
-			if ( SourceFilePath.Len() && IFileManager::Get().FileSize( *SourceFilePath ) != INDEX_NONE )
-			{
-				FPlatformProcess::LaunchFileInDefaultExternalApplication( *SourceFilePath, NULL, ELaunchVerb::Edit );
-			}
-		}
-	}
-}
-
-bool FAssetTypeActions_VectorFieldStatic::CanExecuteSourceCommands(TArray<TWeakObjectPtr<UVectorFieldStatic>> Objects) const
-{
-	bool bHaveSourceAsset = false;
-	for (auto ObjIt = Objects.CreateConstIterator(); ObjIt; ++ObjIt)
-	{
-		auto Object = (*ObjIt).Get();
-		if ( Object )
-		{
-			const FString& SourceFilePath = FReimportManager::ResolveImportFilename(Object->SourceFilePath, Object);
-
-			if ( SourceFilePath.Len() && IFileManager::Get().FileSize(*SourceFilePath) != INDEX_NONE )
-			{
-				return true;
-			}
-		}
-	}
-
-	return false;
-}
-
-#undef LOCTEXT_NAMESPACE
+}
\ No newline at end of file
diff --git a/Engine/Source/Developer/AssetTools/Private/AssetTypeActions/AssetTypeActions_VectorFieldStatic.h b/Engine/Source/Developer/AssetTools/Private/AssetTypeActions/AssetTypeActions_VectorFieldStatic.h
--- a/Engine/Source/Developer/AssetTools/Private/AssetTypeActions/AssetTypeActions_VectorFieldStatic.h
+++ b/Engine/Source/Developer/AssetTools/Private/AssetTypeActions/AssetTypeActions_VectorFieldStatic.h
@@ -4,29 +4,16 @@
 #include "AssetTypeActions/AssetTypeActions_VectorField.h"
 #include "VectorField/VectorFieldStatic.h"
 
-class FAssetTypeActions_VectorFieldStatic : public FAssetTypeActions_VectorField
+class FAssetTypeActions_VectorFieldStatic : public FAssetTypeActions_Base
 {
 public:
 	// IAssetTypeActions Implementation
-	virtual FText GetName() const override { return NSLOCTEXT("AssetTypeActions", "AssetTypeActions_VectorFieldStatic", "Vector Field"); }
+	virtual FText GetName() const override { return NSLOCTEXT("AssetTypeActions", "AssetTypeActions_VectorFieldStatic", "Static Vector Field"); }
+	virtual FColor GetTypeColor() const override { return FColor(200, 128, 128); }
 	virtual UClass* GetSupportedClass() const override { return UVectorFieldStatic::StaticClass(); }
-	virtual bool HasActions ( const TArray<UObject*>& InObjects ) const override { return true; }
-	virtual void GetActions( const TArray<UObject*>& InObjects, FMenuBuilder& MenuBuilder ) override;
 	virtual bool CanFilter() override { return true; }
+	virtual uint32 GetCategories() override { return EAssetTypeCategories::Misc; }
+	virtual bool IsImportedAsset() const override { return true; }
+	virtual void GetResolvedSourceFilePaths(const TArray<UObject*>& TypeAssets, TArray<FString>& OutSourceFilePaths) const override;
 
-private:
-	/** Handler for when Edit is selected */
-	void ExecuteEdit(TArray<TWeakObjectPtr<UVectorFieldStatic>> Objects);
-
-	/** Handler for when Reimport is selected */
-	void ExecuteReimport(TArray<TWeakObjectPtr<UVectorFieldStatic>> Objects);
-
-	/** Handler for when FindInExplorer is selected */
-	void ExecuteFindInExplorer(TArray<TWeakObjectPtr<UVectorFieldStatic>> Objects);
-
-	/** Handler for when OpenInExternalEditor is selected */
-	void ExecuteOpenInExternalEditor(TArray<TWeakObjectPtr<UVectorFieldStatic>> Objects);
-
-	/** Returns true to allow execution of source file commands */
-	bool CanExecuteSourceCommands(TArray<TWeakObjectPtr<UVectorFieldStatic>> Objects) const;
 };
\ No newline at end of file
diff --git a/Engine/Source/Developer/AssetTools/Private/AssetTypeActions/AssetTypeActions_VertexAnimation.cpp b/Engine/Source/Developer/AssetTools/Private/AssetTypeActions/AssetTypeActions_VertexAnimation.cpp
--- a/Engine/Source/Developer/AssetTools/Private/AssetTypeActions/AssetTypeActions_VertexAnimation.cpp
+++ b/Engine/Source/Developer/AssetTools/Private/AssetTypeActions/AssetTypeActions_VertexAnimation.cpp
@@ -5,21 +5,6 @@
 
 #define LOCTEXT_NAMESPACE "AssetTypeActions"
 
-void FAssetTypeActions_VertexAnimation::GetActions( const TArray<UObject*>& InObjects, FMenuBuilder& MenuBuilder )
-{
-	auto VertexAnims = GetTypedWeakObjectPtrs<UVertexAnimation>(InObjects);
-
-	MenuBuilder.AddMenuEntry(
-		LOCTEXT("VertexAnim_Edit", "Edit"),
-		LOCTEXT("VertexAnim_EditTooltip", "Opens the selected vertex animation in Persona."),
-		FSlateIcon(),
-		FUIAction(
-		FExecuteAction::CreateSP( this, &FAssetTypeActions_VertexAnimation::ExecuteEdit, VertexAnims ),
-		FCanExecuteAction()
-		)
-		);
-}
-
 void FAssetTypeActions_VertexAnimation::OpenAssetEditor( const TArray<UObject*>& InObjects, TSharedPtr<IToolkitHost> EditWithinLevelEditor )
 {
 	for (auto ObjIt = InObjects.CreateConstIterator(); ObjIt; ++ObjIt)
@@ -31,18 +16,6 @@ void FAssetTypeActions_VertexAnimation::OpenAssetEditor( const TArray<UObject*>&
 	}
 }
 
-void FAssetTypeActions_VertexAnimation::ExecuteEdit(TArray< TWeakObjectPtr<UVertexAnimation> > Objects)
-{
-	for (auto ObjIt = Objects.CreateConstIterator(); ObjIt; ++ObjIt)
-	{
-		auto Object = (*ObjIt).Get();
-		if ( Object && Object->BaseSkelMesh )
-		{
-			FAssetEditorManager::Get().OpenEditorForAsset(Object);
-		}
-	}
-}
-
 bool FAssetTypeActions_VertexAnimation::CanFilter()
 {
 	return false;
diff --git a/Engine/Source/Developer/AssetTools/Private/AssetTypeActions/AssetTypeActions_VertexAnimation.h b/Engine/Source/Developer/AssetTools/Private/AssetTypeActions/AssetTypeActions_VertexAnimation.h
--- a/Engine/Source/Developer/AssetTools/Private/AssetTypeActions/AssetTypeActions_VertexAnimation.h
+++ b/Engine/Source/Developer/AssetTools/Private/AssetTypeActions/AssetTypeActions_VertexAnimation.h
@@ -11,8 +11,6 @@ public:
 	virtual FText GetName() const override { return NSLOCTEXT("AssetTypeActions", "AssetTypeActions_VertexAnimation", "Vertex Animation"); }
 	virtual FColor GetTypeColor() const override { return FColor(87,0,174); }
 	virtual UClass* GetSupportedClass() const override { return UVertexAnimation::StaticClass(); }
-	virtual bool HasActions ( const TArray<UObject*>& InObjects ) const override { return true; }
-	virtual void GetActions( const TArray<UObject*>& InObjects, FMenuBuilder& MenuBuilder ) override;
 	virtual void OpenAssetEditor( const TArray<UObject*>& InObjects, TSharedPtr<class IToolkitHost> EditWithinLevelEditor = TSharedPtr<IToolkitHost>() ) override;
 	virtual uint32 GetCategories() override { return EAssetTypeCategories::Animation; }
 	
@@ -20,8 +18,4 @@ public:
 	// the files in question; returning false from here disables the filer in the content browser.
 	// Remove this when vertex anims are supported again.
 	virtual bool CanFilter() override;
-
-private:
-	/** Handler for when Edit is selected */
-	void ExecuteEdit(TArray< TWeakObjectPtr<UVertexAnimation> > Objects);
 };
\ No newline at end of file
diff --git a/Engine/Source/Developer/AssetTools/Public/AssetTypeActions_Base.h b/Engine/Source/Developer/AssetTools/Public/AssetTypeActions_Base.h
--- a/Engine/Source/Developer/AssetTools/Public/AssetTypeActions_Base.h
+++ b/Engine/Source/Developer/AssetTools/Public/AssetTypeActions_Base.h
@@ -4,6 +4,7 @@
 
 #include "AssetToolsModule.h"
 #include "Toolkits/SimpleAssetEditor.h"
+#include "EditorReimportHandler.h"
 
 /** A base class for all AssetTypeActions. Provides helper functions useful for many types. Deriving from this class is optional. */
 class FAssetTypeActions_Base : public IAssetTypeActions
@@ -69,6 +70,15 @@ public:
 		return NULL;
 	}
 
+	virtual bool IsImportedAsset() const
+	{
+		return false;
+	}
+
+	virtual void GetResolvedSourceFilePaths(const TArray<UObject*>& TypeAssets, TArray<FString>& OutSourceFilePaths) const override
+	{
+	}
+
 	// End IAssetTypeActions implementation
 
 protected:
diff --git a/Engine/Source/Developer/AssetTools/Public/AssetTypeActions_CSVAssetBase.h b/Engine/Source/Developer/AssetTools/Public/AssetTypeActions_CSVAssetBase.h
--- a/Engine/Source/Developer/AssetTools/Public/AssetTypeActions_CSVAssetBase.h
+++ b/Engine/Source/Developer/AssetTools/Public/AssetTypeActions_CSVAssetBase.h
@@ -4,25 +4,18 @@
 
 class FAssetTypeActions_CSVAssetBase : public FAssetTypeActions_Base
 {
+public:
+	virtual FColor GetTypeColor() const override { return FColor(62, 140, 35); }
+	virtual bool HasActions(const TArray<UObject*>& InObjects) const override { return true; }
+	virtual uint32 GetCategories() override { return EAssetTypeCategories::Misc; }
+	virtual bool IsImportedAsset() const override { return true; }
 
 protected:
-	/** Handler for when Edit is selected */
-	void ExecuteEdit(TArray< TWeakObjectPtr<UObject> > Objects);
+	/** Handler for opening the xls/xlsm version of the source file */
+	void ExecuteFindExcelFileInExplorer(TArray<FString> Filenames, TArray<FString> OverrideExtensions);
 
-	/** Handler for when Reimport is selected */
-	void ExecuteReimport(TArray< TWeakObjectPtr<UObject> > Objects);
-
-	/** Handler for when FindInExplorer is selected */
-	void ExecuteFindInExplorer(TArray<FString> Filenames);
-
-	/** Handler for when OpenInExternalEditor is selected */
-	void ExecuteOpenInExternalEditor(TArray<FString> Filenames, TArray<FString> OverrideExtensions);
-
-	/** Determine whether the launch in external editor commands can execute or not */
-	bool CanExecuteLaunchExternalSourceCommands(TArray<FString> Filenames, TArray<FString> OverrideExtensions) const;
-
-	/** Determine whether the find in exploder editor command can execute or not */
-	bool CanExecuteFindInExplorerSourceCommand(TArray<FString> Filenames) const;
+	/** Determine whether the find excel file in explorer editor command can execute or not */
+	bool CanExecuteFindExcelFileInExplorer(TArray<FString> Filenames, TArray<FString> OverrideExtensions) const;
 
 	/** Verify the specified filename exists */
 	bool VerifyFileExists(const FString& InFileName) const;
diff --git a/Engine/Source/Developer/AssetTools/Public/IAssetTypeActions.h b/Engine/Source/Developer/AssetTools/Public/IAssetTypeActions.h
--- a/Engine/Source/Developer/AssetTools/Public/IAssetTypeActions.h
+++ b/Engine/Source/Developer/AssetTools/Public/IAssetTypeActions.h
@@ -91,4 +91,10 @@ public:
 
 	/** Returns additional tooltip information for the specified asset, if it has any (otherwise return the null widget) */
 	virtual FText GetAssetDescription(const class FAssetData& AssetData) const = 0;
+
+	/** Returns whether the asset was imported from an external source */
+	virtual bool IsImportedAsset() const = 0;
+
+	/** Collects the resolved source paths for the imported assets */
+	virtual void GetResolvedSourceFilePaths(const TArray<UObject*>& TypeAssets, TArray<FString>& OutSourceFilePaths) const = 0;
 };
diff --git a/Engine/Source/Developer/MovieSceneTools/Private/AssetTypeActions_MovieScene.cpp b/Engine/Source/Developer/MovieSceneTools/Private/AssetTypeActions_MovieScene.cpp
--- a/Engine/Source/Developer/MovieSceneTools/Private/AssetTypeActions_MovieScene.cpp
+++ b/Engine/Source/Developer/MovieSceneTools/Private/AssetTypeActions_MovieScene.cpp
@@ -15,23 +15,6 @@ UClass* FAssetTypeActions_MovieScene::GetSupportedClass() const
 	return UMovieScene::StaticClass(); 
 }
 
-
-void FAssetTypeActions_MovieScene::GetActions( const TArray<UObject*>& InObjects, FMenuBuilder& MenuBuilder )
-{
-	auto MovieScenes = GetTypedWeakObjectPtrs<UMovieScene>(InObjects);
-
-	MenuBuilder.AddMenuEntry(
-		LOCTEXT("MovieScene_Edit", "Edit"),
-		LOCTEXT("MovieScene_EditTooltip", "Opens the selected movie scenes in the sequencer for editing."),
-		FSlateIcon(),
-		FUIAction(
-			FExecuteAction::CreateSP( this, &FAssetTypeActions_MovieScene::ExecuteEdit, MovieScenes ),
-			FCanExecuteAction()
-			)
-		);
-}
-
-
 void FAssetTypeActions_MovieScene::OpenAssetEditor( const TArray<UObject*>& InObjects, TSharedPtr<IToolkitHost> EditWithinLevelEditor )
 {
 	EToolkitMode::Type Mode = EditWithinLevelEditor.IsValid() ? EToolkitMode::WorldCentric : EToolkitMode::Standalone;
@@ -51,20 +34,4 @@ void FAssetTypeActions_MovieScene::OpenAssetEditor( const TArray<UObject*>& InOb
 	}
 }
 
-
-void FAssetTypeActions_MovieScene::ExecuteEdit(TArray<TWeakObjectPtr<UMovieScene>> Objects)
-{
-	for (auto ObjIt = Objects.CreateConstIterator(); ObjIt; ++ObjIt)
-	{
-		auto Object = (*ObjIt).Get();
-		if ( Object )
-		{
-			FAssetEditorManager::Get().OpenEditorForAsset(Object);
-		}
-	}
-}
-
-
-
-
 #undef LOCTEXT_NAMESPACE
\ No newline at end of file
diff --git a/Engine/Source/Developer/MovieSceneTools/Private/AssetTypeActions_MovieScene.h b/Engine/Source/Developer/MovieSceneTools/Private/AssetTypeActions_MovieScene.h
--- a/Engine/Source/Developer/MovieSceneTools/Private/AssetTypeActions_MovieScene.h
+++ b/Engine/Source/Developer/MovieSceneTools/Private/AssetTypeActions_MovieScene.h
@@ -11,8 +11,6 @@ public:
 	virtual FText GetName() const override { return NSLOCTEXT("AssetTypeActions", "AssetTypeActions_MovieScene", "Movie Scene"); }
 	virtual FColor GetTypeColor() const override { return FColor(200, 80, 80); }
 	virtual UClass* GetSupportedClass() const override;
-	virtual bool HasActions ( const TArray<UObject*>& InObjects ) const override { return true; }
-	virtual void GetActions( const TArray<UObject*>& InObjects, FMenuBuilder& MenuBuilder ) override;
 	virtual void OpenAssetEditor( const TArray<UObject*>& InObjects, TSharedPtr<class IToolkitHost> EditWithinLevelEditor = TSharedPtr<IToolkitHost>() ) override;
 	virtual uint32 GetCategories() override { return EAssetTypeCategories::Animation; }
 
@@ -21,9 +19,4 @@ public:
 		// @todo sequencer: Hack to force world-centric mode for Sequencer
 		return true;
 	}
-
-protected:
-	/** Handler for when Edit is selected */
-	virtual void ExecuteEdit(TArray<TWeakObjectPtr<class UMovieScene>> Objects);
-
 };
\ No newline at end of file
diff --git a/Engine/Source/Editor/BehaviorTreeEditor/Private/AssetTypeActions_BehaviorTree.cpp b/Engine/Source/Editor/BehaviorTreeEditor/Private/AssetTypeActions_BehaviorTree.cpp
--- a/Engine/Source/Editor/BehaviorTreeEditor/Private/AssetTypeActions_BehaviorTree.cpp
+++ b/Engine/Source/Editor/BehaviorTreeEditor/Private/AssetTypeActions_BehaviorTree.cpp
@@ -13,26 +13,6 @@
 
 #define LOCTEXT_NAMESPACE "AssetTypeActions"
 
-bool FAssetTypeActions_BehaviorTree::HasActions ( const TArray<UObject*>& InObjects ) const
-{
-	return true;
-}
-
-void FAssetTypeActions_BehaviorTree::GetActions( const TArray<UObject*>& InObjects, FMenuBuilder& MenuBuilder )
-{
-	auto BehaviorTrees = GetTypedWeakObjectPtrs<UBehaviorTree>(InObjects);
-
-	MenuBuilder.AddMenuEntry(
-		LOCTEXT("BehaviorTree_Edit", "Edit"),
-		LOCTEXT("BehaviorTree_EditTooltip", "Opens the selected Behavior Tree in editor."),
-		FSlateIcon(),
-		FUIAction(
-		FExecuteAction::CreateSP( this, &FAssetTypeActions_BehaviorTree::ExecuteEdit, BehaviorTrees ),
-		FCanExecuteAction()
-		)
-	);
-}
-
 void FAssetTypeActions_BehaviorTree::OpenAssetEditor( const TArray<UObject*>& InObjects, TSharedPtr<IToolkitHost> EditWithinLevelEditor )
 {
 	EToolkitMode::Type Mode = EditWithinLevelEditor.IsValid() ? EToolkitMode::WorldCentric : EToolkitMode::Standalone;
@@ -64,18 +44,6 @@ void FAssetTypeActions_BehaviorTree::OpenAssetEditor( const TArray<UObject*>& In
 	}
 }
 
-void FAssetTypeActions_BehaviorTree::ExecuteEdit(TArray<TWeakObjectPtr<UBehaviorTree>> Objects)
-{
-	for(auto Object : Objects)
-	{
-		auto BehaviorTree = Object.Get();
-		if(BehaviorTree != NULL)
-		{
-			FAssetEditorManager::Get().OpenEditorForAsset(BehaviorTree);
-		}
-	}
-}
-
 UClass* FAssetTypeActions_BehaviorTree::GetSupportedClass() const
 { 
 	return UBehaviorTree::StaticClass(); 
diff --git a/Engine/Source/Editor/BehaviorTreeEditor/Private/AssetTypeActions_BehaviorTree.h b/Engine/Source/Editor/BehaviorTreeEditor/Private/AssetTypeActions_BehaviorTree.h
--- a/Engine/Source/Editor/BehaviorTreeEditor/Private/AssetTypeActions_BehaviorTree.h
+++ b/Engine/Source/Editor/BehaviorTreeEditor/Private/AssetTypeActions_BehaviorTree.h
@@ -10,8 +10,6 @@ public:
 	virtual FText GetName() const override { return NSLOCTEXT("AssetTypeActions", "AssetTypeActions_BehaviorTree", "Behavior Tree"); }
 	virtual FColor GetTypeColor() const override { return FColor(149,70,255); }
 	virtual UClass* GetSupportedClass() const override;
-	virtual bool HasActions ( const TArray<UObject*>& InObjects ) const override;
-	virtual void GetActions( const TArray<UObject*>& InObjects, FMenuBuilder& MenuBuilder ) override;
 	virtual void OpenAssetEditor( const TArray<UObject*>& InObjects, TSharedPtr<class IToolkitHost> EditWithinLevelEditor = TSharedPtr<IToolkitHost>() ) override;
 	virtual uint32 GetCategories() override { return EAssetTypeCategories::Misc; }
 	virtual void PerformAssetDiff(UObject* Asset1, UObject* Asset2, const struct FRevisionInfo& OldRevision, const struct FRevisionInfo& NewRevision) const override;
@@ -20,7 +18,4 @@ private:
 
 	/* Called to open the Behavior Tree defaults view, this opens whatever text dif tool the user has */
 	void OpenInDefaults( class UBehaviorTree* OldBehaviorTree, class UBehaviorTree* NewBehaviorTree ) const;
-
-	/** Handler for when Edit is selected */
-	void ExecuteEdit(TArray<TWeakObjectPtr<class UBehaviorTree>> Objects);
 };
\ No newline at end of file
diff --git a/Engine/Source/Editor/BehaviorTreeEditor/Private/AssetTypeActions_Blackboard.cpp b/Engine/Source/Editor/BehaviorTreeEditor/Private/AssetTypeActions_Blackboard.cpp
--- a/Engine/Source/Editor/BehaviorTreeEditor/Private/AssetTypeActions_Blackboard.cpp
+++ b/Engine/Source/Editor/BehaviorTreeEditor/Private/AssetTypeActions_Blackboard.cpp
@@ -17,26 +17,6 @@ UClass* FAssetTypeActions_Blackboard::GetSupportedClass() const
 	return UBlackboardData::StaticClass(); 
 }
 
-bool FAssetTypeActions_Blackboard::HasActions ( const TArray<UObject*>& InObjects ) const
-{
-	return true;
-}
-
-void FAssetTypeActions_Blackboard::GetActions( const TArray<UObject*>& InObjects, FMenuBuilder& MenuBuilder )
-{
-	auto BlackboardData = GetTypedWeakObjectPtrs<UBlackboardData>(InObjects);
-
-	MenuBuilder.AddMenuEntry(
-		LOCTEXT("Blackboard_Edit", "Edit"),
-		LOCTEXT("Blackboard_EditTooltip", "Opens the selected Blackboard in editor."),
-		FSlateIcon(),
-		FUIAction(
-		FExecuteAction::CreateSP( this, &FAssetTypeActions_Blackboard::ExecuteEdit, BlackboardData ),
-		FCanExecuteAction()
-		)
-	);
-}
-
 void FAssetTypeActions_Blackboard::OpenAssetEditor( const TArray<UObject*>& InObjects, TSharedPtr<class IToolkitHost> EditWithinLevelEditor )
 {
 	EToolkitMode::Type Mode = EditWithinLevelEditor.IsValid() ? EToolkitMode::WorldCentric : EToolkitMode::Standalone;
@@ -52,16 +32,4 @@ void FAssetTypeActions_Blackboard::OpenAssetEditor( const TArray<UObject*>& InOb
 	}	
 }
 
-void FAssetTypeActions_Blackboard::ExecuteEdit(TArray<TWeakObjectPtr<class UBlackboardData>> Objects)
-{
-	for(auto Object : Objects)
-	{
-		auto BlackboardData = Object.Get();
-		if(BlackboardData != nullptr)
-		{
-			FAssetEditorManager::Get().OpenEditorForAsset(BlackboardData);
-		}
-	}
-}
-
 #undef LOCTEXT_NAMESPACE
diff --git a/Engine/Source/Editor/BehaviorTreeEditor/Private/AssetTypeActions_Blackboard.h b/Engine/Source/Editor/BehaviorTreeEditor/Private/AssetTypeActions_Blackboard.h
--- a/Engine/Source/Editor/BehaviorTreeEditor/Private/AssetTypeActions_Blackboard.h
+++ b/Engine/Source/Editor/BehaviorTreeEditor/Private/AssetTypeActions_Blackboard.h
@@ -10,12 +10,6 @@ public:
 	virtual FText GetName() const override { return NSLOCTEXT("AssetTypeActions", "AssetTypeActions_Blackboard", "Blackboard"); }
 	virtual FColor GetTypeColor() const override { return FColor(201, 29, 85); }
 	virtual UClass* GetSupportedClass() const override;
-	virtual bool HasActions ( const TArray<UObject*>& InObjects ) const override;
-	virtual void GetActions( const TArray<UObject*>& InObjects, FMenuBuilder& MenuBuilder ) override;
 	virtual void OpenAssetEditor( const TArray<UObject*>& InObjects, TSharedPtr<class IToolkitHost> EditWithinLevelEditor = TSharedPtr<IToolkitHost>() ) override;
 	virtual uint32 GetCategories() override { return EAssetTypeCategories::Misc; }
-
-private:
-	/** Handler for when Edit is selected */
-	void ExecuteEdit(TArray<TWeakObjectPtr<class UBlackboardData>> Objects);
 };
\ No newline at end of file
diff --git a/Engine/Source/Editor/ContentBrowser/Private/AssetContextMenu.cpp b/Engine/Source/Editor/ContentBrowser/Private/AssetContextMenu.cpp
--- a/Engine/Source/Editor/ContentBrowser/Private/AssetContextMenu.cpp
+++ b/Engine/Source/Editor/ContentBrowser/Private/AssetContextMenu.cpp
@@ -43,12 +43,22 @@ FAssetContextMenu::FAssetContextMenu(const TWeakPtr<SAssetView>& InAssetView)
 	, bCanExecuteSCCRevert(false)
 	, bCanExecuteSCCSync(false)
 {
+	
+}
 
+void FAssetContextMenu::BindCommands(TSharedPtr< FUICommandList >& Commands)
+{
+	Commands->MapAction(FGenericCommands::Get().Duplicate, FUIAction(
+		FExecuteAction::CreateSP(this, &FAssetContextMenu::ExecuteDuplicate),
+		FCanExecuteAction::CreateSP(this, &FAssetContextMenu::CanExecuteDuplicate),
+		FIsActionChecked(),
+		FIsActionButtonVisible::CreateSP(this, &FAssetContextMenu::CanExecuteDuplicate)
+		));
 }
 
 TSharedRef<SWidget> FAssetContextMenu::MakeContextMenu(const TArray<FAssetData>& InSelectedAssets, const FSourcesData& InSourcesData, TSharedPtr< FUICommandList > InCommandList)
 {
-	SelectedAssets = InSelectedAssets;
+	SetSelectedAssets(InSelectedAssets);
 	SourcesData = InSourcesData;
 
 	// Cache any vars that are used in determining if you can execute any actions.
@@ -77,6 +87,9 @@ TSharedRef<SWidget> FAssetContextMenu::MakeContextMenu(const TArray<FAssetData>&
 		// Add any type-specific context menu options
 		AddAssetTypeMenuOptions(MenuBuilder);
 
+		// Add imported asset context menu options
+		AddImportedAssetMenuOptions(MenuBuilder);
+
 		// Add quick access to common commands.
 		AddCommonMenuOptions(MenuBuilder);
 
@@ -86,16 +99,21 @@ TSharedRef<SWidget> FAssetContextMenu::MakeContextMenu(const TArray<FAssetData>&
 		// Add documentation options
 		AddDocumentationMenuOptions(MenuBuilder);
 
-		// Add source control options
-		AddSourceControlMenuOptions(MenuBuilder);
-
 		// Add collection options
 		AddCollectionMenuOptions(MenuBuilder);
+
+		// Add source control options
+		AddSourceControlMenuOptions(MenuBuilder);
 	}
 	
 	return MenuBuilder.MakeWidget();
 }
 
+void FAssetContextMenu::SetSelectedAssets(const TArray<FAssetData>& InSelectedAssets)
+{
+	SelectedAssets = InSelectedAssets;
+}
+
 void FAssetContextMenu::SetOnFindInAssetTreeRequested(const FOnFindInAssetTreeRequested& InOnFindInAssetTreeRequested)
 {
 	OnFindInAssetTreeRequested = InOnFindInAssetTreeRequested;
@@ -121,210 +139,277 @@ void FAssetContextMenu::SetOnAssetViewRefreshRequested(const FOnAssetViewRefresh
 	OnAssetViewRefreshRequested = InOnAssetViewRefreshRequested;
 }
 
+bool FAssetContextMenu::AddImportedAssetMenuOptions(FMenuBuilder& MenuBuilder)
+{
+	if (AreImportedAssetActionsVisible())
+	{
+		TArray<FString> ResolvedFilePaths;
+		GetSelectedAssetSourceFilePaths(ResolvedFilePaths);
+
+		MenuBuilder.BeginSection("ImportedAssetActions", LOCTEXT("ImportedAssetActionsMenuHeading", "Imported Asset"));
+		{
+			// Reimport
+			MenuBuilder.AddMenuEntry(
+				LOCTEXT("Reimport", "Reimport"),
+				LOCTEXT("ReimportTooltip", "Reimport the selected asset(s) from the source file on disk."),
+				FSlateIcon(FEditorStyle::GetStyleSetName(), "ContentBrowser.AssetActions.ReimportAsset"),
+				FUIAction(
+					FExecuteAction::CreateSP(this, &FAssetContextMenu::ExecuteReimport),
+					FCanExecuteAction()
+					)
+				);
+
+			// Show Source In Explorer
+			MenuBuilder.AddMenuEntry(
+				LOCTEXT("FindSourceFile", "Open Source Location"),
+				LOCTEXT("FindSourceFileTooltip", "Opens the folder containing the source of the selected asset(s)."),
+				FSlateIcon(),
+				FUIAction(
+					FExecuteAction::CreateSP(this, &FAssetContextMenu::ExecuteFindSourceInExplorer, ResolvedFilePaths),
+					FCanExecuteAction::CreateSP(this, &FAssetContextMenu::CanExecuteImportedAssetActions, ResolvedFilePaths)
+					)
+				);
+
+			// Open In External Editor
+			MenuBuilder.AddMenuEntry(
+				LOCTEXT("OpenInExternalEditor", "Open In External Editor"),
+				LOCTEXT("OpenInExternalEditorTooltip", "Open the selected asset(s) in the default external editor."),
+				FSlateIcon(),
+				FUIAction(
+					FExecuteAction::CreateSP(this, &FAssetContextMenu::ExecuteOpenInExternalEditor, ResolvedFilePaths),
+					FCanExecuteAction::CreateSP(this, &FAssetContextMenu::CanExecuteImportedAssetActions, ResolvedFilePaths)
+					)
+				);
+		}
+		MenuBuilder.EndSection();
+
+		return true;
+	}
+	
+
+	return false;
+}
+
 bool FAssetContextMenu::AddCommonMenuOptions(FMenuBuilder& MenuBuilder)
 {
-	MenuBuilder.BeginSection("AssetContextActions", LOCTEXT("AssetActionsMenuHeading", "Asset Actions"));
+	MenuBuilder.BeginSection("CommonAssetActions", LOCTEXT("CommonAssetActionsMenuHeading", "Common"));
 	{
+	// Edit
 	MenuBuilder.AddMenuEntry(
-		LOCTEXT("SyncToAssetTree", "Find in Asset Tree"),
-		LOCTEXT("SyncToAssetTreeTooltip", "Selects the folder in the asset tree containing this asset."),
-		FSlateIcon(),
-		FUIAction(
-			FExecuteAction::CreateSP( this, &FAssetContextMenu::ExecuteSyncToAssetTree ),
-			FCanExecuteAction::CreateSP( this, &FAssetContextMenu::CanExecuteSyncToAssetTree )
-			)
+		LOCTEXT("EditAsset", "Edit..."),
+		LOCTEXT("EditAssetTooltip", "Opens the selected asset(s) for edit."),
+		FSlateIcon(FEditorStyle::GetStyleSetName(), "ContentBrowser.AssetActions.Edit"),
+		FUIAction( FExecuteAction::CreateSP(this, &FAssetContextMenu::ExecuteEditAsset) )
+		);
+	
+	// Rename
+	MenuBuilder.AddMenuEntry(FGenericCommands::Get().Rename, NAME_None,
+		LOCTEXT("Rename", "Rename"),
+		LOCTEXT("RenameTooltip", "Rename the selected asset.")
 		);
 
-	MenuBuilder.AddMenuEntry(
-		LOCTEXT("FindAssetInWorld", "Select Actors Using This Asset"),
-		LOCTEXT("FindAssetInWorldTooltip", "Selects all actors referencing this asset."),
-		FSlateIcon(),
-		FUIAction(
-			FExecuteAction::CreateSP( this, &FAssetContextMenu::ExecuteFindAssetInWorld ),
-			FCanExecuteAction::CreateSP( this, &FAssetContextMenu::CanExecuteFindAssetInWorld )
-			)
+	// Duplicate
+	MenuBuilder.AddMenuEntry(FGenericCommands::Get().Duplicate, NAME_None,
+		LOCTEXT("Duplicate", "Duplicate"),
+		LOCTEXT("DuplicateTooltip", "Create a copy of the selected asset(s).")
 		);
 
+	// Save
 	MenuBuilder.AddMenuEntry(
-		LOCTEXT("Properties", "Details..."),
-		LOCTEXT("PropertiesTooltip", "Opens the details for the selected assets."),
-		FSlateIcon(),
+		LOCTEXT("SaveAsset", "Save"),
+		LOCTEXT("SaveAssetTooltip", "Saves the asset to file."),
+		FSlateIcon(FEditorStyle::GetStyleSetName(), "Level.SaveIcon16x"),
 		FUIAction(
-			FExecuteAction::CreateSP( this, &FAssetContextMenu::ExecuteProperties ),
-			FCanExecuteAction::CreateSP( this, &FAssetContextMenu::CanExecuteProperties )
+			FExecuteAction::CreateSP( this, &FAssetContextMenu::ExecuteSaveAsset ),
+			FCanExecuteAction::CreateSP( this, &FAssetContextMenu::CanExecuteSaveAsset )
 			)
 		);
 
-	MenuBuilder.AddMenuEntry(
-		LOCTEXT("PropertyMatrix", "Property Matrix..."),
-		LOCTEXT("PropertyMatrixTooltip", "Opens the property matrix editor for the selected assets."),
-		FSlateIcon(),
-		FUIAction(
-		FExecuteAction::CreateSP( this, &FAssetContextMenu::ExecutePropertyMatrix ),
-		FCanExecuteAction::CreateSP( this, &FAssetContextMenu::CanExecuteProperties )
-		)
+	// Delete
+	MenuBuilder.AddMenuEntry(FGenericCommands::Get().Delete, NAME_None,
+		LOCTEXT("Delete", "Delete"),
+		LOCTEXT("DeleteTooltip", "Delete the selected assets."),
+		FSlateIcon(FEditorStyle::GetStyleSetName(), "ContentBrowser.AssetActions.Delete")
 		);
 
-	MenuBuilder.AddMenuEntry(
-		LOCTEXT("Duplicate", "Duplicate"),
-		LOCTEXT("DuplicateTooltip", "Create a copy of the selected assets."),
-		FSlateIcon(),
-		FUIAction(
-			FExecuteAction::CreateSP( this, &FAssetContextMenu::ExecuteDuplicate ),
-			FCanExecuteAction::CreateSP( this, &FAssetContextMenu::CanExecuteDuplicate )
-			)
+	// Asset Actions sub-menu
+	MenuBuilder.AddSubMenu(
+		LOCTEXT("AssetActionsSubMenuLabel", "Asset Actions"),
+		LOCTEXT("AssetActionsSubMenuToolTip", "Other asset actions"),
+		FNewMenuDelegate::CreateSP(this, &FAssetContextMenu::MakeAssetActionsSubMenu), false, FSlateIcon(FEditorStyle::GetStyleSetName(), "ContentBrowser.AssetActions")
 		);
+	}
+	MenuBuilder.EndSection();
 
-	MenuBuilder.AddMenuEntry( FGenericCommands::Get().Rename, NAME_None, 
-		LOCTEXT("Rename", "Rename"),
-		LOCTEXT("RenameTooltip", "Rename the selected asset.")
-		);
+	return true;
+}
 
-	MenuBuilder.AddMenuEntry( FGenericCommands::Get().Delete, NAME_None, 
-		LOCTEXT("Delete", "Delete"),
-		LOCTEXT("DeleteTooltip", "Delete the selected assets.")
+void FAssetContextMenu::MakeAssetActionsSubMenu(FMenuBuilder& MenuBuilder)
+{	
+	// Create BP Using This
+	MenuBuilder.AddMenuEntry(
+		LOCTEXT("CreateBlueprintUsing", "Create Blueprint Using This..."),
+		LOCTEXT("CreateBlueprintUsingTooltip", "Create a new Blueprint and add this asset to it"),
+		FSlateIcon(FEditorStyle::GetStyleSetName(), "LevelEditor.CreateClassBlueprint"),
+		FUIAction(
+		FExecuteAction::CreateSP(this, &FAssetContextMenu::ExecuteCreateBlueprintUsing),
+		FCanExecuteAction::CreateSP(this, &FAssetContextMenu::CanExecuteCreateBlueprintUsing)
+			)
 		);
 
-	if ( CanExecuteConsolidate() )
+	// Capture Thumbnail
+	FAssetToolsModule& AssetToolsModule = FModuleManager::LoadModuleChecked<FAssetToolsModule>("AssetTools");
+	if (SelectedAssets.Num() == 1 && AssetToolsModule.Get().AssetUsesGenericThumbnail(SelectedAssets[0]))
 	{
 		MenuBuilder.AddMenuEntry(
-			LOCTEXT("ReplaceReferences", "Replace References"),
-			LOCTEXT("ConsolidateTooltip", "Replace references to the selected assets."),
+			LOCTEXT("CaptureThumbnail", "Capture Thumbnail"),
+			LOCTEXT("CaptureThumbnailTooltip", "Captures a thumbnail from the active viewport."),
 			FSlateIcon(),
 			FUIAction(
-				FExecuteAction::CreateSP( this, &FAssetContextMenu::ExecuteConsolidate )
+			FExecuteAction::CreateSP(this, &FAssetContextMenu::ExecuteCaptureThumbnail),
+			FCanExecuteAction::CreateSP(this, &FAssetContextMenu::CanExecuteCaptureThumbnail)
 				)
 			);
 	}
 
-	MenuBuilder.AddMenuEntry(
-		LOCTEXT("Export", "Export..."),
-		LOCTEXT("ExportTooltip", "Export the selected assets to file."),
-		FSlateIcon(),
-		FUIAction( FExecuteAction::CreateSP( this, &FAssetContextMenu::ExecuteExport ) )
-		);
-
-	if (SelectedAssets.Num() > 1)
+	// Clear Thumbnail
+	if (CanClearCustomThumbnails())
 	{
 		MenuBuilder.AddMenuEntry(
-			LOCTEXT("BulkExport", "Bulk Export..."),
-			LOCTEXT("BulkExportTooltip", "Export the selected assets to file in the selected directory"),
+			LOCTEXT("ClearCustomThumbnail", "Clear Thumbnail"),
+			LOCTEXT("ClearCustomThumbnailTooltip", "Clears all custom thumbnails for selected assets."),
 			FSlateIcon(),
-			FUIAction( FExecuteAction::CreateSP( this, &FAssetContextMenu::ExecuteBulkExport ) )
+			FUIAction( FExecuteAction::CreateSP(this, &FAssetContextMenu::ExecuteClearThumbnail) ) 
 			);
 	}
 
-	MenuBuilder.AddMenuEntry(
-		LOCTEXT("MigrateAsset", "Migrate..."),
-		LOCTEXT("MigrateAssetTooltip", "Copies all selected assets and their dependencies to another game"),
-		FSlateIcon(),
-		FUIAction( FExecuteAction::CreateSP( this, &FAssetContextMenu::ExecuteMigrateAsset ) )
-		);
-
-	MenuBuilder.AddMenuEntry(
-		ContentBrowserUtils::GetExploreFolderText(),
-		LOCTEXT("FindInExplorerTooltip", "Finds this asset on disk"),
-		FSlateIcon(),
-		FUIAction(
-		FExecuteAction::CreateSP( this, &FAssetContextMenu::ExecuteFindInExplorer ),
-		FCanExecuteAction::CreateSP( this, &FAssetContextMenu::CanExecuteFindInExplorer )
-		)
-		);
-
-	MenuBuilder.AddMenuEntry(
-		LOCTEXT("CreateBlueprintUsing", "Create Blueprint Using..."),
-		LOCTEXT("CreateBlueprintUsingTooltip", "Create a new Blueprint and add this asset to it"),
-		FSlateIcon(),
-		FUIAction(
-		FExecuteAction::CreateSP(this, &FAssetContextMenu::ExecuteCreateBlueprintUsing),
-		FCanExecuteAction::CreateSP(this, &FAssetContextMenu::CanExecuteCreateBlueprintUsing)
-		)
-		);
-
-
-	MenuBuilder.AddMenuEntry(
-		LOCTEXT("SaveAsset", "Save"),
-		LOCTEXT("SaveAssetTooltip", "Saves the asset to file."),
-		FSlateIcon(),
-		FUIAction(
-			FExecuteAction::CreateSP( this, &FAssetContextMenu::ExecuteSaveAsset ),
-			FCanExecuteAction::CreateSP( this, &FAssetContextMenu::CanExecuteSaveAsset )
-			)
-		);
-
-	if (GetDefault<UEditorExperimentalSettings>()->bContextMenuChunkAssignments)
+	// FIND ACTIONS
+	MenuBuilder.BeginSection("AssetContextFindActions", LOCTEXT("AssetContextFindActionsMenuHeading", "Find"));
 	{
+		// Select Actors Using This Asset
 		MenuBuilder.AddMenuEntry(
-			LOCTEXT("AssignAssetChunk","Assign to Chunk..."),
-			LOCTEXT("AssignAssetChunkTooltip","Assign this asset to a specific Chunk"),
-			FSlateIcon(),
+			LOCTEXT("FindAssetInWorld", "Select Actors Using This Asset"),
+			LOCTEXT("FindAssetInWorldTooltip", "Selects all actors referencing this asset."),
+			FSlateIcon(FEditorStyle::GetStyleSetName(), "ContentBrowser.AssetActions.FindAssetInWorld"),
 			FUIAction(
-				FExecuteAction::CreateSP(this,&FAssetContextMenu::ExecuteAssignChunkID)
-			)
-		);
+				FExecuteAction::CreateSP(this, &FAssetContextMenu::ExecuteFindAssetInWorld),
+				FCanExecuteAction::CreateSP(this, &FAssetContextMenu::CanExecuteFindAssetInWorld)
+				)
+			);
 
-		MenuBuilder.AddSubMenu(
-			LOCTEXT("RemoveAssetFromChunk","Remove from Chunk..."),
-			LOCTEXT("RemoveAssetFromChunkTooltip","Removed an asset from a Chunk it's assigned to."),
-			FNewMenuDelegate::CreateRaw(this,&FAssetContextMenu::MakeChunkIDListMenu)
-		);
+		// Open Containing Folder
+		MenuBuilder.AddMenuEntry(
+			LOCTEXT("SyncToAssetTree", "Select Containing Folder"),
+			LOCTEXT("SyncToAssetTreeTooltip", "Selects the folder in the path view that contains this asset."),
+			FSlateIcon(FEditorStyle::GetStyleSetName(), "MainFrame.OpenProject"),
+			FUIAction(
+				FExecuteAction::CreateSP(this, &FAssetContextMenu::ExecuteSyncToAssetTree),
+				FCanExecuteAction::CreateSP(this, &FAssetContextMenu::CanExecuteSyncToAssetTree)
+				)
+			);
 
+		// Find in Explorer
 		MenuBuilder.AddMenuEntry(
-			LOCTEXT("RemoveAllChunkAssignments","Remove from all Chunks"),
-			LOCTEXT("RemoveAllChunkAssignmentsTooltip","Removed an asset from all Chunks it's assigned to."),
+			ContentBrowserUtils::GetExploreFolderText(),
+			LOCTEXT("FindInExplorerTooltip", "Finds this asset on disk"),
 			FSlateIcon(),
 			FUIAction(
-				FExecuteAction::CreateSP(this,&FAssetContextMenu::ExecuteRemoveAllChunkID)
-			)
-		);
+				FExecuteAction::CreateSP( this, &FAssetContextMenu::ExecuteFindInExplorer ),
+				FCanExecuteAction::CreateSP( this, &FAssetContextMenu::CanExecuteFindInExplorer )
+				)
+			);
 	}
+	MenuBuilder.EndSection();
 
-	if (CanExecuteDiffSelected())
+	// MOVE ACTIONS
+	MenuBuilder.BeginSection("AssetContextMoveActions", LOCTEXT("AssetContextMoveActionsMenuHeading", "Move"));
 	{
+		// Export
 		MenuBuilder.AddMenuEntry(
-			LOCTEXT("DiffSelected", "Diff Selected"),
-			LOCTEXT("DiffSelectedTooltip", "Diff the two assets that you have selected."),
+			LOCTEXT("Export", "Export..."),
+			LOCTEXT("ExportTooltip", "Export the selected assets to file."),
 			FSlateIcon(),
-			FUIAction(
-				FExecuteAction::CreateSP(this, &FAssetContextMenu::ExecuteDiffSelected)
-			)
-		);
-	}
+			FUIAction( FExecuteAction::CreateSP( this, &FAssetContextMenu::ExecuteExport ) )
+			);
 
-	FAssetToolsModule& AssetToolsModule = FModuleManager::LoadModuleChecked<FAssetToolsModule>("AssetTools");
-	if ( SelectedAssets.Num() == 1 && AssetToolsModule.Get().AssetUsesGenericThumbnail(SelectedAssets[0]) )
-	{
+		// Bulk Export
+		if (SelectedAssets.Num() > 1)
+		{
+			MenuBuilder.AddMenuEntry(
+				LOCTEXT("BulkExport", "Bulk Export..."),
+				LOCTEXT("BulkExportTooltip", "Export the selected assets to file in the selected directory"),
+				FSlateIcon(),
+				FUIAction( FExecuteAction::CreateSP( this, &FAssetContextMenu::ExecuteBulkExport ) )
+				);
+		}
+
+		// Migrate
 		MenuBuilder.AddMenuEntry(
-			LOCTEXT("CaptureThumbnail", "Capture Thumbnail"),
-			LOCTEXT("CaptureThumbnailTooltip", "Captures a thumbnail from the active viewport."),
+			LOCTEXT("MigrateAsset", "Migrate..."),
+			LOCTEXT("MigrateAssetTooltip", "Copies all selected assets and their dependencies to another game"),
 			FSlateIcon(),
-			FUIAction(
-				FExecuteAction::CreateSP( this, &FAssetContextMenu::ExecuteCaptureThumbnail ),
-				FCanExecuteAction::CreateSP( this, &FAssetContextMenu::CanExecuteCaptureThumbnail )
-				)
+			FUIAction( FExecuteAction::CreateSP( this, &FAssetContextMenu::ExecuteMigrateAsset ) )
 			);
 	}
+	MenuBuilder.EndSection();
 
-	if ( CanClearCustomThumbnails() )
+	// ADVANCED ACTIONS
+	MenuBuilder.BeginSection("AssetContextAdvancedActions", LOCTEXT("AssetContextAdvancedActionsMenuHeading", "Advanced"));
 	{
+		// Replace References
+		if (CanExecuteConsolidate())
+		{
+			MenuBuilder.AddMenuEntry(
+				LOCTEXT("ReplaceReferences", "Replace References"),
+				LOCTEXT("ConsolidateTooltip", "Replace references to the selected assets."),
+				FSlateIcon(),
+				FUIAction(
+				FExecuteAction::CreateSP(this, &FAssetContextMenu::ExecuteConsolidate)
+				)
+				);
+		}
+
+		// Property Matrix
 		MenuBuilder.AddMenuEntry(
-			LOCTEXT("ClearCustomThumbnail", "Clear Thumbnail"),
-			LOCTEXT("ClearCustomThumbnailTooltip", "Clears all custom thumbnails for selected assets."),
+			LOCTEXT("PropertyMatrix", "Property Matrix..."),
+			LOCTEXT("PropertyMatrixTooltip", "Opens the property matrix editor for the selected assets."),
 			FSlateIcon(),
 			FUIAction(
-				FExecuteAction::CreateSP( this, &FAssetContextMenu::ExecuteClearThumbnail )
-				)
+			FExecuteAction::CreateSP(this, &FAssetContextMenu::ExecutePropertyMatrix),
+			FCanExecuteAction::CreateSP(this, &FAssetContextMenu::CanExecuteProperties)
+			)
 			);
-	}
+
+		// Chunk actions
+		if (GetDefault<UEditorExperimentalSettings>()->bContextMenuChunkAssignments)
+		{
+			MenuBuilder.AddMenuEntry(
+				LOCTEXT("AssignAssetChunk", "Assign to Chunk..."),
+				LOCTEXT("AssignAssetChunkTooltip", "Assign this asset to a specific Chunk"),
+				FSlateIcon(),
+				FUIAction( FExecuteAction::CreateSP(this, &FAssetContextMenu::ExecuteAssignChunkID) )
+				);
+
+			MenuBuilder.AddSubMenu(
+				LOCTEXT("RemoveAssetFromChunk", "Remove from Chunk..."),
+				LOCTEXT("RemoveAssetFromChunkTooltip", "Removed an asset from a Chunk it's assigned to."),
+				FNewMenuDelegate::CreateRaw(this, &FAssetContextMenu::MakeChunkIDListMenu)
+				);
+
+			MenuBuilder.AddMenuEntry(
+				LOCTEXT("RemoveAllChunkAssignments", "Remove from all Chunks"),
+				LOCTEXT("RemoveAllChunkAssignmentsTooltip", "Removed an asset from all Chunks it's assigned to."),
+				FSlateIcon(),
+				FUIAction( FExecuteAction::CreateSP(this, &FAssetContextMenu::ExecuteRemoveAllChunkID) )
+				);
+		}
 	}
 	MenuBuilder.EndSection();
-
-	return true;
 }
 
 bool FAssetContextMenu::AddReferenceMenuOptions(FMenuBuilder& MenuBuilder)
 {
-	MenuBuilder.BeginSection("AssetContextReferences", LOCTEXT("ReferencesMenuHeading", "References"));
+	MenuBuilder.BeginSection("AssetContextReferences"/*, LOCTEXT("ReferencesMenuHeading", "References")*/);
 	{
 		MenuBuilder.AddMenuEntry(
 			LOCTEXT("CopyReference", "Copy Reference"),
@@ -385,12 +470,12 @@ bool FAssetContextMenu::AddDocumentationMenuOptions(FMenuBuilder& MenuBuilder)
 
 				const FString CodeFileName = FPaths::GetCleanFilename( *ClassHeaderPath );
 
-				MenuBuilder.BeginSection( "AssetCode", LOCTEXT("AssetCodeHeading", "C++") );
+				MenuBuilder.BeginSection( "AssetCode"/*, LOCTEXT("AssetCodeHeading", "C++")*/ );
 				{
 					MenuBuilder.AddMenuEntry(
 						FText::Format( LOCTEXT("GoToCodeForAsset", "Open {0}"), FText::FromString( CodeFileName ) ),
 						FText::Format( LOCTEXT("GoToCodeForAsset_ToolTip", "Opens the header file for this asset ({0}) in a code editing program"), FText::FromString( CodeFileName ) ),
-						FSlateIcon(),
+						FSlateIcon(FEditorStyle::GetStyleSetName(), "ContentBrowser.AssetActions.GoToCodeForAsset"),
 						FUIAction( FExecuteAction::CreateSP( this, &FAssetContextMenu::ExecuteGoToCodeForAsset, SelectedClass ) )
 						);
 				}
@@ -403,7 +488,7 @@ bool FAssetContextMenu::AddDocumentationMenuOptions(FMenuBuilder& MenuBuilder)
 		{
 			bAddedOption = true;
 
-			MenuBuilder.BeginSection( "AssetDocumentation", LOCTEXT("AseetDocsHeading", "Documentation") );
+			MenuBuilder.BeginSection( "AssetDocumentation"/*, LOCTEXT("AseetDocsHeading", "Documentation")*/ );
 			{
 					if (bIsBlueprint)
 					{
@@ -497,11 +582,51 @@ bool FAssetContextMenu::AddAssetTypeMenuOptions(FMenuBuilder& MenuBuilder)
 
 bool FAssetContextMenu::AddSourceControlMenuOptions(FMenuBuilder& MenuBuilder)
 {
-	MenuBuilder.BeginSection("AssetContextSourceControl", LOCTEXT("AssetSCCOptionsMenuHeading", "Source Control"));
-
+	MenuBuilder.AddMenuSeparator();
+	
 	if ( ISourceControlModule::Get().IsEnabled() )
 	{
-		if( CanExecuteSCCSync() )
+		// SCC sub menu
+		MenuBuilder.AddSubMenu(
+			LOCTEXT("SourceControlSubMenuLabel", "Source Control"),
+			LOCTEXT("SourceControlSubMenuToolTip", "Source control actions."),
+			FNewMenuDelegate::CreateSP(this, &FAssetContextMenu::FillSourceControlSubMenu),
+			false,
+			FSlateIcon(FEditorStyle::GetStyleSetName(), "SourceControl.StatusIcon.On")
+			);
+	}
+	else
+	{
+		MenuBuilder.AddMenuEntry(
+			LOCTEXT("SCCConnectToSourceControl", "Connect To Source Control"),
+			LOCTEXT("SCCConnectToSourceControlTooltip", "Connect to source control to allow source control operations to be performed on content and levels."),
+			FSlateIcon(FEditorStyle::GetStyleSetName(), "SourceControl.StatusIcon.Unknown"),
+			FUIAction(
+				FExecuteAction::CreateSP( this, &FAssetContextMenu::ExecuteEnableSourceControl ),
+				FCanExecuteAction()
+				)
+			);		
+	}
+
+	// Diff selected
+	if (CanExecuteDiffSelected())
+	{
+		MenuBuilder.AddMenuEntry(
+			LOCTEXT("DiffSelected", "Diff Selected"),
+			LOCTEXT("DiffSelectedTooltip", "Diff the two assets that you have selected."),
+			FSlateIcon(),
+			FUIAction(
+				FExecuteAction::CreateSP(this, &FAssetContextMenu::ExecuteDiffSelected)
+			)
+		);
+	}
+
+	return true;
+}
+
+void FAssetContextMenu::FillSourceControlSubMenu(FMenuBuilder& MenuBuilder)
+{
+	if( CanExecuteSCCSync() )
 		{
 			MenuBuilder.AddMenuEntry(
 				LOCTEXT("SCCSync", "Sync"),
@@ -519,7 +644,7 @@ bool FAssetContextMenu::AddSourceControlMenuOptions(FMenuBuilder& MenuBuilder)
 			MenuBuilder.AddMenuEntry(
 				LOCTEXT("SCCCheckOut", "Check Out"),
 				LOCTEXT("SCCCheckOutTooltip", "Checks out the selected asset from source control."),
-				FSlateIcon(),
+				FSlateIcon(FEditorStyle::GetStyleSetName(), "SourceControl.Actions.CheckOut"),
 				FUIAction(
 					FExecuteAction::CreateSP( this, &FAssetContextMenu::ExecuteSCCCheckOut ),
 					FCanExecuteAction::CreateSP( this, &FAssetContextMenu::CanExecuteSCCCheckOut )
@@ -532,7 +657,7 @@ bool FAssetContextMenu::AddSourceControlMenuOptions(FMenuBuilder& MenuBuilder)
 			MenuBuilder.AddMenuEntry(
 				LOCTEXT("SCCOpenForAdd", "Mark For Add"),
 				LOCTEXT("SCCOpenForAddTooltip", "Adds the selected asset to source control."),
-				FSlateIcon(),
+				FSlateIcon(FEditorStyle::GetStyleSetName(), "SourceControl.Actions.Add"),
 				FUIAction(
 					FExecuteAction::CreateSP( this, &FAssetContextMenu::ExecuteSCCOpenForAdd ),
 					FCanExecuteAction::CreateSP( this, &FAssetContextMenu::CanExecuteSCCOpenForAdd )
@@ -556,7 +681,7 @@ bool FAssetContextMenu::AddSourceControlMenuOptions(FMenuBuilder& MenuBuilder)
 		MenuBuilder.AddMenuEntry(
 			LOCTEXT("SCCRefresh", "Refresh"),
 			LOCTEXT("SCCRefreshTooltip", "Updates the source control status of the asset."),
-			FSlateIcon(),
+			FSlateIcon(FEditorStyle::GetStyleSetName(), "SourceControl.Actions.Refresh"),
 			FUIAction(
 				FExecuteAction::CreateSP( this, &FAssetContextMenu::ExecuteSCCRefresh ),
 				FCanExecuteAction::CreateSP( this, &FAssetContextMenu::CanExecuteSCCRefresh )
@@ -598,24 +723,6 @@ bool FAssetContextMenu::AddSourceControlMenuOptions(FMenuBuilder& MenuBuilder)
 					)
 				);
 		}
-	}
-	else
-	{
-		MenuBuilder.AddMenuEntry(
-			LOCTEXT("SCCConnectToSourceControl", "Connect To Source Control"),
-			LOCTEXT("SCCConnectToSourceControlTooltip", "Connect to source control to allow source control operations to be performed on content and levels."),
-			FSlateIcon(),
-			FUIAction(
-				FExecuteAction::CreateSP( this, &FAssetContextMenu::ExecuteEnableSourceControl ),
-				FCanExecuteAction()
-				)
-			);		
-
-	}
-
-	MenuBuilder.EndSection();
-
-	return true;
 }
 
 bool FAssetContextMenu::AddCollectionMenuOptions(FMenuBuilder& MenuBuilder)
@@ -643,6 +750,112 @@ bool FAssetContextMenu::AddCollectionMenuOptions(FMenuBuilder& MenuBuilder)
 	return false;
 }
 
+bool FAssetContextMenu::AreImportedAssetActionsVisible() const
+{
+	FAssetToolsModule& AssetToolsModule = FModuleManager::Get().LoadModuleChecked<FAssetToolsModule>("AssetTools");
+	
+	// Check that all of the selected assets are imported
+	for (auto& SelectedAsset : SelectedAssets)
+	{
+		auto AssetClass = SelectedAsset.GetClass();
+		if (AssetClass)
+		{
+			auto AssetTypeActions = AssetToolsModule.Get().GetAssetTypeActionsForClass(AssetClass).Pin();
+			if (!AssetTypeActions.IsValid() || !AssetTypeActions->IsImportedAsset())
+			{
+				return false;
+			}
+		}
+	}
+
+	return true;
+}
+
+bool FAssetContextMenu::CanExecuteImportedAssetActions(const TArray<FString> ResolvedFilePaths) const
+{
+	// Verify that all the file paths are legitimate
+	for (const auto& SourceFilePath : ResolvedFilePaths)
+	{
+		if (!SourceFilePath.Len() || IFileManager::Get().FileSize(*SourceFilePath) == INDEX_NONE)
+		{
+			return false;
+		}
+	}
+
+	return true;
+}
+
+void FAssetContextMenu::ExecuteReimport()
+{
+	// Reimport all selected assets
+	for (auto& SelectedAsset : SelectedAssets)
+	{
+		const auto Asset = SelectedAsset.GetAsset();
+		if (Asset)
+		{
+			FReimportManager::Instance()->Reimport(Asset, /*bAskForNewFileIfMissing=*/true);
+		}
+	}
+}
+
+void FAssetContextMenu::ExecuteFindSourceInExplorer(const TArray<FString> ResolvedFilePaths)
+{
+	// Open all files in the explorer
+	for (const auto& SourceFilePath : ResolvedFilePaths)
+	{
+		FPlatformProcess::ExploreFolder(*FPaths::GetPath(SourceFilePath));
+	}
+}
+
+void FAssetContextMenu::ExecuteOpenInExternalEditor(const TArray<FString> ResolvedFilePaths)
+{
+	// Open all files in their respective editor
+	for (const auto& SourceFilePath : ResolvedFilePaths)
+	{
+		FPlatformProcess::LaunchFileInDefaultExternalApplication(*SourceFilePath, NULL, ELaunchVerb::Edit);
+	}
+}
+
+void FAssetContextMenu::GetSelectedAssetsByClass(TMap<UClass*, TArray<UObject*> >& OutSelectedAssetsByClass) const
+{
+	// Sort all selected assets by class
+	for (const auto& SelectedAsset : SelectedAssets)
+	{
+		auto Asset = SelectedAsset.GetAsset();
+		auto AssetClass = Asset->GetClass();
+
+		if ( !OutSelectedAssetsByClass.Contains(AssetClass) )
+		{
+			OutSelectedAssetsByClass.Add(AssetClass);
+		}
+		
+		OutSelectedAssetsByClass[AssetClass].Add(Asset);
+	}
+}
+
+void FAssetContextMenu::GetSelectedAssetSourceFilePaths(TArray<FString>& OutFilePaths) const
+{
+	OutFilePaths.Empty();
+	
+	TMap<UClass*, TArray<UObject*> > SelectedAssetsByClass;
+	GetSelectedAssetsByClass(SelectedAssetsByClass);
+	FAssetToolsModule& AssetToolsModule = FModuleManager::Get().LoadModuleChecked<FAssetToolsModule>("AssetTools");
+
+	// Get the source file paths for the assets of each type
+	for (const auto& AssetsByClassPair : SelectedAssetsByClass)
+	{
+		const auto AssetTypeActions = AssetToolsModule.Get().GetAssetTypeActionsForClass(AssetsByClassPair.Key);
+		if (AssetTypeActions.IsValid())
+		{
+			const auto& TypeAssets = AssetsByClassPair.Value;
+			TArray<FString> AssetSourcePaths;
+			AssetTypeActions.Pin()->GetResolvedSourceFilePaths(TypeAssets, AssetSourcePaths);
+
+			OutFilePaths.Append(AssetSourcePaths);
+		}
+	}
+}
+
 void FAssetContextMenu::ExecuteSyncToAssetTree()
 {
 	OnFindInAssetTreeRequested.ExecuteIfBound(SelectedAssets);
@@ -823,7 +1036,7 @@ void FAssetContextMenu::ExecuteFindAssetInWorld()
 	}
 }
 
-void FAssetContextMenu::ExecuteProperties()
+void FAssetContextMenu::ExecutePropertyMatrix()
 {
 	TArray<UObject*> ObjectsForPropertiesMenu;
 	const bool SkipRedirectors = true;
@@ -831,20 +1044,21 @@ void FAssetContextMenu::ExecuteProperties()
 
 	if ( ObjectsForPropertiesMenu.Num() > 0 )
 	{
-		FAssetEditorManager::Get().OpenEditorForAssets( ObjectsForPropertiesMenu );
+		FPropertyEditorModule& PropertyEditorModule = FModuleManager::LoadModuleChecked<FPropertyEditorModule>( "PropertyEditor" );
+		PropertyEditorModule.CreatePropertyEditorToolkit( EToolkitMode::Standalone, TSharedPtr<IToolkitHost>(), ObjectsForPropertiesMenu );
 	}
 }
 
-void FAssetContextMenu::ExecutePropertyMatrix()
+void FAssetContextMenu::ExecuteEditAsset()
 {
-	TArray<UObject*> ObjectsForPropertiesMenu;
-	const bool SkipRedirectors = true;
-	GetSelectedAssets(ObjectsForPropertiesMenu, SkipRedirectors);
+	TMap<UClass*, TArray<UObject*> > SelectedAssetsByClass;
+	GetSelectedAssetsByClass(SelectedAssetsByClass);
 
-	if ( ObjectsForPropertiesMenu.Num() > 0 )
+	// Open 
+	for (const auto& AssetsByClassPair : SelectedAssetsByClass)
 	{
-		FPropertyEditorModule& PropertyEditorModule = FModuleManager::LoadModuleChecked<FPropertyEditorModule>( "PropertyEditor" );
-		PropertyEditorModule.CreatePropertyEditorToolkit( EToolkitMode::Standalone, TSharedPtr<IToolkitHost>(), ObjectsForPropertiesMenu );
+		const auto& TypeAssets = AssetsByClassPair.Value;
+		FAssetEditorManager::Get().OpenEditorForAssets(TypeAssets);
 	}
 }
 
diff --git a/Engine/Source/Editor/ContentBrowser/Private/AssetContextMenu.h b/Engine/Source/Editor/ContentBrowser/Private/AssetContextMenu.h
--- a/Engine/Source/Editor/ContentBrowser/Private/AssetContextMenu.h
+++ b/Engine/Source/Editor/ContentBrowser/Private/AssetContextMenu.h
@@ -8,9 +8,15 @@ public:
 	/** Constructor */
 	FAssetContextMenu(const TWeakPtr<SAssetView>& InAssetView);
 
+	/** Binds the commands used by the asset view context menu to the content browser command list */
+	void BindCommands(TSharedPtr< FUICommandList >& Commands);
+
 	/** Makes the context menu widget */
 	TSharedRef<SWidget> MakeContextMenu(const TArray<FAssetData>& SelectedAssets, const FSourcesData& InSourcesData, TSharedPtr< FUICommandList > InCommandList);
 
+	/** Updates the list of currently selected assets to those passed in */
+	void SetSelectedAssets(const TArray<FAssetData>& InSelectedAssets);
+
 	/** Delegate for when the context menu requests a rename */
 	void SetOnFindInAssetTreeRequested(const FOnFindInAssetTreeRequested& InOnFindInAssetTreeRequested);
 
@@ -43,21 +49,58 @@ public:
 	void ExecuteDelete();
 
 private:
+	/** Helper to load selected assets and sort them by UClass */
+	void GetSelectedAssetsByClass(TMap<UClass*, TArray<UObject*> >& OutSelectedAssetsByClass) const;
+
+	/** Helper to collect resolved filepaths for all selected assets */
+	void GetSelectedAssetSourceFilePaths(TArray<FString>& OutFilePaths) const;
+
+	/** Handler to check to see if a imported asset actions should be visible in the menu */
+	bool AreImportedAssetActionsVisible() const;
+
+	/** Handler to check to see if imported asset actions are allowed */
+	bool CanExecuteImportedAssetActions(const TArray<FString> ResolvedFilePaths) const;
+
+	/** Handler for Reimport */
+	void ExecuteReimport();
+
+	/** Handler for FindInExplorer */
+	void ExecuteFindSourceInExplorer(const TArray<FString> ResolvedFilePaths);
+
+	/** Handler for OpenInExternalEditor */
+	void ExecuteOpenInExternalEditor(const TArray<FString> ResolvedFilePaths);
+
+	/** Handler to check to see if a duplicate command is allowed */
+	bool CanExecuteDuplicate() const;
+
+	/** Handler for Duplicate */
+	void ExecuteDuplicate();
+
+private:
+	/** Adds asset type-specific menu options to a menu builder. Returns true if any options were added. */
+	bool AddAssetTypeMenuOptions(FMenuBuilder& MenuBuilder);
+
+	/** Adds asset type-specific menu options to a menu builder. Returns true if any options were added. */
+	bool AddImportedAssetMenuOptions(FMenuBuilder& MenuBuilder);
+	
 	/** Adds common menu options to a menu builder. Returns true if any options were added. */
 	bool AddCommonMenuOptions(FMenuBuilder& MenuBuilder);
 
+	/** Adds Asset Actions sub-menu to a menu builder. */
+	void MakeAssetActionsSubMenu(FMenuBuilder& MenuBuilder);
+
 	/** Adds asset reference menu options to a menu builder. Returns true if any options were added. */
 	bool AddReferenceMenuOptions(FMenuBuilder& MenuBuilder);
 
 	/** Adds asset documentation menu options to a menu builder. Returns true if any options were added. */
 	bool AddDocumentationMenuOptions(FMenuBuilder& MenuBuilder);
-
-	/** Adds asset type-specific menu options to a menu builder. Returns true if any options were added. */
-	bool AddAssetTypeMenuOptions(FMenuBuilder& MenuBuilder);
-
-	/** Adds source control menu options to a menu builder. Returns true if any options were added. */
+	
+	/** Adds source control menu options to a menu builder. */
 	bool AddSourceControlMenuOptions(FMenuBuilder& MenuBuilder);
 
+	/** Fills the source control sub-menu */
+	void FillSourceControlSubMenu(FMenuBuilder& MenuBuilder);
+
 	/** Adds menu options related to working with collections */
 	bool AddCollectionMenuOptions(FMenuBuilder& MenuBuilder);
 
@@ -76,21 +119,18 @@ private:
 	/** Handler for when "Find in World" is selected */
 	void ExecuteFindAssetInWorld();
 
-	/** Handler for when "Properties" is selected */
-	void ExecuteProperties();
-
 	/** Handler for when "Property Matrix..." is selected */
 	void ExecutePropertyMatrix();
 
+	/** Handler for when "Edit Asset" is selected */
+	void ExecuteEditAsset();
+
 	/** Handler for when "Save Asset" is selected */
 	void ExecuteSaveAsset();
 
 	/** Handler for when "Diff Selected" is selected */
 	void ExecuteDiffSelected() const;
 
-	/** Handler for Duplicate */
-	void ExecuteDuplicate();
-
 	/** Handler for confirmation of folder deletion */
 	FReply ExecuteDeleteFolderConfirmed();
 
@@ -184,9 +224,6 @@ private:
 	/** Handler to check to see if a property matrix command is allowed */
 	bool CanExecutePropertyMatrix() const;
 
-	/** Handler to check to see if a duplicate command is allowed */
-	bool CanExecuteDuplicate() const;
-
 	/** Handler to check to see if a "Remove from collection" command is allowed */
 	bool CanExecuteRemoveFromCollection() const;
 
diff --git a/Engine/Source/Editor/ContentBrowser/Private/PathContextMenu.cpp b/Engine/Source/Editor/ContentBrowser/Private/PathContextMenu.cpp
--- a/Engine/Source/Editor/ContentBrowser/Private/PathContextMenu.cpp
+++ b/Engine/Source/Editor/ContentBrowser/Private/PathContextMenu.cpp
@@ -81,8 +81,8 @@ void FPathContextMenu::MakePathViewContextMenu(FMenuBuilder& MenuBuilder)
 		{
 			// New Asset (submenu)
 			MenuBuilder.AddSubMenu(
-				LOCTEXT("NewAsset", "New Asset"),
-				LOCTEXT("NewAssetTooltip", "Create an asset."),
+				LOCTEXT("CreateAsset", "Create Asset"),
+				LOCTEXT("CreateAssetTooltip", "Create a new asset."),
 				FNewMenuDelegate::CreateRaw( this, &FPathContextMenu::MakeNewAssetSubMenu )
 				);
 
diff --git a/Engine/Source/Editor/ContentBrowser/Private/SContentBrowser.cpp b/Engine/Source/Editor/ContentBrowser/Private/SContentBrowser.cpp
--- a/Engine/Source/Editor/ContentBrowser/Private/SContentBrowser.cpp
+++ b/Engine/Source/Editor/ContentBrowser/Private/SContentBrowser.cpp
@@ -127,7 +127,7 @@ void SContentBrowser::Construct( const FArguments& InArgs, const FName& InInstan
 								[
 									SNew( STextBlock )
 									.TextStyle( FEditorStyle::Get(), "ContentBrowser.TopBar.Font" )
-									.Text( LOCTEXT( "NewButton", "New" ) )
+									.Text( LOCTEXT( "NewButton", "Create" ) )
 								]
 							]
 						]
@@ -184,8 +184,27 @@ void SContentBrowser::Construct( const FArguments& InArgs, const FName& InInstan
 							.OnClicked( this, &SContentBrowser::OnSaveClicked )
 							.AddMetaData<FTagMetaData>(FTagMetaData(TEXT("ContentBrowserSaveDirtyPackages")))
 							[
-								SNew( SImage )
-								.Image( FEditorStyle::GetBrush( "ContentBrowser.SaveDirtyPackages" ) )
+								SNew( SHorizontalBox )
+
+								// Save All Icon
+								+ SHorizontalBox::Slot()
+								.AutoWidth()
+								.VAlign( VAlign_Center )
+								[
+									SNew( SImage )
+									.Image( FEditorStyle::GetBrush( "ContentBrowser.SaveDirtyPackages" ) )
+								]
+
+								// Save All Text
+								+ SHorizontalBox::Slot()
+								.AutoWidth()
+								.VAlign(VAlign_Center)
+								.Padding(0,0,2,0)
+								[
+									SNew( STextBlock )
+									.TextStyle( FEditorStyle::Get(), "ContentBrowser.TopBar.Font" )
+									.Text( LOCTEXT( "SaveAll", "Save All" ) )
+								]
 							]
 						]
 					]
@@ -530,6 +549,7 @@ void SContentBrowser::Construct( const FArguments& InArgs, const FName& InInstan
 	];
 
 	AssetContextMenu = MakeShareable(new FAssetContextMenu(AssetViewPtr));
+	AssetContextMenu->BindCommands(Commands);
 	AssetContextMenu->SetOnFindInAssetTreeRequested( FOnFindInAssetTreeRequested::CreateSP(this, &SContentBrowser::OnFindInAssetTreeRequested) );
 	AssetContextMenu->SetOnRenameRequested( FAssetContextMenu::FOnRenameRequested::CreateSP(this, &SContentBrowser::OnRenameRequested) );
 	AssetContextMenu->SetOnRenameFolderRequested( FAssetContextMenu::FOnRenameFolderRequested::CreateSP(this, &SContentBrowser::OnRenameFolderRequested) );
@@ -1317,9 +1337,11 @@ void SContentBrowser::OnAssetSelectionChanged(const FAssetData& SelectedAsset)
 	// Notify 'asset selection changed' delegate
 	FContentBrowserModule& ContentBrowserModule = FModuleManager::GetModuleChecked<FContentBrowserModule>( TEXT("ContentBrowser") );
 	FContentBrowserModule::FOnAssetSelectionChanged& AssetSelectionChangedDelegate = ContentBrowserModule.GetOnAssetSelectionChanged();
+	
+	const TArray<FAssetData>& SelectedAssets = AssetViewPtr->GetSelectedAssets();
+	AssetContextMenu->SetSelectedAssets(SelectedAssets);
 	if(AssetSelectionChangedDelegate.IsBound())
 	{
-		const TArray<FAssetData>& SelectedAssets = AssetViewPtr->GetSelectedAssets();
 		AssetSelectionChangedDelegate.Broadcast(SelectedAssets, bIsPrimaryBrowser);
 	}
 }
diff --git a/Engine/Source/Editor/EnvironmentQueryEditor/Private/AssetTypeActions_EnvironmentQuery.cpp b/Engine/Source/Editor/EnvironmentQueryEditor/Private/AssetTypeActions_EnvironmentQuery.cpp
--- a/Engine/Source/Editor/EnvironmentQueryEditor/Private/AssetTypeActions_EnvironmentQuery.cpp
+++ b/Engine/Source/Editor/EnvironmentQueryEditor/Private/AssetTypeActions_EnvironmentQuery.cpp
@@ -11,21 +11,6 @@
 #include "AssetTypeActions_EnvironmentQuery.h"
 #define LOCTEXT_NAMESPACE "AssetTypeActions"
 
-void FAssetTypeActions_EnvironmentQuery::GetActions( const TArray<UObject*>& InObjects, FMenuBuilder& MenuBuilder )
-{
-	auto Scripts = GetTypedWeakObjectPtrs<UEnvQuery>(InObjects);
-
-	MenuBuilder.AddMenuEntry(
-		LOCTEXT("EnvironmentQuery_Edit", "Edit"),
-		LOCTEXT("EnvironmentQuery_EditTooltip", "Opens the selected Environment Query in editor."),
-		FSlateIcon(),
-		FUIAction(
-		FExecuteAction::CreateSP( this, &FAssetTypeActions_EnvironmentQuery::ExecuteEdit, Scripts ),
-		FCanExecuteAction()
-		)
-	);
-}
-
 void FAssetTypeActions_EnvironmentQuery::OpenAssetEditor( const TArray<UObject*>& InObjects, TSharedPtr<IToolkitHost> EditWithinLevelEditor )
 {
 	EToolkitMode::Type Mode = EditWithinLevelEditor.IsValid() ? EToolkitMode::WorldCentric : EToolkitMode::Standalone;
@@ -41,18 +26,6 @@ void FAssetTypeActions_EnvironmentQuery::OpenAssetEditor( const TArray<UObject*>
 	}
 }
 
-void FAssetTypeActions_EnvironmentQuery::ExecuteEdit(TArray<TWeakObjectPtr<UEnvQuery>> Objects)
-{
-	for (auto ObjIt = Objects.CreateConstIterator(); ObjIt; ++ObjIt)
-	{
-		auto Object = (*ObjIt).Get();
-		if ( Object )
-		{
-			FAssetEditorManager::Get().OpenEditorForAsset(Object);
-		}
-	}
-}
-
 UClass* FAssetTypeActions_EnvironmentQuery::GetSupportedClass() const
 { 
 	return UEnvQuery::StaticClass(); 
diff --git a/Engine/Source/Editor/EnvironmentQueryEditor/Private/AssetTypeActions_EnvironmentQuery.h b/Engine/Source/Editor/EnvironmentQueryEditor/Private/AssetTypeActions_EnvironmentQuery.h
--- a/Engine/Source/Editor/EnvironmentQueryEditor/Private/AssetTypeActions_EnvironmentQuery.h
+++ b/Engine/Source/Editor/EnvironmentQueryEditor/Private/AssetTypeActions_EnvironmentQuery.h
@@ -10,12 +10,6 @@ public:
 	virtual FText GetName() const override { return NSLOCTEXT("AssetTypeActions", "AssetTypeActions_EnvironmentQuery", "Environment Query"); }
 	virtual FColor GetTypeColor() const override { return FColor(129,50,255); }
 	virtual UClass* GetSupportedClass() const override;
-	virtual bool HasActions ( const TArray<UObject*>& InObjects ) const override { return true; }
-	virtual void GetActions( const TArray<UObject*>& InObjects, FMenuBuilder& MenuBuilder ) override;
 	virtual void OpenAssetEditor( const TArray<UObject*>& InObjects, TSharedPtr<class IToolkitHost> EditWithinLevelEditor = TSharedPtr<IToolkitHost>() ) override;
 	virtual uint32 GetCategories() override { return EAssetTypeCategories::Misc; }
-
-private:
-	/** Handler for when Edit is selected */
-	void ExecuteEdit(TArray<TWeakObjectPtr<class UEnvQuery>> Objects);
 };
\ No newline at end of file
diff --git a/Engine/Source/Editor/UnrealEd/Private/Toolkits/AssetEditorManager.cpp b/Engine/Source/Editor/UnrealEd/Private/Toolkits/AssetEditorManager.cpp
--- a/Engine/Source/Editor/UnrealEd/Private/Toolkits/AssetEditorManager.cpp
+++ b/Engine/Source/Editor/UnrealEd/Private/Toolkits/AssetEditorManager.cpp
@@ -233,7 +233,7 @@ bool FAssetEditorManager::OpenEditorForAsset(UObject* Asset, const EToolkitMode:
 	
 	AssetEditorRequestOpenEvent.Broadcast(Asset);
 
-	if( FindEditorForAsset(Asset, bBringToFrontIfOpen) != NULL )
+	if( FindEditorForAsset(Asset, bBringToFrontIfOpen) != nullptr )
 	{
 		// This asset is already open in an editor! (the call to FindEditorForAsset above will bring it to the front)
 		return true;
@@ -265,8 +265,6 @@ bool FAssetEditorManager::OpenEditorForAsset(UObject* Asset, const EToolkitMode:
 		}
 	}
 	
-	GWarn->EndSlowTask();
-	
 	if( ActualToolkitMode != EToolkitMode::WorldCentric && OpenedFromLevelEditor.IsValid() )
 	{
 		// @todo toolkit minor: Kind of lame use of a static variable here to prime the new asset editor.  This was done to avoid refactoring a few dozen files for a very minor change.
@@ -301,39 +299,123 @@ bool FAssetEditorManager::OpenEditorForAsset(UObject* Asset, const EToolkitMode:
 		FSimpleAssetEditor::CreateEditor(ActualToolkitMode, ActualToolkitMode == EToolkitMode::WorldCentric ? OpenedFromLevelEditor : TSharedPtr<IToolkitHost>(), Asset);
 	}
 
+	GWarn->EndSlowTask();
 	return true;
 }
 
 
 bool FAssetEditorManager::OpenEditorForAssets( const TArray< UObject* >& Assets, const EToolkitMode::Type ToolkitMode, TSharedPtr< IToolkitHost > OpenedFromLevelEditor )
 {
-	//@todo Temporarily whenever someone attempts to open an editor editing multiple assets the editor type will default to FSimpleAssetEditor [8/14/2012 Justin.Sargent]
-
-	if( ToolkitMode != EToolkitMode::WorldCentric && OpenedFromLevelEditor.IsValid() )
+	if( Assets.Num() == 1 )
 	{
-		// @todo toolkit minor: Kind of lame use of a static variable here to prime the new asset editor.  This was done to avoid refactoring a few dozen files for a very minor change.
-		FAssetEditorToolkit::SetPreviousWorldCentricToolkitHostForNewAssetEditor( OpenedFromLevelEditor.ToSharedRef() );
+		return OpenEditorForAsset(Assets[0], ToolkitMode, OpenedFromLevelEditor);
 	}
-
-	if( Assets.Num() == 1 )
+	else if (Assets.Num() > 0)
 	{
-		// Find any existing simple editors for this asset
-		const TArray<IAssetEditorInstance*> Editors = FAssetEditorManager::Get().FindEditorsForAsset( Assets[0] );
+		// If any of the assets are already open, remove them from the list of assets to open an editor for
+		TArray<UObject*> AlreadyOpenAssets;
+		for (auto Asset : Assets)
+		{
+			if( FindEditorForAsset(Asset, true) != nullptr )
+			{
+				AlreadyOpenAssets.Add(Asset);
+			}
+		}
+
+		// Verify that all the assets are of the same class
+		bool bAssetClassesMatch = true;
+		auto AssetClass = Assets[0]->GetClass();
+		for (int32 i = 1; i < Assets.Num(); i++)
+		{
+			if (Assets[i]->GetClass() != AssetClass)
+			{
+				bAssetClassesMatch = false;
+				break;
+			}
+		}
+
+		// If the classes don't match or any of the selected assets are already open, just open each asset in its own editor.
+		if (bAssetClassesMatch && AlreadyOpenAssets.Num() == 0)
+		{
+			FAssetToolsModule& AssetToolsModule = FModuleManager::LoadModuleChecked<FAssetToolsModule>(TEXT("AssetTools"));
+			TWeakPtr<IAssetTypeActions> AssetTypeActions = AssetToolsModule.Get().GetAssetTypeActionsForClass(AssetClass);
+
+			if (AssetTypeActions.IsValid())
+			{
+				GWarn->BeginSlowTask(LOCTEXT("OpenEditor", "Opening Editor(s)..."), true);
+
+				// Determine the appropriate toolkit mode for the asset type
+				auto ActualToolkitMode = ToolkitMode;
+				if (AssetTypeActions.Pin()->ShouldForceWorldCentric())
+				{
+					// This asset type prefers a specific toolkit mode
+					ActualToolkitMode = EToolkitMode::WorldCentric;
+
+					if (!OpenedFromLevelEditor.IsValid())
+					{
+						// We don't have a level editor to spawn in world-centric mode, so we'll find one now
+						// @todo sequencer: We should eventually eliminate this code (incl include dependencies) or change it to not make assumptions about a single level editor
+						OpenedFromLevelEditor = FModuleManager::LoadModuleChecked< FLevelEditorModule >("LevelEditor").GetFirstLevelEditor();
+					}
+				}
+
+				if (ActualToolkitMode != EToolkitMode::WorldCentric && OpenedFromLevelEditor.IsValid())
+				{
+					// @todo toolkit minor: Kind of lame use of a static variable here to prime the new asset editor.  This was done to avoid refactoring a few dozen files for a very minor change.
+					FAssetEditorToolkit::SetPreviousWorldCentricToolkitHostForNewAssetEditor(OpenedFromLevelEditor.ToSharedRef());
+				}
+
+				// Some assets (like UWorlds) may be destroyed and recreated as part of opening. To protect against this, keep the path to each asset and try to re-find any if they disappear.
+				struct FLocalAssetInfo
+				{
+					TWeakObjectPtr<UObject> WeakAsset;
+					FString AssetPath;
+
+					FLocalAssetInfo(const TWeakObjectPtr<UObject>& InWeakAsset, const FString& InAssetPath)
+						: WeakAsset(InWeakAsset), AssetPath(InAssetPath) {}
+				};
+
+				TArray<FLocalAssetInfo> AssetInfoList;
+				AssetInfoList.Reserve(Assets.Num());
+				for (auto Asset : Assets)
+				{
+					AssetInfoList.Add(FLocalAssetInfo(Asset, Asset->GetPathName()));
+				}
+
+				// How to handle multiple assets is left up to the type actions (i.e. open a single shared editor or an editor for each)
+				AssetTypeActions.Pin()->OpenAssetEditor(Assets, ActualToolkitMode == EToolkitMode::WorldCentric ? OpenedFromLevelEditor : TSharedPtr<IToolkitHost>());
+
+				// If any assets were destroyed, attempt to find them if they were recreated
+				for (int32 i = 0; i < Assets.Num(); i++)
+				{
+					auto& AssetInfo = AssetInfoList[i];
+					auto Asset = Assets[i];
 
-		IAssetEditorInstance* const * ExistingInstance = Editors.FindByPredicate( [&](IAssetEditorInstance* Editor){
-			return Editor->GetEditorName() == FSimpleAssetEditor::ToolkitFName;
-		} );
+					if (!AssetInfo.WeakAsset.IsValid() && !AssetInfo.AssetPath.IsEmpty())
+					{
+						Asset = FindObject<UObject>(nullptr, *AssetInfo.AssetPath);
+					}
+				}
 
-		if( ExistingInstance )
+				//@todo if needed, broadcast the event for every asset. It is possible, however, that a single shared editor was opened by the AssetTypeActions, not an editor for each asset.
+				/*AssetEditorOpenedEvent.Broadcast(Asset);*/
+
+				GWarn->EndSlowTask();
+			}
+		}
+		else
 		{
-			(*ExistingInstance)->FocusWindow();
-			return true;
+			// Asset types don't match or some are already open, so just open individual editors for the unopened ones
+			for (auto Asset : Assets)
+			{
+				if (!AlreadyOpenAssets.Contains(Asset))
+				{
+					OpenEditorForAsset(Asset, ToolkitMode, OpenedFromLevelEditor);
+				}
+			}
 		}
 	}
 
-	// No asset type actions for this asset. Just use a properties editor.
-	FSimpleAssetEditor::CreateEditor(ToolkitMode, ToolkitMode == EToolkitMode::WorldCentric ? OpenedFromLevelEditor : TSharedPtr<IToolkitHost>(), Assets );
-
 	return true;
 }
 
diff --git a/Engine/Source/Editor/UnrealEd/Public/Toolkits/AssetEditorManager.h b/Engine/Source/Editor/UnrealEd/Public/Toolkits/AssetEditorManager.h
--- a/Engine/Source/Editor/UnrealEd/Public/Toolkits/AssetEditorManager.h
+++ b/Engine/Source/Editor/UnrealEd/Public/Toolkits/AssetEditorManager.h
@@ -33,47 +33,56 @@ class UNREALED_API FAssetEditorManager
 {
 public:
 
-	// Get the singleton instance of the asset editor manager
+	/** Get the singleton instance of the asset editor manager */
 	static FAssetEditorManager& Get();
 
-	// Called when the editor is exiting to shutdown the manager
+	/** Called when the editor is exiting to shutdown the manager */
 	void OnExit();
 
-	// Tries to open an editor for the specified asset.  Returns true if the asset is open in an editor.
-	//   If the file is already open in an editor, it will not create another editor window but instead bring it to front
+	/** 
+	* Tries to open an editor for the specified asset.  Returns true if the asset is open in an editor.
+	* If the file is already open in an editor, it will not create another editor window but instead bring it to front
+	*/
 	bool OpenEditorForAsset( UObject* Asset, const EToolkitMode::Type ToolkitMode = EToolkitMode::Standalone, TSharedPtr< class IToolkitHost > OpenedFromLevelEditor = TSharedPtr< IToolkitHost >() );
-	bool OpenEditorForAssets( const TArray< UObject* >& Asset, const EToolkitMode::Type ToolkitMode = EToolkitMode::Standalone, TSharedPtr< class IToolkitHost > OpenedFromLevelEditor = TSharedPtr< IToolkitHost >() );
 
-	// Opens editors for the supplied assets (via OpenEditorForAsset
+	/** 
+	* Tries to open an editor for all of the specified assets. 
+	* If any of the assets are already open, it will not create a new editor for them.
+	* If all assets are of the same type, the supporting AssetTypeAction (if it exists) is responsible for the details of how to handle opening multiple assets at once.
+	*/
+	bool OpenEditorForAssets( const TArray< UObject* >& Assets, const EToolkitMode::Type ToolkitMode = EToolkitMode::Standalone, TSharedPtr< class IToolkitHost > OpenedFromLevelEditor = TSharedPtr< IToolkitHost >() );
+
+	/** Opens editors for the supplied assets (via OpenEditorForAsset) */
 	void OpenEditorsForAssets(const TArray<FString>& AssetsToOpen);
 
-	// Returns the primary editor if one is already open for the specified asset.
-	// If there is one open and bFocusIfOpen is true, that editor will be brought to the foreground and focused if possible.
+	/** Returns the primary editor if one is already open for the specified asset.
+	* If there is one open and bFocusIfOpen is true, that editor will be brought to the foreground and focused if possible.
+	*/
 	IAssetEditorInstance* FindEditorForAsset(UObject* Asset, bool bFocusIfOpen);
 
-	// Returns all editors currently opened for the specified asset
+	/** Returns all editors currently opened for the specified asset */
 	TArray<IAssetEditorInstance*> FindEditorsForAsset(UObject* Asset);
 
-	// Close any editor which is not this one
+	/** Close any editor which is not this one */
 	void CloseOtherEditors(UObject* Asset, IAssetEditorInstance* OnlyEditor);
 
 	/** Get all assets currently being tracked with open editors */
 	TArray<UObject*> GetAllEditedAssets();
 
-	// Notify the asset editor manager that an asset was opened
+	/** Notify the asset editor manager that an asset was opened */
 	void NotifyAssetOpened(UObject* Asset, IAssetEditorInstance* Instance);
 	void NotifyAssetsOpened( const TArray< UObject* >& Assets, IAssetEditorInstance* Instance);
 
-	// Notify the asset editor manager that an asset editor is done editing an asset
+	/** Notify the asset editor manager that an asset editor is done editing an asset */
 	void NotifyAssetClosed(UObject* Asset, IAssetEditorInstance* Instance);
 
-	// Notify the asset editor manager that an asset was closed
+	/** Notify the asset editor manager that an asset was closed */
 	void NotifyEditorClosed(IAssetEditorInstance* Instance);
 
 	// FGCObject interface
 	virtual void AddReferencedObjects( FReferenceCollector& Collector ) override;
 
-	// Close all open asset editors
+	/** Close all open asset editors */
 	bool CloseAllAssetEditors();
 	
 	/** Called when an asset editor is requested to be opened */
@@ -89,31 +98,31 @@ public:
 
 private:
 
-	// Hidden default constructor since the asset editor manager is a singleton.
+	/** Hidden default constructor since the asset editor manager is a singleton. */
 	FAssetEditorManager();
 
 
 private:
 
-	// Handles FAssetEditorRequestOpenAsset messages.
+	/** Handles FAssetEditorRequestOpenAsset messages. */
 	void HandleRequestOpenAssetMessage( const FAssetEditorRequestOpenAsset& Message, const IMessageContextRef& Context );
 
-	// Opens an asset by path
+	/** Opens an asset by path */
 	void OpenEditorForAsset(const FString& AssetPathName);
 
-	// Handles ticks from the ticker.
+	/** Handles ticks from the ticker. */
 	bool HandleTicker( float DeltaTime );
 
-	// Spawn a notification asking the user if they want to restore their previously open assets
+	/** Spawn a notification asking the user if they want to restore their previously open assets */
 	void SpawnRestorePreviouslyOpenAssetsNotification(const bool bCleanShutdown, const TArray<FString>& AssetsToOpen);
 
-	// Handler for when the "Restore Now" button is clicked on the RestorePreviouslyOpenAssets notification
+	/** Handler for when the "Restore Now" button is clicked on the RestorePreviouslyOpenAssets notification */
 	void OnConfirmRestorePreviouslyOpenAssets(TArray<FString> AssetsToOpen);
 
-	// Handler for when the "Don't Restore" button is clicked on the RestorePreviouslyOpenAssets notification
+	/** Handler for when the "Don't Restore" button is clicked on the RestorePreviouslyOpenAssets notification */
 	void OnCancelRestorePreviouslyOpenAssets();
 
-	// Saves a list of open asset editors so they can be restored on editor restart
+	/** Saves a list of open asset editors so they can be restored on editor restart */
 	void SaveOpenAssetEditors(bool bOnShutdown);
 
 private:
@@ -147,13 +156,13 @@ private:
 
 private:
 
-	// The singleton instance
+	/** The singleton instance */
 	static FAssetEditorManager* Instance;
 
-	// Holds the messaging endpoint.
+	/** Holds the messaging endpoint. */
 	FMessageEndpointPtr MessageEndpoint;
 
-	// Holds a delegate to be invoked when the widget ticks.
+	/** Holds a delegate to be invoked when the widget ticks. */
 	FTickerDelegate TickDelegate;
 
 	/** Multicast delegate executed when an asset editor is requested to be opened */
diff --git a/Engine/Source/Editor/UnrealEd/Public/Toolkits/SimpleAssetEditor.h b/Engine/Source/Editor/UnrealEd/Public/Toolkits/SimpleAssetEditor.h
--- a/Engine/Source/Editor/UnrealEd/Public/Toolkits/SimpleAssetEditor.h
+++ b/Engine/Source/Editor/UnrealEd/Public/Toolkits/SimpleAssetEditor.h
@@ -35,7 +35,7 @@ public:
 	virtual FText GetToolkitName() const override;
 	virtual FString GetWorldCentricTabPrefix() const override;
 	virtual FLinearColor GetWorldCentricTabColorScale() const override;
-	virtual bool IsPrimaryEditor() const override { return false; }
+	virtual bool IsPrimaryEditor() const override { return true; }
 	
 	/** Used to show or hide certain properties */
 	void SetPropertyVisibilityDelegate(FIsPropertyVisible InVisibilityDelegate);
