- Added a few icons to the path context menu. More are on the way.

diff --git a/Engine/Source/Editor/ContentBrowser/Private/PathContextMenu.cpp b/Engine/Source/Editor/ContentBrowser/Private/PathContextMenu.cpp
--- a/Engine/Source/Editor/ContentBrowser/Private/PathContextMenu.cpp
+++ b/Engine/Source/Editor/ContentBrowser/Private/PathContextMenu.cpp
@@ -81,9 +81,11 @@ void FPathContextMenu::MakePathViewContextMenu(FMenuBuilder& MenuBuilder)
 		{
 			// New Asset (submenu)
 			MenuBuilder.AddSubMenu(
-				LOCTEXT("CreateAsset", "Create Asset"),
-				LOCTEXT("CreateAssetTooltip", "Create a new asset."),
-				FNewMenuDelegate::CreateRaw( this, &FPathContextMenu::MakeNewAssetSubMenu )
+				LOCTEXT( "CreateAsset", "Create Asset" ),
+				LOCTEXT( "CreateAssetTooltip", "Create a new asset." ),
+				FNewMenuDelegate::CreateRaw( this, &FPathContextMenu::MakeNewAssetSubMenu ),
+				false,
+				FSlateIcon( FEditorStyle::GetStyleSetName(), "ContentBrowser.PathActions.NewAsset" )
 				);
 
 			// Explore
@@ -106,7 +108,9 @@ void FPathContextMenu::MakePathViewContextMenu(FMenuBuilder& MenuBuilder)
 				MenuBuilder.AddSubMenu(
 					LOCTEXT("SetColor", "Set Color"),
 					LOCTEXT("SetColorTooltip", "Sets the color this folder should appear as."),
-					FNewMenuDelegate::CreateRaw( this, &FPathContextMenu::MakeSetColorSubMenu )
+					FNewMenuDelegate::CreateRaw( this, &FPathContextMenu::MakeSetColorSubMenu ),
+					false,
+					FSlateIcon( FEditorStyle::GetStyleSetName(), "ContentBrowser.PathActions.SetColor" )
 					);
 			}
 			else
@@ -115,7 +119,7 @@ void FPathContextMenu::MakePathViewContextMenu(FMenuBuilder& MenuBuilder)
 				MenuBuilder.AddMenuEntry(
 					LOCTEXT("SetColor", "Set Color"),
 					LOCTEXT("SetColorTooltip", "Sets the color this folder should appear as."),
-					FSlateIcon(),
+					FSlateIcon(FEditorStyle::GetStyleSetName(), "ContentBrowser.PathActions.SetColor"),
 					FUIAction( FExecuteAction::CreateSP( this, &FPathContextMenu::ExecutePickColor ) )
 					);
 			}			
@@ -127,16 +131,17 @@ void FPathContextMenu::MakePathViewContextMenu(FMenuBuilder& MenuBuilder)
 		{
 		    // Save
 		    MenuBuilder.AddMenuEntry(
-			    LOCTEXT("SaveFolder", "Save"),
+			    LOCTEXT("SaveFolder", "Save All"),
 			    LOCTEXT("SaveFolderTooltip", "Saves all modified assets in this folder."),
-			    FSlateIcon(),
+			    FSlateIcon(FEditorStyle::GetStyleSetName(), "MainFrame.SaveAll"),
 			    FUIAction( FExecuteAction::CreateSP( this, &FPathContextMenu::ExecuteSaveFolder ) )
 			    );
     
 		    // Delete
 			MenuBuilder.AddMenuEntry(FGenericCommands::Get().Delete, NAME_None,
 			    LOCTEXT("DeleteFolder", "Delete"),
-			    LOCTEXT("DeleteFolderTooltip", "Removes this folder and all assets it contains.")
+			    LOCTEXT("DeleteFolderTooltip", "Removes this folder and all assets it contains."),
+				FSlateIcon( FEditorStyle::GetStyleSetName(), "ContentBrowser.AssetActions.Delete" )
 			    );
 
 			// Reference Viewer
diff --git a/Engine/Source/Editor/EditorStyle/Private/SlateEditorStyle.cpp b/Engine/Source/Editor/EditorStyle/Private/SlateEditorStyle.cpp
--- a/Engine/Source/Editor/EditorStyle/Private/SlateEditorStyle.cpp
+++ b/Engine/Source/Editor/EditorStyle/Private/SlateEditorStyle.cpp
@@ -4928,6 +4928,9 @@ void FSlateEditorStyle::FStyle::SetupContentBrowserStyle()
 		// New Asset
 		Set( "ContentBrowser.NewAsset", new IMAGE_BRUSH( "Icons/icon_file_new_40x", Icon25x25 ) );
 
+		Set( "ContentBrowser.PathActions.NewAsset", new IMAGE_BRUSH( "Icons/icon_file_new_16px", Icon16x16 ) );
+		Set( "ContentBrowser.PathActions.SetColor", new IMAGE_BRUSH( "Icons/icon_Cascade_Color_40x", Icon16x16 ) );
+
 		FComboButtonStyle NewAssetComboButton = GetWidgetStyle<FComboButtonStyle>( "ToolbarComboButton" );
 		Set( "ContentBrowser.NewAsset.Style", NewAssetComboButton );
 
