Adjusted LOCTEXT key for the "Create Asset" option tooltip

diff --git a/Engine/Source/Editor/ContentBrowser/Private/PathContextMenu.cpp b/Engine/Source/Editor/ContentBrowser/Private/PathContextMenu.cpp
--- a/Engine/Source/Editor/ContentBrowser/Private/PathContextMenu.cpp
+++ b/Engine/Source/Editor/ContentBrowser/Private/PathContextMenu.cpp
@@ -82,8 +82,8 @@ void FPathContextMenu::MakePathViewContextMenu(FMenuBuilder& MenuBuilder)
 		{
 			// New Asset (submenu)
 			MenuBuilder.AddSubMenu(
-				LOCTEXT( "CreateAsset", "Create Asset" ),
-				LOCTEXT( "CreateAssetTooltip", "Create a new asset." ),
+				LOCTEXT( "CreateNewAsset", "Create Asset" ),
+				LOCTEXT( "CreateNewAssetTooltip", "Create a new asset." ),
 				FNewMenuDelegate::CreateRaw( this, &FPathContextMenu::MakeNewAssetSubMenu ),
 				false,
 				FSlateIcon( FEditorStyle::GetStyleSetName(), "ContentBrowser.PathActions.NewAsset" )
