Can open merge tool from context menu, moved ini switch into editor ini

diff --git a/Engine/Source/Developer/AssetTools/Private/AssetTypeActions/AssetTypeActions_Blueprint.cpp b/Engine/Source/Developer/AssetTools/Private/AssetTypeActions/AssetTypeActions_Blueprint.cpp
--- a/Engine/Source/Developer/AssetTools/Private/AssetTypeActions/AssetTypeActions_Blueprint.cpp
+++ b/Engine/Source/Developer/AssetTools/Private/AssetTypeActions/AssetTypeActions_Blueprint.cpp
@@ -5,6 +5,7 @@
 #include "Kismet2/BlueprintEditorUtils.h"
 #include "Kismet2/KismetEditorUtilities.h"
 #include "BlueprintEditorModule.h"
+#include "BlueprintEditor.h"
 #include "AssetRegistryModule.h"
 #include "SBlueprintDiff.h"
 #include "ISourceControlModule.h"
@@ -80,6 +81,25 @@ void FAssetTypeActions_Blueprint::OpenAssetEditor( const TArray<UObject*>& InObj
 	}
 }
 
+bool FAssetTypeActions_Blueprint::CanMerge() const
+{
+	return true;
+}
+
+void FAssetTypeActions_Blueprint::Merge(UObject* InObject)
+{
+	UBlueprint* AsBlueprint = CastChecked<UBlueprint>(InObject);
+	// Kludge to get the merge panel in the blueprint editor to show up:
+	bool Success = FAssetEditorManager::Get().OpenEditorForAsset(InObject);
+	if( Success )
+	{
+		FBlueprintEditorModule& BlueprintEditorModule = FModuleManager::LoadModuleChecked<FBlueprintEditorModule>( "Kismet" );
+
+		FBlueprintEditor* BlueprintEditor = static_cast<FBlueprintEditor*>(FAssetEditorManager::Get().FindEditorForAsset(AsBlueprint, false));
+		BlueprintEditor->CreateMergeToolTab();
+	}
+}
+
 bool FAssetTypeActions_Blueprint::CanCreateNewDerivedBlueprint() const
 {
 	return true;
diff --git a/Engine/Source/Developer/AssetTools/Private/AssetTypeActions/AssetTypeActions_Blueprint.h b/Engine/Source/Developer/AssetTools/Private/AssetTypeActions/AssetTypeActions_Blueprint.h
--- a/Engine/Source/Developer/AssetTools/Private/AssetTypeActions/AssetTypeActions_Blueprint.h
+++ b/Engine/Source/Developer/AssetTools/Private/AssetTypeActions/AssetTypeActions_Blueprint.h
@@ -13,7 +13,9 @@ public:
 	virtual UClass* GetSupportedClass() const override { return UBlueprint::StaticClass(); }
 	virtual bool HasActions ( const TArray<UObject*>& InObjects ) const override { return true; }
 	virtual void GetActions( const TArray<UObject*>& InObjects, FMenuBuilder& MenuBuilder ) override;
-	virtual void OpenAssetEditor( const TArray<UObject*>& InObjects, TSharedPtr<class IToolkitHost> EditWithinLevelEditor = TSharedPtr<IToolkitHost>() ) override;
+	virtual void OpenAssetEditor(const TArray<UObject*>& InObjects, TSharedPtr<class IToolkitHost> EditWithinLevelEditor = TSharedPtr<IToolkitHost>()) override;
+	virtual bool CanMerge() const override;
+	virtual void Merge(UObject* InObject) override;
 	virtual uint32 GetCategories() override { return EAssetTypeCategories::Blueprint | EAssetTypeCategories::Basic; }
 	virtual void PerformAssetDiff(UObject* Asset1, UObject* Asset2, const struct FRevisionInfo& OldRevision, const struct FRevisionInfo& NewRevision) const override;
 	virtual class UThumbnailInfo* GetThumbnailInfo(UObject* Asset) const override;
diff --git a/Engine/Source/Developer/AssetTools/Public/AssetTypeActions_Base.h b/Engine/Source/Developer/AssetTools/Public/AssetTypeActions_Base.h
--- a/Engine/Source/Developer/AssetTools/Public/AssetTypeActions_Base.h
+++ b/Engine/Source/Developer/AssetTools/Public/AssetTypeActions_Base.h
@@ -46,6 +46,16 @@ public:
 		return true;
 	}
 
+	virtual bool CanMerge() const override
+	{
+		return false;
+	}
+
+	virtual void Merge(UObject* InObject) override
+	{
+		check(false); // no generic merge operation exists yet, did you override CanMerge but not Merge?
+	}
+
 	virtual bool ShouldForceWorldCentric() override
 	{
 		return false;
diff --git a/Engine/Source/Developer/AssetTools/Public/IAssetTypeActions.h b/Engine/Source/Developer/AssetTools/Public/IAssetTypeActions.h
--- a/Engine/Source/Developer/AssetTools/Public/IAssetTypeActions.h
+++ b/Engine/Source/Developer/AssetTools/Public/IAssetTypeActions.h
@@ -77,6 +77,12 @@ public:
 	/** Returns true if this class can be used as a filter in the content browser */
 	virtual bool CanFilter() = 0;
 
+	/** Returns true if this class can be merged (either manually or automatically) */
+	virtual bool CanMerge() const = 0;
+
+	/** Begins a merge operation for InObjects */
+	virtual void Merge( UObject* InObject ) = 0;
+
 	/** Returns the categories that this asset type. The return value is one or more flags from EAssetTypeCategories.  */
 	virtual uint32 GetCategories() = 0;
 
diff --git a/Engine/Source/Developer/Merge/Private/Merge.cpp b/Engine/Source/Developer/Merge/Private/Merge.cpp
--- a/Engine/Source/Developer/Merge/Private/Merge.cpp
+++ b/Engine/Source/Developer/Merge/Private/Merge.cpp
@@ -241,7 +241,7 @@ TSharedRef<SDockTab> FMerge::GenerateMergeWidget(const UBlueprint& Object, TShar
 bool FMerge::PendingMerge(const UBlueprint& BlueprintObj) const
 {
 	bool bIsMergeEnabled = false;
-	GConfig->GetBool(TEXT("AssetMerge"), TEXT("EnableAssetMerge"), bIsMergeEnabled, GEngineIni);
+	GConfig->GetBool(TEXT("AssetMerge"), TEXT("EnableAssetMerge"), bIsMergeEnabled, GEditorIni);
 	ISourceControlProvider& SourceControlProvider = ISourceControlModule::Get().GetProvider();
 
 	bool bPendingMerge = false;
diff --git a/Engine/Source/Editor/ContentBrowser/Private/AssetContextMenu.cpp b/Engine/Source/Editor/ContentBrowser/Private/AssetContextMenu.cpp
--- a/Engine/Source/Editor/ContentBrowser/Private/AssetContextMenu.cpp
+++ b/Engine/Source/Editor/ContentBrowser/Private/AssetContextMenu.cpp
@@ -39,6 +39,7 @@
 FAssetContextMenu::FAssetContextMenu(const TWeakPtr<SAssetView>& InAssetView)
 	: AssetView(InAssetView)
 	, bAtLeastOneNonRedirectorSelected(false)
+	, bCanExecuteSCCMerge(false)
 	, bCanExecuteSCCCheckOut(false)
 	, bCanExecuteSCCOpenForAdd(false)
 	, bCanExecuteSCCCheckIn(false)
@@ -629,103 +630,116 @@ bool FAssetContextMenu::AddSourceControlMenuOptions(FMenuBuilder& MenuBuilder)
 
 void FAssetContextMenu::FillSourceControlSubMenu(FMenuBuilder& MenuBuilder)
 {
-	if( CanExecuteSCCSync() )
-		{
-			MenuBuilder.AddMenuEntry(
-				LOCTEXT("SCCSync", "Sync"),
-				LOCTEXT("SCCSyncTooltip", "Updates the item to the latest version in source control."),
-				FSlateIcon(),
-				FUIAction(
-					FExecuteAction::CreateSP( this, &FAssetContextMenu::ExecuteSCCSync ),
-					FCanExecuteAction::CreateSP( this, &FAssetContextMenu::CanExecuteSCCSync )
-					)
-				);
-		}
+	if( CanExecuteSCCMerge() )
+	{
+		MenuBuilder.AddMenuEntry(
+			LOCTEXT("SCCMerge", "Merge"),
+			LOCTEXT("SCCMergeTooltip", "Opens the blueprint editor with the merge tool open."),
+			FSlateIcon(),
+			FUIAction(
+				FExecuteAction::CreateSP(this, &FAssetContextMenu::ExecuteSCCMerge),
+				FCanExecuteAction::CreateSP(this, &FAssetContextMenu::CanExecuteSCCMerge)
+			)
+		);
+	}
 
-		if ( CanExecuteSCCCheckOut() )
-		{
-			MenuBuilder.AddMenuEntry(
-				LOCTEXT("SCCCheckOut", "Check Out"),
-				LOCTEXT("SCCCheckOutTooltip", "Checks out the selected asset from source control."),
-				FSlateIcon(FEditorStyle::GetStyleSetName(), "SourceControl.Actions.CheckOut"),
-				FUIAction(
-					FExecuteAction::CreateSP( this, &FAssetContextMenu::ExecuteSCCCheckOut ),
-					FCanExecuteAction::CreateSP( this, &FAssetContextMenu::CanExecuteSCCCheckOut )
-					)
-				);
-		}
+	if( CanExecuteSCCSync() )
+	{
+		MenuBuilder.AddMenuEntry(
+			LOCTEXT("SCCSync", "Sync"),
+			LOCTEXT("SCCSyncTooltip", "Updates the item to the latest version in source control."),
+			FSlateIcon(),
+			FUIAction(
+				FExecuteAction::CreateSP( this, &FAssetContextMenu::ExecuteSCCSync ),
+				FCanExecuteAction::CreateSP( this, &FAssetContextMenu::CanExecuteSCCSync )
+			)
+		);
+	}
 
-		if ( CanExecuteSCCOpenForAdd() )
-		{
-			MenuBuilder.AddMenuEntry(
-				LOCTEXT("SCCOpenForAdd", "Mark For Add"),
-				LOCTEXT("SCCOpenForAddTooltip", "Adds the selected asset to source control."),
-				FSlateIcon(FEditorStyle::GetStyleSetName(), "SourceControl.Actions.Add"),
-				FUIAction(
-					FExecuteAction::CreateSP( this, &FAssetContextMenu::ExecuteSCCOpenForAdd ),
-					FCanExecuteAction::CreateSP( this, &FAssetContextMenu::CanExecuteSCCOpenForAdd )
-					)
-				);
-		}
+	if ( CanExecuteSCCCheckOut() )
+	{
+		MenuBuilder.AddMenuEntry(
+			LOCTEXT("SCCCheckOut", "Check Out"),
+			LOCTEXT("SCCCheckOutTooltip", "Checks out the selected asset from source control."),
+			FSlateIcon(FEditorStyle::GetStyleSetName(), "SourceControl.Actions.CheckOut"),
+			FUIAction(
+				FExecuteAction::CreateSP( this, &FAssetContextMenu::ExecuteSCCCheckOut ),
+				FCanExecuteAction::CreateSP( this, &FAssetContextMenu::CanExecuteSCCCheckOut )
+			)
+		);
+	}
 
-		if ( CanExecuteSCCCheckIn() )
-		{
-			MenuBuilder.AddMenuEntry(
-				LOCTEXT("SCCCheckIn", "Check In"),
-				LOCTEXT("SCCCheckInTooltip", "Checks in the selected asset to source control."),
-				FSlateIcon(),
-				FUIAction(
-					FExecuteAction::CreateSP( this, &FAssetContextMenu::ExecuteSCCCheckIn ),
-					FCanExecuteAction::CreateSP( this, &FAssetContextMenu::CanExecuteSCCCheckIn )
-					)
-				);
-		}
+	if ( CanExecuteSCCOpenForAdd() )
+	{
+		MenuBuilder.AddMenuEntry(
+			LOCTEXT("SCCOpenForAdd", "Mark For Add"),
+			LOCTEXT("SCCOpenForAddTooltip", "Adds the selected asset to source control."),
+			FSlateIcon(FEditorStyle::GetStyleSetName(), "SourceControl.Actions.Add"),
+			FUIAction(
+				FExecuteAction::CreateSP( this, &FAssetContextMenu::ExecuteSCCOpenForAdd ),
+				FCanExecuteAction::CreateSP( this, &FAssetContextMenu::CanExecuteSCCOpenForAdd )
+			)
+		);
+	}
 
+	if ( CanExecuteSCCCheckIn() )
+	{
 		MenuBuilder.AddMenuEntry(
-			LOCTEXT("SCCRefresh", "Refresh"),
-			LOCTEXT("SCCRefreshTooltip", "Updates the source control status of the asset."),
-			FSlateIcon(FEditorStyle::GetStyleSetName(), "SourceControl.Actions.Refresh"),
+			LOCTEXT("SCCCheckIn", "Check In"),
+			LOCTEXT("SCCCheckInTooltip", "Checks in the selected asset to source control."),
+			FSlateIcon(),
 			FUIAction(
-				FExecuteAction::CreateSP( this, &FAssetContextMenu::ExecuteSCCRefresh ),
-				FCanExecuteAction::CreateSP( this, &FAssetContextMenu::CanExecuteSCCRefresh )
-				)
-			);
+				FExecuteAction::CreateSP( this, &FAssetContextMenu::ExecuteSCCCheckIn ),
+				FCanExecuteAction::CreateSP( this, &FAssetContextMenu::CanExecuteSCCCheckIn )
+			)
+		);
+	}
 
-		if( CanExecuteSCCHistory() )
-		{
-			MenuBuilder.AddMenuEntry(
-				LOCTEXT("SCCHistory", "History"),
-				LOCTEXT("SCCHistoryTooltip", "Displays the source control revision history of the selected asset."),
-				FSlateIcon(),
-				FUIAction(
-					FExecuteAction::CreateSP( this, &FAssetContextMenu::ExecuteSCCHistory ),
-					FCanExecuteAction::CreateSP( this, &FAssetContextMenu::CanExecuteSCCHistory )
-					)
-				);
+	MenuBuilder.AddMenuEntry(
+		LOCTEXT("SCCRefresh", "Refresh"),
+		LOCTEXT("SCCRefreshTooltip", "Updates the source control status of the asset."),
+		FSlateIcon(FEditorStyle::GetStyleSetName(), "SourceControl.Actions.Refresh"),
+		FUIAction(
+			FExecuteAction::CreateSP( this, &FAssetContextMenu::ExecuteSCCRefresh ),
+			FCanExecuteAction::CreateSP( this, &FAssetContextMenu::CanExecuteSCCRefresh )
+			)
+		);
 
-			MenuBuilder.AddMenuEntry(
-				LOCTEXT("SCCDiffAgainstDepot", "Diff Against Depot"),
-				LOCTEXT("SCCDiffAgainstDepotTooltip", "Look at differences between your version of the asset and that in source control."),
-				FSlateIcon(),
-				FUIAction(
-					FExecuteAction::CreateSP( this, &FAssetContextMenu::ExecuteSCCDiffAgainstDepot ),
-					FCanExecuteAction::CreateSP( this, &FAssetContextMenu::CanExecuteSCCDiffAgainstDepot )
-					)
-				);	
-		}
+	if( CanExecuteSCCHistory() )
+	{
+		MenuBuilder.AddMenuEntry(
+			LOCTEXT("SCCHistory", "History"),
+			LOCTEXT("SCCHistoryTooltip", "Displays the source control revision history of the selected asset."),
+			FSlateIcon(),
+			FUIAction(
+				FExecuteAction::CreateSP( this, &FAssetContextMenu::ExecuteSCCHistory ),
+				FCanExecuteAction::CreateSP( this, &FAssetContextMenu::CanExecuteSCCHistory )
+			)
+		);
 
-		if( CanExecuteSCCRevert() )
-		{
-			MenuBuilder.AddMenuEntry(
-				LOCTEXT("SCCRevert", "Revert"),
-				LOCTEXT("SCCRevertTooltip", "Reverts the asset to the state it was before it was checked out."),
-				FSlateIcon(),
-				FUIAction(
-					FExecuteAction::CreateSP( this, &FAssetContextMenu::ExecuteSCCRevert ),
-					FCanExecuteAction::CreateSP( this, &FAssetContextMenu::CanExecuteSCCRevert )
-					)
-				);
-		}
+		MenuBuilder.AddMenuEntry(
+			LOCTEXT("SCCDiffAgainstDepot", "Diff Against Depot"),
+			LOCTEXT("SCCDiffAgainstDepotTooltip", "Look at differences between your version of the asset and that in source control."),
+			FSlateIcon(),
+			FUIAction(
+				FExecuteAction::CreateSP( this, &FAssetContextMenu::ExecuteSCCDiffAgainstDepot ),
+				FCanExecuteAction::CreateSP( this, &FAssetContextMenu::CanExecuteSCCDiffAgainstDepot )
+			)
+		);	
+	}
+
+	if( CanExecuteSCCRevert() )
+	{
+		MenuBuilder.AddMenuEntry(
+			LOCTEXT("SCCRevert", "Revert"),
+			LOCTEXT("SCCRevertTooltip", "Reverts the asset to the state it was before it was checked out."),
+			FSlateIcon(),
+			FUIAction(
+				FExecuteAction::CreateSP( this, &FAssetContextMenu::ExecuteSCCRevert ),
+				FCanExecuteAction::CreateSP( this, &FAssetContextMenu::CanExecuteSCCRevert )
+			)
+		);
+	}
 }
 
 bool FAssetContextMenu::AddCollectionMenuOptions(FMenuBuilder& MenuBuilder)
@@ -1374,6 +1388,29 @@ void FAssetContextMenu::ExecuteSCCRefresh()
 	ISourceControlModule::Get().GetProvider().Execute(ISourceControlOperation::Create<FUpdateStatus>(), SourceControlHelpers::PackageFilenames(PackageNames), EConcurrency::Asynchronous);
 }
 
+void FAssetContextMenu::ExecuteSCCMerge()
+{
+	FAssetToolsModule& AssetToolsModule = FModuleManager::GetModuleChecked<FAssetToolsModule>("AssetTools");
+
+	for (int32 AssetIdx = 0; AssetIdx < SelectedAssets.Num(); AssetIdx++)
+	{
+		// Get the actual asset (will load it)
+		const FAssetData& AssetData = SelectedAssets[AssetIdx];
+
+		UObject* CurrentObject = AssetData.GetAsset();
+		if (CurrentObject)
+		{
+			const FString PackagePath = AssetData.PackageName.ToString();
+			const FString PackageName = AssetData.AssetName.ToString();
+			auto AssetTypeActions = AssetToolsModule.Get().GetAssetTypeActionsForClass( CurrentObject->GetClass() ).Pin();
+			if( AssetTypeActions.IsValid() )
+			{
+				AssetTypeActions->Merge(CurrentObject);
+			}
+		}
+	}
+}
+
 void FAssetContextMenu::ExecuteSCCCheckOut()
 {
 	TArray<UPackage*> PackagesToCheckOut;
@@ -1605,6 +1642,37 @@ bool FAssetContextMenu::CanExecuteSCCRefresh() const
 	return ISourceControlModule::Get().IsEnabled();
 }
 
+bool FAssetContextMenu::CanExecuteSCCMerge() const
+{
+	FAssetToolsModule& AssetToolsModule = FModuleManager::GetModuleChecked<FAssetToolsModule>("AssetTools");
+
+	bool bCanExecuteMerge = false;
+	GConfig->GetBool(TEXT("AssetMerge"), TEXT("EnableAssetMerge"), bCanExecuteMerge, GEditorIni);
+	if( bCanExecuteMerge )
+	{
+		bCanExecuteMerge = bCanExecuteSCCMerge;
+		for (int32 AssetIdx = 0; AssetIdx < SelectedAssets.Num() && bCanExecuteMerge; AssetIdx++)
+		{
+			// Get the actual asset (will load it)
+			const FAssetData& AssetData = SelectedAssets[AssetIdx];
+			UObject* CurrentObject = AssetData.GetAsset();
+			if (CurrentObject)
+			{
+				auto AssetTypeActions = AssetToolsModule.Get().GetAssetTypeActionsForClass(CurrentObject->GetClass()).Pin();
+				if (AssetTypeActions.IsValid())
+				{
+					bCanExecuteMerge = AssetTypeActions->CanMerge();
+				}
+			}
+			else
+			{
+				bCanExecuteMerge = false;
+			}
+		}
+	}
+	return bCanExecuteMerge;
+}
+
 bool FAssetContextMenu::CanExecuteSCCCheckOut() const
 {
 	return bCanExecuteSCCCheckOut;
@@ -1725,6 +1793,7 @@ bool FAssetContextMenu::CanClearCustomThumbnails() const
 void FAssetContextMenu::CacheCanExecuteVars()
 {
 	bAtLeastOneNonRedirectorSelected = false;
+	bCanExecuteSCCMerge = false;
 	bCanExecuteSCCCheckOut = false;
 	bCanExecuteSCCOpenForAdd = false;
 	bCanExecuteSCCCheckIn = false;
@@ -1752,6 +1821,11 @@ void FAssetContextMenu::CacheCanExecuteVars()
 			FSourceControlStatePtr SourceControlState = SourceControlProvider.GetState(SourceControlHelpers::PackageFilename(AssetData.PackageName.ToString()), EStateCacheUsage::Use);
 			if(SourceControlState.IsValid())
 			{
+				if (SourceControlState->IsConflicted() )
+				{
+					bCanExecuteSCCMerge = true;
+				}
+
 				if ( SourceControlState->CanCheckout() )
 				{
 					bCanExecuteSCCCheckOut = true;
@@ -1780,6 +1854,7 @@ void FAssetContextMenu::CacheCanExecuteVars()
 		}
 
 		if ( bAtLeastOneNonRedirectorSelected
+			&& bCanExecuteSCCMerge
 			&& bCanExecuteSCCCheckOut
 			&& bCanExecuteSCCOpenForAdd
 			&& bCanExecuteSCCCheckIn
diff --git a/Engine/Source/Editor/ContentBrowser/Private/AssetContextMenu.h b/Engine/Source/Editor/ContentBrowser/Private/AssetContextMenu.h
--- a/Engine/Source/Editor/ContentBrowser/Private/AssetContextMenu.h
+++ b/Engine/Source/Editor/ContentBrowser/Private/AssetContextMenu.h
@@ -173,6 +173,9 @@ private:
 	/** Handler for when "Refresh source control" is selected */
 	void ExecuteSCCRefresh();
 
+	/** Handler for when "Merge" is selected */
+	void ExecuteSCCMerge();
+
 	/** Handler for when "Checkout from source control" is selected */
 	void ExecuteSCCCheckOut();
 
@@ -230,6 +233,9 @@ private:
 	/** Handler to check to see if "Refresh source control" can be executed */
 	bool CanExecuteSCCRefresh() const;
 
+	/** Handler to check to see if "Merge" can be executed */
+	bool CanExecuteSCCMerge() const;
+
 	/** Handler to check to see if "Checkout from source control" can be executed */
 	bool CanExecuteSCCCheckOut() const;
 
@@ -311,6 +317,7 @@ private:
 
 	/** Cached CanExecute vars */
 	bool bAtLeastOneNonRedirectorSelected;
+	bool bCanExecuteSCCMerge;
 	bool bCanExecuteSCCCheckOut;
 	bool bCanExecuteSCCOpenForAdd;
 	bool bCanExecuteSCCCheckIn;
