Fixed tooltip.

diff --git a/Engine/Source/Editor/ContentBrowser/Private/AssetContextMenu.cpp b/Engine/Source/Editor/ContentBrowser/Private/AssetContextMenu.cpp
--- a/Engine/Source/Editor/ContentBrowser/Private/AssetContextMenu.cpp
+++ b/Engine/Source/Editor/ContentBrowser/Private/AssetContextMenu.cpp
@@ -350,7 +350,7 @@ void FAssetContextMenu::MakeAssetActionsSubMenu(FMenuBuilder& MenuBuilder)
 		// Migrate
 		MenuBuilder.AddMenuEntry(
 			LOCTEXT("MigrateAsset", "Migrate..."),
-			LOCTEXT("MigrateAssetTooltip", "Copies all selected assets and their dependencies to another game"),
+			LOCTEXT("MigrateAssetTooltip", "Copies all selected assets and their dependencies to another project"),
 			FSlateIcon(),
 			FUIAction( FExecuteAction::CreateSP( this, &FAssetContextMenu::ExecuteMigrateAsset ) )
 			);
