Moved Blueprint Merge tool toggle to experimental editor settings

diff --git a/Engine/Source/Developer/Merge/Private/Merge.cpp b/Engine/Source/Developer/Merge/Private/Merge.cpp
--- a/Engine/Source/Developer/Merge/Private/Merge.cpp
+++ b/Engine/Source/Developer/Merge/Private/Merge.cpp
@@ -240,8 +240,7 @@ TSharedRef<SDockTab> FMerge::GenerateMergeWidget(const UBlueprint& Object, TShar
 
 bool FMerge::PendingMerge(const UBlueprint& BlueprintObj) const
 {
-	bool bIsMergeEnabled = false;
-	GConfig->GetBool(TEXT("AssetMerge"), TEXT("EnableAssetMerge"), bIsMergeEnabled, GEditorIni);
+	bool bIsMergeEnabled = GetDefault<UEditorExperimentalSettings>()->bEnableBlueprintMergeTool;
 	ISourceControlProvider& SourceControlProvider = ISourceControlModule::Get().GetProvider();
 
 	bool bPendingMerge = false;
diff --git a/Engine/Source/Editor/ContentBrowser/Private/AssetContextMenu.cpp b/Engine/Source/Editor/ContentBrowser/Private/AssetContextMenu.cpp
--- a/Engine/Source/Editor/ContentBrowser/Private/AssetContextMenu.cpp
+++ b/Engine/Source/Editor/ContentBrowser/Private/AssetContextMenu.cpp
@@ -1646,8 +1646,7 @@ bool FAssetContextMenu::CanExecuteSCCMerge() const
 {
 	FAssetToolsModule& AssetToolsModule = FModuleManager::GetModuleChecked<FAssetToolsModule>("AssetTools");
 
-	bool bCanExecuteMerge = false;
-	GConfig->GetBool(TEXT("AssetMerge"), TEXT("EnableAssetMerge"), bCanExecuteMerge, GEditorIni);
+	bool bCanExecuteMerge = GetDefault<UEditorExperimentalSettings>()->bEnableBlueprintMergeTool;
 	if( bCanExecuteMerge )
 	{
 		bCanExecuteMerge = bCanExecuteSCCMerge;
diff --git a/Engine/Source/Editor/UnrealEd/Classes/Settings/EditorExperimentalSettings.h b/Engine/Source/Editor/UnrealEd/Classes/Settings/EditorExperimentalSettings.h
--- a/Engine/Source/Editor/UnrealEd/Classes/Settings/EditorExperimentalSettings.h
+++ b/Engine/Source/Editor/UnrealEd/Classes/Settings/EditorExperimentalSettings.h
@@ -85,6 +85,10 @@ public:
 	UPROPERTY(EditAnywhere, config, Category = AI, meta = (DisplayName = "Environment Querying System"))
 	bool bEQSEditor;
 
+	/** Enables the Blueprint merge tool */
+	UPROPERTY(EditAnywhere, config, Category = AI, meta = (DisplayName = "Enable Blueprint Merge Tool"))
+	bool bEnableBlueprintMergeTool;
+
 	/**
 	 * Returns an event delegate that is executed when a setting has changed.
 	 *
