Editor - Making the breadcrumb bar now internally use a scrollbox in order to always bias having the most child path visible.  Enhancing the SScrollBox so that it has functions you can call to scroll to the end/start.  Fixing several styling issues with the border shadow of the scroll boxes, it now is much smaller in terms of original image, and is a BoxBrush now, so that it doesn't try to desire 64 su, as the minimum size of a h or v scroll box.

diff --git a/Engine/Source/Editor/ContentBrowser/Private/SContentBrowser.cpp b/Engine/Source/Editor/ContentBrowser/Private/SContentBrowser.cpp
--- a/Engine/Source/Editor/ContentBrowser/Private/SContentBrowser.cpp
+++ b/Engine/Source/Editor/ContentBrowser/Private/SContentBrowser.cpp
@@ -304,27 +304,20 @@ void SContentBrowser::Construct( const FArguments& InArgs, const FName& InInstan
 
 					// Path
 					+ SHorizontalBox::Slot()
-					.VAlign(VAlign_Fill)
+					.VAlign(VAlign_Center)
+					.HAlign(HAlign_Left)
 					.FillWidth(1.0f)
-					.Padding( FMargin(0) )
+					.Padding(FMargin(0))
 					[
-						SNew(SHorizontalBox)
-
-						+ SHorizontalBox::Slot()
-						.VAlign(VAlign_Center)
-						.FillWidth(1.0f)
-						[
-							SAssignNew(PathBreadcrumbTrail, SBreadcrumbTrail<FString>)
-							//.ToolTipText( LOCTEXT("PathTooltip", "Content Path") )
-							.ButtonContentPadding(FMargin(3.0f, 3.0f))
-							.DelimiterImage(FEditorStyle::GetBrush("ContentBrowser.PathDelimiter"))
-							.TextStyle(FEditorStyle::Get(), "ContentBrowser.PathText")
-							.ShowLeadingDelimiter( false )
-							.InvertTextColorOnHover( false )
-							.OnCrumbClicked(this, &SContentBrowser::OnPathClicked)
-							.GetCrumbMenuContent(this, &SContentBrowser::OnGetCrumbDelimiterContent)
-							.AddMetaData<FTagMetaData>(FTagMetaData(TEXT("ContentBrowserPath")))
-						]
+						SAssignNew(PathBreadcrumbTrail, SBreadcrumbTrail<FString>)
+						.ButtonContentPadding(FMargin(3.0f, 3.0f))
+						.DelimiterImage(FEditorStyle::GetBrush("ContentBrowser.PathDelimiter"))
+						.TextStyle(FEditorStyle::Get(), "ContentBrowser.PathText")
+						.ShowLeadingDelimiter(false)
+						.InvertTextColorOnHover(false)
+						.OnCrumbClicked(this, &SContentBrowser::OnPathClicked)
+						.GetCrumbMenuContent(this, &SContentBrowser::OnGetCrumbDelimiterContent)
+						.AddMetaData<FTagMetaData>(FTagMetaData(TEXT("ContentBrowserPath")))
 					]
 
 					// Lock button
diff --git a/Engine/Source/Runtime/Slate/Private/Widgets/Layout/SScrollBox.cpp b/Engine/Source/Runtime/Slate/Private/Widgets/Layout/SScrollBox.cpp
--- a/Engine/Source/Runtime/Slate/Private/Widgets/Layout/SScrollBox.cpp
+++ b/Engine/Source/Runtime/Slate/Private/Widgets/Layout/SScrollBox.cpp
@@ -156,6 +156,7 @@ void SScrollBox::Construct( const FArguments& InArgs )
 	SoftwareCursorPosition = FVector2D::ZeroVector;
 	OnUserScrolled = InArgs._OnUserScrolled;
 	Orientation = InArgs._Orientation;
+	bScrollToEnd = false;
 
 	if (InArgs._ExternalScrollbar.IsValid())
 	{
@@ -347,6 +348,17 @@ float SScrollBox::GetScrollOffset()
 void SScrollBox::SetScrollOffset( float NewScrollOffset )
 {
 	DesiredScrollOffset = NewScrollOffset;
+	bScrollToEnd = false;
+}
+
+void SScrollBox::ScrollToStart()
+{
+	SetScrollOffset(0);
+}
+
+void SScrollBox::ScrollToEnd()
+{
+	bScrollToEnd = true;
 }
 
 EOrientation SScrollBox::GetOrientation()
@@ -407,6 +419,12 @@ void SScrollBox::Tick( const FGeometry& AllottedGeometry, const double InCurrent
 	
 	const float ContentSize = GetScrollComponentFromVector(ScrollPanel->GetDesiredSize());
 
+	if ( bScrollToEnd )
+	{
+		DesiredScrollOffset = ContentSize;
+		bScrollToEnd = false;
+	}
+
 	// If this scroll box has no size, do not compute a view fraction because it will be wrong and causes pop in when the size is available
 	const float ViewFraction = GetScrollComponentFromVector(AllottedGeometry.Size) > 0 ? GetScrollComponentFromVector(ScrollPanelGeometry.Size) / ContentSize : 1;
 	const float ViewOffset = FMath::Clamp<float>( DesiredScrollOffset/ContentSize, 0.0, 1.0 - ViewFraction );
diff --git a/Engine/Source/Runtime/Slate/Public/Widgets/Layout/SScrollBox.h b/Engine/Source/Runtime/Slate/Public/Widgets/Layout/SScrollBox.h
--- a/Engine/Source/Runtime/Slate/Public/Widgets/Layout/SScrollBox.h
+++ b/Engine/Source/Runtime/Slate/Public/Widgets/Layout/SScrollBox.h
@@ -97,6 +97,10 @@ public:
 
 	void SetScrollOffset( float NewScrollOffset );
 
+	void ScrollToStart();
+
+	void ScrollToEnd();
+
 	/** Get the current orientation of the scrollbox. */
 	EOrientation GetOrientation();
 
@@ -212,6 +216,9 @@ private:
 
 	/** Style resource for the scrollbar */
 	const FScrollBarStyle* ScrollBarStyle;
+
+	/** If true, will scroll to the end next Tick */
+	bool bScrollToEnd;
 };
 
 
diff --git a/Engine/Source/Runtime/Slate/Public/Widgets/Navigation/SBreadcrumbTrail.h b/Engine/Source/Runtime/Slate/Public/Widgets/Navigation/SBreadcrumbTrail.h
--- a/Engine/Source/Runtime/Slate/Public/Widgets/Navigation/SBreadcrumbTrail.h
+++ b/Engine/Source/Runtime/Slate/Public/Widgets/Navigation/SBreadcrumbTrail.h
@@ -104,7 +104,9 @@ public:
 
 		ChildSlot
 		[
-			SAssignNew(CrumbBox, SHorizontalBox)
+			SAssignNew(CrumbBox, SScrollBox)
+			.Orientation(Orient_Horizontal)
+			.ScrollBarVisibility(EVisibility::Collapsed)
 		];
 
 		AddLeadingDelimiter();
@@ -119,7 +121,6 @@ public:
 
 		// Add the crumb button
 		CrumbBox->AddSlot()
-		.AutoWidth()
 		[
 			SNew(SVerticalBox)
 
@@ -165,7 +166,6 @@ public:
 		}
 
 		CrumbBox->AddSlot()
-		.AutoWidth()
 		[
 			SNew(SVerticalBox)
 
@@ -188,6 +188,9 @@ public:
 
 		// Trigger event
 		OnCrumbPushed.ExecuteIfBound(NewCrumbData);
+
+		// We've added a new crumb so defer scrolling to the end to next tick
+		CrumbBox->ScrollToEnd();
 	}
 
 	/** Pops a crumb off the end of the trail. Returns the crumb data. */
@@ -387,7 +390,6 @@ private:
 	void AddLeadingDelimiter()
 	{
 		CrumbBox->AddSlot()
-		.AutoWidth()
 		.VAlign(VAlign_Center)
 		[
 			SNew(SImage)
@@ -399,7 +401,7 @@ private:
 private:
 
 	/** The horizontal box which contains all the breadcrumbs */
-	TSharedPtr<SHorizontalBox> CrumbBox;
+	TSharedPtr<SScrollBox> CrumbBox;
 
 	/** The list of crumbs and their data */
 	TArray<FCrumbItem> CrumbList;
diff --git a/Engine/Source/Runtime/SlateCore/Private/Styling/CoreStyle.cpp b/Engine/Source/Runtime/SlateCore/Private/Styling/CoreStyle.cpp
--- a/Engine/Source/Runtime/SlateCore/Private/Styling/CoreStyle.cpp
+++ b/Engine/Source/Runtime/SlateCore/Private/Styling/CoreStyle.cpp
@@ -1001,18 +1001,18 @@ TSharedRef<ISlateStyle> FCoreStyle::Create( const FName& InStyleSetName )
 	// SScrollBox defaults...
 	{
 		Style->Set( "ScrollBox", FScrollBoxStyle()
-			.SetTopShadowBrush( IMAGE_BRUSH( "Common/ScrollBoxShadowTop", FVector2D(64,8) ) )
-			.SetBottomShadowBrush( IMAGE_BRUSH( "Common/ScrollBoxShadowBottom", FVector2D(64,8) ) )
-			.SetLeftShadowBrush( IMAGE_BRUSH( "Common/ScrollBoxShadowLeft", FVector2D(8, 64) ) )
-			.SetRightShadowBrush( IMAGE_BRUSH( "Common/ScrollBoxShadowRight", FVector2D(8, 64) ) )
+			.SetTopShadowBrush( BOX_BRUSH( "Common/ScrollBoxShadowTop", FVector2D(16, 8), FMargin(0.5, 1, 0.5, 0) ) )
+			.SetBottomShadowBrush(BOX_BRUSH("Common/ScrollBoxShadowBottom", FVector2D(16, 8), FMargin(0.5, 0, 0.5, 1)))
+			.SetLeftShadowBrush(BOX_BRUSH("Common/ScrollBoxShadowLeft", FVector2D(8, 16), FMargin(1, 0.5, 0, 0.5)))
+			.SetRightShadowBrush(BOX_BRUSH("Common/ScrollBoxShadowRight", FVector2D(8, 16), FMargin(0, 0.5, 1, 0.5)))
 			);
 	}
 
 	// SScrollBorder defaults...
 	{
 		Style->Set( "ScrollBorder", FScrollBorderStyle()
-			.SetTopShadowBrush( IMAGE_BRUSH( "Common/ScrollBorderShadowTop", FVector2D( 64, 8 ) ) )
-			.SetBottomShadowBrush( IMAGE_BRUSH( "Common/ScrollBorderShadowBottom", FVector2D( 64, 8 ) ) )
+			.SetTopShadowBrush(BOX_BRUSH("Common/ScrollBorderShadowTop", FVector2D(16, 8), FMargin(0.5, 1, 0.5, 0)))
+			.SetBottomShadowBrush(BOX_BRUSH("Common/ScrollBorderShadowBottom", FVector2D(16, 8), FMargin(0.5, 0, 0.5, 1)))
 			);
 	}
 
