UE-2474 Show phys asset in skelmesh asset tooltip (as orange if none)

diff --git a/Engine/Source/Editor/ContentBrowser/Private/AssetViewWidgets.cpp b/Engine/Source/Editor/ContentBrowser/Private/AssetViewWidgets.cpp
--- a/Engine/Source/Editor/ContentBrowser/Private/AssetViewWidgets.cpp
+++ b/Engine/Source/Editor/ContentBrowser/Private/AssetViewWidgets.cpp
@@ -147,6 +147,10 @@ void SAssetViewItem::Construct( const FArguments& InArgs )
 	ImportantStaticMeshTags.Add("CollisionPrims", TEXT("0"));
 	ImportantTagMap.Add("StaticMesh", ImportantStaticMeshTags);
 
+	TMap<FName, FString> ImportantSkelMeshTags;
+	ImportantSkelMeshTags.Add("PhysicsAsset", TEXT("None"));
+	ImportantTagMap.Add("SkeletalMesh", ImportantSkelMeshTags);
+
 	AssetDirtyBrush = FEditorStyle::GetBrush("ContentBrowser.ContentDirty");
 	SCCStateBrush = nullptr;
 
diff --git a/Engine/Source/Runtime/Engine/Classes/Engine/SkeletalMesh.h b/Engine/Source/Runtime/Engine/Classes/Engine/SkeletalMesh.h
--- a/Engine/Source/Runtime/Engine/Classes/Engine/SkeletalMesh.h
+++ b/Engine/Source/Runtime/Engine/Classes/Engine/SkeletalMesh.h
@@ -630,7 +630,7 @@ public:
 	 *	Physics and collision information used for this USkeletalMesh, set up in PhAT.
 	 *	This is used for per-bone hit detection, accurate bounding box calculation and ragdoll physics for example.
 	 */
-	UPROPERTY(EditAnywhere, BlueprintReadOnly, Category=Physics)
+	UPROPERTY(EditAnywhere, AssetRegistrySearchable, BlueprintReadOnly, Category=Physics)
 	class UPhysicsAsset* PhysicsAsset;
 
 #if WITH_EDITORONLY_DATA
