Editor: Add dummy icons to some asset context menu options to provide consistent padding in the sections they appear in

diff --git a/Engine/Source/Editor/ContentBrowser/Private/AssetContextMenu.cpp b/Engine/Source/Editor/ContentBrowser/Private/AssetContextMenu.cpp
--- a/Engine/Source/Editor/ContentBrowser/Private/AssetContextMenu.cpp
+++ b/Engine/Source/Editor/ContentBrowser/Private/AssetContextMenu.cpp
@@ -168,7 +168,7 @@ bool FAssetContextMenu::AddImportedAssetMenuOptions(FMenuBuilder& MenuBuilder)
 			MenuBuilder.AddMenuEntry(
 				LOCTEXT("FindSourceFile", "Open Source Location"),
 				LOCTEXT("FindSourceFileTooltip", "Opens the folder containing the source of the selected asset(s)."),
-				FSlateIcon(),
+				FSlateIcon(FEditorStyle::GetStyleSetName(), "ContentBrowser.AssetActions.OpenSourceLocation"),
 				FUIAction(
 					FExecuteAction::CreateSP(this, &FAssetContextMenu::ExecuteFindSourceInExplorer, ResolvedFilePaths),
 					FCanExecuteAction::CreateSP(this, &FAssetContextMenu::CanExecuteImportedAssetActions, ResolvedFilePaths)
@@ -179,7 +179,7 @@ bool FAssetContextMenu::AddImportedAssetMenuOptions(FMenuBuilder& MenuBuilder)
 			MenuBuilder.AddMenuEntry(
 				LOCTEXT("OpenInExternalEditor", "Open In External Editor"),
 				LOCTEXT("OpenInExternalEditorTooltip", "Open the selected asset(s) in the default external editor."),
-				FSlateIcon(),
+				FSlateIcon(FEditorStyle::GetStyleSetName(), "ContentBrowser.AssetActions.OpenInExternalEditor"),
 				FUIAction(
 					FExecuteAction::CreateSP(this, &FAssetContextMenu::ExecuteOpenInExternalEditor, ResolvedFilePaths),
 					FCanExecuteAction::CreateSP(this, &FAssetContextMenu::CanExecuteImportedAssetActions, ResolvedFilePaths)
@@ -210,13 +210,15 @@ bool FAssetContextMenu::AddCommonMenuOptions(FMenuBuilder& MenuBuilder)
 	// Rename
 	MenuBuilder.AddMenuEntry(FGenericCommands::Get().Rename, NAME_None,
 		LOCTEXT("Rename", "Rename"),
-		LOCTEXT("RenameTooltip", "Rename the selected asset.")
+		LOCTEXT("RenameTooltip", "Rename the selected asset."),
+		FSlateIcon(FEditorStyle::GetStyleSetName(), "ContentBrowser.AssetActions.Rename")
 		);
 
 	// Duplicate
 	MenuBuilder.AddMenuEntry(FGenericCommands::Get().Duplicate, NAME_None,
 		LOCTEXT("Duplicate", "Duplicate"),
-		LOCTEXT("DuplicateTooltip", "Create a copy of the selected asset(s).")
+		LOCTEXT("DuplicateTooltip", "Create a copy of the selected asset(s)."),
+		FSlateIcon(FEditorStyle::GetStyleSetName(), "ContentBrowser.AssetActions.Duplicate")
 		);
 
 	// Save
diff --git a/Engine/Source/Editor/EditorStyle/Private/SlateEditorStyle.cpp b/Engine/Source/Editor/EditorStyle/Private/SlateEditorStyle.cpp
--- a/Engine/Source/Editor/EditorStyle/Private/SlateEditorStyle.cpp
+++ b/Engine/Source/Editor/EditorStyle/Private/SlateEditorStyle.cpp
@@ -5080,6 +5080,10 @@ void FSlateEditorStyle::FStyle::SetupContentBrowserStyle()
 		// Asset Context Menu
 		Set( "ContentBrowser.AssetActions.Edit", new IMAGE_BRUSH( "Icons/icon_Editor_Modes_16x", Icon16x16 ) );
 		Set( "ContentBrowser.AssetActions.Delete", new IMAGE_BRUSH( "Icons/icon_delete_16px", Icon16x16, FLinearColor( 0.4f, 0.5f, 0.7f, 1.0f ) ) );
+		Set( "ContentBrowser.AssetActions.Rename", new FSlateNoResource(Icon16x16) ); //@TODO: Placeholder to preserve indentation in multibox
+		Set( "ContentBrowser.AssetActions.Duplicate", new FSlateNoResource(Icon16x16) ); //@TODO: Placeholder to preserve indentation in multibox
+		Set( "ContentBrowser.AssetActions.OpenSourceLocation", new FSlateNoResource(Icon16x16) ); //@TODO: Placeholder to preserve indentation in multibox
+		Set( "ContentBrowser.AssetActions.OpenInExternalEditor", new FSlateNoResource(Icon16x16) ); //@TODO: Placeholder to preserve indentation in multibox
 		Set( "ContentBrowser.AssetActions.ReimportAsset", new IMAGE_BRUSH( "Icons/icon_TextureEd_Reimport_40x", Icon16x16 ) );
 		Set( "ContentBrowser.AssetActions", new IMAGE_BRUSH( "Icons/icon_tab_Tools_16x", Icon16x16 ) );
 		Set( "ContentBrowser.AssetActions.GoToCodeForAsset", new IMAGE_BRUSH( "GameProjectDialog/feature_code_32x", Icon16x16 ) );
