#UE4 DOCS: Adjusted the names of the Show and Hide Sources Panel Buttons so they could be highlighed individually in the tutorial system.

diff --git a/Engine/Source/Editor/ContentBrowser/Private/SContentBrowser.cpp b/Engine/Source/Editor/ContentBrowser/Private/SContentBrowser.cpp
--- a/Engine/Source/Editor/ContentBrowser/Private/SContentBrowser.cpp
+++ b/Engine/Source/Editor/ContentBrowser/Private/SContentBrowser.cpp
@@ -427,7 +427,7 @@ void SContentBrowser::Construct( const FArguments& InArgs, const FName& InInstan
 							.SearchContent()
 							[
 								SNew( SVerticalBox )
-								.AddMetaData<FTagMetaData>(FTagMetaData(TEXT("ContentBrowserSourcesToggle")))
+								.AddMetaData<FTagMetaData>(FTagMetaData(TEXT("ContentBrowserSourcesToggle1")))
 								+ SVerticalBox::Slot()
 								.FillHeight( 1.0f )
 								.Padding(0,0,2,0)
@@ -486,7 +486,7 @@ void SContentBrowser::Construct( const FArguments& InArgs, const FName& InInstan
 						.Padding( 0, 0, 4, 0 )
 						[
 							SNew( SVerticalBox )
-							.AddMetaData<FTagMetaData>(FTagMetaData(TEXT("ContentBrowserSourcesToggle")))
+							.AddMetaData<FTagMetaData>(FTagMetaData(TEXT("ContentBrowserSourcesToggle2")))
 							+ SVerticalBox::Slot()
 							.FillHeight( 1.0f )
 							[
