Allowed enabling of 'GetContent' in via Experimental menu

diff --git a/Engine/Source/Editor/ContentBrowser/Private/SContentBrowser.cpp b/Engine/Source/Editor/ContentBrowser/Private/SContentBrowser.cpp
--- a/Engine/Source/Editor/ContentBrowser/Private/SContentBrowser.cpp
+++ b/Engine/Source/Editor/ContentBrowser/Private/SContentBrowser.cpp
@@ -186,7 +186,7 @@ void SContentBrowser::Construct( const FArguments& InArgs, const FName& InInstan
 							.ButtonStyle(FEditorStyle::Get(), "ToggleButton")
 							.ToolTipText(LOCTEXT("AddContentTooltip", "Get more content."))
 							.ContentPadding(0)
-							.Visibility(EVisibility::Collapsed) // TODO: Make the button visible once there are actually content packs available.
+							.Visibility(UEditorExperimentalSettings::StaticClass()->GetDefaultObject<UEditorExperimentalSettings>()->bGetFeatureContent == false ? EVisibility::Collapsed : EVisibility::Visible)
 							.OnClicked(this, &SContentBrowser::OnAddContentClicked)
 							[
 								SNew(SHorizontalBox)
diff --git a/Engine/Source/Editor/UnrealEd/Classes/Settings/EditorExperimentalSettings.h b/Engine/Source/Editor/UnrealEd/Classes/Settings/EditorExperimentalSettings.h
--- a/Engine/Source/Editor/UnrealEd/Classes/Settings/EditorExperimentalSettings.h
+++ b/Engine/Source/Editor/UnrealEd/Classes/Settings/EditorExperimentalSettings.h
@@ -93,6 +93,10 @@ public:
 	UPROPERTY(EditAnywhere, config, Category = Tools, meta = (DisplayName = "Gameplay Debugger"))
 	bool bGameplayDebugger;
 
+	/** Enables GetContent button in asset browser */
+	UPROPERTY(EditAnywhere, config, Category = Tools, meta = (DisplayName = "Get Feature Content (Requires content browser reopening)"))
+	bool bGetFeatureContent;
+
 	/**
 	 * Returns an event delegate that is executed when a setting has changed.
 	 *
