Added path text display to the SAssetDialog widget.

diff --git a/Engine/Source/Editor/ContentBrowser/Private/SAssetDialog.cpp b/Engine/Source/Editor/ContentBrowser/Private/SAssetDialog.cpp
--- a/Engine/Source/Editor/ContentBrowser/Private/SAssetDialog.cpp
+++ b/Engine/Source/Editor/ContentBrowser/Private/SAssetDialog.cpp
@@ -138,24 +138,39 @@ void SAssetDialog::Construct(const FArguments& InArgs, const FSharedAssetDialogC
 		];
 	}
 
-	// Buttons and asset name
-	TSharedRef<SHorizontalBox> ButtonsAndNameBox = SNew(SHorizontalBox);
+	TSharedRef<SVerticalBox> LabelsBox = SNew(SVerticalBox)
+		+ SVerticalBox::Slot()
+		.FillHeight(1)
+		.VAlign(VAlign_Center)
+		.Padding(0, 2, 0, 2)
+		[
+			SNew(STextBlock).Text(LOCTEXT("PathBoxLabel", "Path:"))
+		];
+
+	TSharedRef<SVerticalBox> ContentBox = SNew(SVerticalBox)
+		+ SVerticalBox::Slot()
+		.FillHeight(1)
+		.VAlign(VAlign_Center)
+		.Padding(0, 2, 0, 2)
+		[
+			SAssignNew(PathText, STextBlock)
+			.Text(this, &SAssetDialog::GetPathNameText)
+		];
 
-	if ( bIncludeNameBox )
+	if (bIncludeNameBox)
 	{
-		ButtonsAndNameBox->AddSlot()
-			.AutoWidth()
-			.HAlign(HAlign_Right)
+		LabelsBox->AddSlot()
+			.FillHeight(1)
 			.VAlign(VAlign_Center)
-			.Padding(80, 3, 4, 3)
+			.Padding(0, 2, 0, 2)
 			[
 				SNew(STextBlock).Text(LOCTEXT("NameBoxLabel", "Name:"))
 			];
 
-		ButtonsAndNameBox->AddSlot()
-			.FillWidth(1)
+		ContentBox->AddSlot()
+			.FillHeight(1)
 			.VAlign(VAlign_Center)
-			.Padding(4, 3)
+			.Padding(0, 2, 0, 2)
 			[
 				SAssignNew(NameEditableText, SEditableTextBox)
 				.Text(this, &SAssetDialog::GetAssetNameText)
@@ -165,20 +180,36 @@ void SAssetDialog::Construct(const FArguments& InArgs, const FSharedAssetDialogC
 			];
 	}
 
-	ButtonsAndNameBox->AddSlot()
+	// Buttons and asset name
+	TSharedRef<SHorizontalBox> ButtonsAndNameBox = SNew(SHorizontalBox)
+		+ SHorizontalBox::Slot()
 		.AutoWidth()
-		.VAlign(VAlign_Center)
+		.HAlign(HAlign_Right)
+		.VAlign(VAlign_Fill)
+		.Padding(bIncludeNameBox ? 80 : 4, 3, 4, 3)
+		[
+			LabelsBox
+		]
+		+ SHorizontalBox::Slot()
+		.FillWidth(1)
+		.VAlign(VAlign_Fill)
+		.Padding(4, 3)
+		[
+			ContentBox
+		]
+		+ SHorizontalBox::Slot()
+		.AutoWidth()
+		.VAlign(VAlign_Bottom)
 		.Padding(4, 3)
 		[
 			SNew(SButton)
 			.Text(ConfirmButtonText)
 			.IsEnabled(this, &SAssetDialog::IsConfirmButtonEnabled)
 			.OnClicked(this, &SAssetDialog::OnConfirmClicked)
-		];
-
-	ButtonsAndNameBox->AddSlot()
+		]
+		+ SHorizontalBox::Slot()
 		.AutoWidth()
-		.VAlign(VAlign_Center)
+		.VAlign(VAlign_Bottom)
 		.Padding(4, 3)
 		[
 			SNew(SButton)
@@ -188,7 +219,7 @@ void SAssetDialog::Construct(const FArguments& InArgs, const FSharedAssetDialogC
 
 	MainVerticalBox->AddSlot()
 		.AutoHeight()
-		.HAlign(bIncludeNameBox ? HAlign_Fill : HAlign_Right)
+		.HAlign(HAlign_Fill)
 		.Padding(8, 2, 8, 4)
 		[
 			SNew(SBorder)
@@ -255,6 +286,11 @@ FText SAssetDialog::GetAssetNameText() const
 	return FText::FromString(CurrentlyEnteredAssetName);
 }
 
+FText SAssetDialog::GetPathNameText() const
+{
+	return FText::FromString(CurrentlySelectedPath);
+}
+
 void SAssetDialog::OnAssetNameTextCommited(const FText& InText, ETextCommit::Type InCommitType)
 {
 	SetCurrentlyEnteredAssetName(InText.ToString());
