Additional file which was supposed to be in CL 2383463.

diff --git a/Engine/Source/Editor/ContentBrowser/Private/SAssetDialog.h b/Engine/Source/Editor/ContentBrowser/Private/SAssetDialog.h
--- a/Engine/Source/Editor/ContentBrowser/Private/SAssetDialog.h
+++ b/Engine/Source/Editor/ContentBrowser/Private/SAssetDialog.h
@@ -34,6 +34,9 @@ private:
 	/** Gets the name to display in the asset name box */
 	FText GetAssetNameText() const;
 
+	/** Gets the name to display in the path text block */
+	FText GetPathNameText() const;
+
 	/** Fired when the asset name box text is commited */
 	void OnAssetNameTextCommited(const FText& InText, ETextCommit::Type InCommitType);
 
@@ -113,6 +116,9 @@ private:
 	/** The name box. Only used in save dialogs. */
 	TSharedPtr<SEditableTextBox> NameEditableText;
 
+	/** The path box. */
+	TSharedPtr<STextBlock> PathText;
+
 	/** The object path of the asset to save. Only used in save dialogs. */
 	FString CurrentlySelectedPath;
 
