Added tooltip support to content browser header row columns.  Added AssetRegistryTagMetadata class to Editor builds, for reporting additional tag information.

diff --git a/Engine/Source/Editor/ContentBrowser/Private/AssetViewWidgets.cpp b/Engine/Source/Editor/ContentBrowser/Private/AssetViewWidgets.cpp
--- a/Engine/Source/Editor/ContentBrowser/Private/AssetViewWidgets.cpp
+++ b/Engine/Source/Editor/ContentBrowser/Private/AssetViewWidgets.cpp
@@ -143,14 +143,6 @@ void SAssetViewItem::Construct( const FArguments& InArgs )
 
 	AssetItem->OnAssetDataChanged.AddSP(this, &SAssetViewItem::OnAssetDataChanged);
 
-	TMap<FName, FString> ImportantStaticMeshTags;
-	ImportantStaticMeshTags.Add("CollisionPrims", TEXT("0"));
-	ImportantTagMap.Add("StaticMesh", ImportantStaticMeshTags);
-
-	TMap<FName, FString> ImportantSkelMeshTags;
-	ImportantSkelMeshTags.Add("PhysicsAsset", TEXT("None"));
-	ImportantTagMap.Add("SkeletalMesh", ImportantSkelMeshTags);
-
 	AssetDirtyBrush = FEditorStyle::GetBrush("ContentBrowser.ContentDirty");
 	SCCStateBrush = nullptr;
 
@@ -467,9 +459,12 @@ TSharedRef<SToolTip> SAssetViewItem::CreateToolTipWidget() const
 			AddToToolTipInfoBox( InfoBox, LOCTEXT("TileViewTooltipPath", "Path"), FText::FromName(AssetData.PackagePath), false );
 
 			// If we are using a loaded class, find all the hidden tags so we don't display them
+			TMap<FName, UObject::FAssetRegistryTagMetadata> MetadataMap;
 			TSet<FName> ShownTags;
 			if ( AssetClass != NULL && AssetClass->GetDefaultObject() != NULL )
 			{
+				AssetClass->GetDefaultObject()->GetAssetRegistryTagMetadata(MetadataMap);
+
 				TArray<UObject::FAssetRegistryTag> Tags;
 				AssetClass->GetDefaultObject()->GetAssetRegistryTags(Tags);
 
@@ -482,9 +477,6 @@ TSharedRef<SToolTip> SAssetViewItem::CreateToolTipWidget() const
 				}
 			}
 
-			// Get the list of important tags for this class
-			TMap<FName, FString> ImportantTags = ImportantTagMap.FindRef(AssetData.AssetClass);
-
 			// If an asset class could not be loaded we cannot determine hidden tags so display no tags.  
 			if( AssetClass != NULL )
 			{
@@ -494,8 +486,9 @@ TSharedRef<SToolTip> SAssetViewItem::CreateToolTipWidget() const
 					// Skip tags that are set to be hidden
 					if ( ShownTags.Contains(TagIt.Key()) )
 					{
-						const FString* ImportantValue = ImportantTags.Find(TagIt.Key());
-						const bool bImportant = (ImportantValue && (*ImportantValue) == TagIt.Value());
+						const UObject::FAssetRegistryTagMetadata* Metadata = MetadataMap.Find(TagIt.Key());
+
+						const bool bImportant = (Metadata != nullptr && !Metadata->ImportantValue.IsEmpty() && Metadata->ImportantValue == TagIt.Value());
 
 						// Since all we have at this point is a string, we can't be very smart here.
 						// We need to strip some noise off class paths in some cases, but can't load the asset to inspect its UPROPERTYs manually due to performance concerns.
@@ -536,6 +529,13 @@ TSharedRef<SToolTip> SAssetViewItem::CreateToolTipWidget() const
 							DisplayName = FText::FromString(FName::NameToDisplayString(TagIt.Key().ToString(), bIsBool));
 						}
 					
+						// Add suffix to the value string, if one is defined for this tag
+						if (Metadata != nullptr && !Metadata->Suffix.IsEmpty())
+						{
+							ValueString += TEXT(" ");
+							ValueString += Metadata->Suffix.ToString();
+						}
+
 						AddToToolTipInfoBox(InfoBox, DisplayName, FText::FromString(ValueString), bImportant);
 					}
 				}
diff --git a/Engine/Source/Editor/ContentBrowser/Private/AssetViewWidgets.h b/Engine/Source/Editor/ContentBrowser/Private/AssetViewWidgets.h
--- a/Engine/Source/Editor/ContentBrowser/Private/AssetViewWidgets.h
+++ b/Engine/Source/Editor/ContentBrowser/Private/AssetViewWidgets.h
@@ -233,13 +233,6 @@ protected:
 	/** The substring to be highlighted in the name and tooltip path */
 	TAttribute<FText> HighlightText;
 
-	/**
-	 * A map of class names to their important asset registry tags and values.
-	 * Important tags will be featured in the tooltip if the value matches.
-	 * An empty string is interpreted as indicating that any value is always important.
-	 */
-	TMap< FName, TMap<FName, FString> > ImportantTagMap;
-
 	/** Cached brushes for the dirty state */
 	const FSlateBrush* AssetDirtyBrush;
 
diff --git a/Engine/Source/Editor/ContentBrowser/Private/SAssetView.cpp b/Engine/Source/Editor/ContentBrowser/Private/SAssetView.cpp
--- a/Engine/Source/Editor/ContentBrowser/Private/SAssetView.cpp
+++ b/Engine/Source/Editor/ContentBrowser/Private/SAssetView.cpp
@@ -190,8 +190,6 @@ void SAssetView::Construct( const FArguments& InArgs )
 	bPendingFocusOnSync = false;
 	bWereItemsRecursivelyFiltered = false;
 
-	TagColumnRenames.Add("ResourceSize", TEXT("Size (kb)"));
-
 	FEditorWidgetsModule& EditorWidgetsModule = FModuleManager::LoadModuleChecked<FEditorWidgetsModule>("EditorWidgets");
 	TSharedRef<SWidget> AssetDiscoveryIndicator = EditorWidgetsModule.CreateAssetDiscoveryIndicator(EAssetDiscoveryIndicatorScaleMode::Scale_Vertical);
 
@@ -1700,23 +1698,40 @@ void SAssetView::SetMajorityAssetType(FName NewMajorityAssetType)
 
 							if ( !OnAssetTagWantsToBeDisplayed.IsBound() || OnAssetTagWantsToBeDisplayed.Execute(NewMajorityAssetType, Tag) )
 							{
-								const FString* DisplayNamePtr = TagColumnRenames.Find(Tag);
+								// Get tag metadata
+								TMap<FName, UObject::FAssetRegistryTagMetadata> MetadataMap;
+								CDO->GetAssetRegistryTagMetadata(MetadataMap);
+								const UObject::FAssetRegistryTagMetadata* Metadata = MetadataMap.Find(Tag);
+
 								FText DisplayName;
-								if ( DisplayNamePtr )
+								if (Metadata != nullptr && !Metadata->DisplayName.IsEmpty())
 								{
-									DisplayName = FText::FromString(*DisplayNamePtr);
+									DisplayName = Metadata->DisplayName;
 								}
 								else
 								{
 									DisplayName = FText::FromName(Tag);
 								}
 
+								FText TooltipText;
+								if (Metadata != nullptr && !Metadata->TooltipText.IsEmpty())
+								{
+									TooltipText = Metadata->TooltipText;
+								}
+								else
+								{
+									// If the tag name corresponds to a property name, use the property tooltip
+									UProperty* Property = FindField<UProperty>(TypeClass, Tag);
+									TooltipText = (Property != nullptr) ? Property->GetToolTipText() : FText::FromString(FName::NameToDisplayString(Tag.ToString(), false));
+								}
+
 								ColumnView->GetHeaderRow()->AddColumn(
 										SHeaderRow::Column(Tag)
 										.SortMode( TAttribute< EColumnSortMode::Type >::Create( TAttribute< EColumnSortMode::Type >::FGetter::CreateSP( this, &SAssetView::GetColumnSortMode, Tag ) ) )
 										.SortPriority(TAttribute< EColumnSortPriority::Type >::Create(TAttribute< EColumnSortPriority::Type >::FGetter::CreateSP(this, &SAssetView::GetColumnSortPriority, Tag)))
 										.OnSort( FOnSortModeChanged::CreateSP( this, &SAssetView::OnSortColumnHeader ) )
 										.DefaultLabel( DisplayName )
+										.DefaultTooltip( TooltipText )
 										.HAlignCell( (TagIt->Type == UObject::FAssetRegistryTag::TT_Numerical) ? HAlign_Right : HAlign_Left )
 										.FillWidth(180)
 									);
diff --git a/Engine/Source/Editor/ContentBrowser/Private/SAssetView.h b/Engine/Source/Editor/ContentBrowser/Private/SAssetView.h
--- a/Engine/Source/Editor/ContentBrowser/Private/SAssetView.h
+++ b/Engine/Source/Editor/ContentBrowser/Private/SAssetView.h
@@ -722,9 +722,6 @@ private:
 	/** When in columns view, this is the name of the asset type which is most commonly found in the recent results */
 	FName MajorityAssetType;
 
-	/** The map of Tag names to display names in column headers. If a tag is not found in this map, it will use the string version of the name, which is fine most of the time */
-	TMap<FName, FString> TagColumnRenames;
-
 	/** The manager responsible for sorting assets in the view */
 	FAssetViewSortManager SortManager;
 
diff --git a/Engine/Source/Runtime/CoreUObject/Private/UObject/Obj.cpp b/Engine/Source/Runtime/CoreUObject/Private/UObject/Obj.cpp
--- a/Engine/Source/Runtime/CoreUObject/Private/UObject/Obj.cpp
+++ b/Engine/Source/Runtime/CoreUObject/Private/UObject/Obj.cpp
@@ -1181,6 +1181,18 @@ void UObject::GetAssetRegistryTags(TArray<FAssetRegistryTag>& OutTags) const
 	FAssetRegistryTag::GetAssetRegistryTagsFromSearchableProperties(this, OutTags);
 }
 
+#if WITH_EDITOR
+void UObject::GetAssetRegistryTagMetadata(TMap<FName, FAssetRegistryTagMetadata>& OutMetadata) const
+{
+	OutMetadata.Add("ResourceSize",
+		FAssetRegistryTagMetadata()
+			.SetDisplayName(NSLOCTEXT("UObject", "Size", "Size"))
+			.SetSuffix(NSLOCTEXT("UObject", "KilobytesSuffix", "Kb"))
+			.SetTooltip(NSLOCTEXT("UObject", "SizeTooltip", "The size of the asset in kilobytes"))
+		);
+}
+#endif
+
 bool UObject::IsAsset () const
 {
 	// Assets are not transient or CDOs. They must be public.
diff --git a/Engine/Source/Runtime/CoreUObject/Public/UObject/UObject.h b/Engine/Source/Runtime/CoreUObject/Public/UObject/UObject.h
--- a/Engine/Source/Runtime/CoreUObject/Public/UObject/UObject.h
+++ b/Engine/Source/Runtime/CoreUObject/Public/UObject/UObject.h
@@ -519,6 +519,50 @@ public:
 	};
 	virtual void GetAssetRegistryTags(TArray<FAssetRegistryTag>& OutTags) const;
 
+#if WITH_EDITOR
+	/**
+	 * Additional data pertaining to asset registry tags used by the editor
+	 */
+	struct FAssetRegistryTagMetadata
+	{
+		FText DisplayName;
+		FText TooltipText;
+		FText Suffix;
+		FString ImportantValue;
+
+		/** Set override display name */
+		FAssetRegistryTagMetadata& SetDisplayName(const FText& InDisplayName)
+		{
+			DisplayName = InDisplayName;
+			return *this;
+		}
+
+		/** Set tooltip text pertaining to the asset registry tag in the column view header */
+		FAssetRegistryTagMetadata& SetTooltip(const FText& InTooltipText)
+		{
+			TooltipText = InTooltipText;
+			return *this;
+		}
+
+		/** Set suffix appended to the tag value */
+		FAssetRegistryTagMetadata& SetSuffix(const FText& InSuffix)
+		{
+			Suffix = InSuffix;
+			return *this;
+		}
+
+		/** Set value deemed to be 'important' for this registry tag */
+		FAssetRegistryTagMetadata& SetImportantValue(const FString& InImportantValue)
+		{
+			ImportantValue = InImportantValue;
+			return *this;
+		}
+	};
+
+	/** Gathers a collection of asset registry tag metadata */
+	virtual void GetAssetRegistryTagMetadata(TMap<FName, FAssetRegistryTagMetadata>& OutMetadata) const;
+#endif
+
 	/** Returns true if this object is considered an asset. */
 	virtual bool IsAsset() const;
 
diff --git a/Engine/Source/Runtime/Engine/Classes/Engine/SkeletalMesh.h b/Engine/Source/Runtime/Engine/Classes/Engine/SkeletalMesh.h
--- a/Engine/Source/Runtime/Engine/Classes/Engine/SkeletalMesh.h
+++ b/Engine/Source/Runtime/Engine/Classes/Engine/SkeletalMesh.h
@@ -726,6 +726,7 @@ public:
 	virtual void PreEditChange(UProperty* PropertyAboutToChange) override;
 	virtual void PostEditChangeProperty(FPropertyChangedEvent& PropertyChangedEvent) override;
 	virtual void PostEditUndo() override;
+	virtual void GetAssetRegistryTagMetadata(TMap<FName, FAssetRegistryTagMetadata>& OutMetadata) const override;
 #endif // WITH_EDITOR
 	virtual void BeginDestroy() override;
 	virtual bool IsReadyForFinishDestroy() override;
diff --git a/Engine/Source/Runtime/Engine/Classes/Engine/StaticMesh.h b/Engine/Source/Runtime/Engine/Classes/Engine/StaticMesh.h
--- a/Engine/Source/Runtime/Engine/Classes/Engine/StaticMesh.h
+++ b/Engine/Source/Runtime/Engine/Classes/Engine/StaticMesh.h
@@ -444,6 +444,7 @@ public:
 #if WITH_EDITOR
 	ENGINE_API virtual void PreEditChange(UProperty* PropertyAboutToChange) override;
 	ENGINE_API virtual void PostEditChangeProperty(FPropertyChangedEvent& PropertyChangedEvent) override;
+	ENGINE_API virtual void GetAssetRegistryTagMetadata(TMap<FName, FAssetRegistryTagMetadata>& OutMetadata) const override;
 #endif // WITH_EDITOR
 	ENGINE_API virtual void Serialize(FArchive& Ar) override;
 	ENGINE_API virtual void PostLoad() override;
diff --git a/Engine/Source/Runtime/Engine/Private/StaticMesh.cpp b/Engine/Source/Runtime/Engine/Private/StaticMesh.cpp
--- a/Engine/Source/Runtime/Engine/Private/StaticMesh.cpp
+++ b/Engine/Source/Runtime/Engine/Private/StaticMesh.cpp
@@ -1491,6 +1491,18 @@ void UStaticMesh::GetAssetRegistryTags(TArray<FAssetRegistryTag>& OutTags) const
 	Super::GetAssetRegistryTags(OutTags);
 }
 
+#if WITH_EDITOR
+void UStaticMesh::GetAssetRegistryTagMetadata(TMap<FName, FAssetRegistryTagMetadata>& OutMetadata) const
+{
+	OutMetadata.Add("CollisionPrims",
+		FAssetRegistryTagMetadata()
+			.SetTooltip(NSLOCTEXT("UStaticMesh", "CollisionPrimsTooltip", "The number of collision primitives in the static mesh"))
+			.SetImportantValue(TEXT("0"))
+		);
+}
+#endif
+
+
 /*------------------------------------------------------------------------------
 	FStaticMeshSourceModel
 ------------------------------------------------------------------------------*/
diff --git a/Engine/Source/Runtime/Slate/Private/Widgets/Views/SHeaderRow.cpp b/Engine/Source/Runtime/Slate/Private/Widgets/Views/SHeaderRow.cpp
--- a/Engine/Source/Runtime/Slate/Private/Widgets/Views/SHeaderRow.cpp
+++ b/Engine/Source/Runtime/Slate/Private/Widgets/Views/SHeaderRow.cpp
@@ -49,10 +49,19 @@ public:
 		FMargin AdjustedDefaultHeaderContentPadding = DefaultHeaderContentPadding;
 
 		TAttribute< FText > LabelText = Column.DefaultText;
+		TAttribute< FText > TooltipText = Column.DefaultTooltip;
 
-		if (Column.HeaderContent.Widget == SNullWidget::NullWidget && !Column.DefaultText.IsBound() && Column.DefaultText.Get().IsEmpty())
+		if (Column.HeaderContent.Widget == SNullWidget::NullWidget)
 		{
-			LabelText = FText::FromString( Column.ColumnId.ToString() + TEXT("[LabelMissing]") );
+			if (!Column.DefaultText.IsBound() && Column.DefaultText.Get().IsEmpty())
+			{
+				LabelText = FText::FromString( Column.ColumnId.ToString() + TEXT("[LabelMissing]") );
+			}
+
+			if (!Column.DefaultTooltip.IsBound() && Column.DefaultTooltip.Get().IsEmpty())
+			{
+				TooltipText = LabelText;
+			}
 		}
 
 		TSharedPtr< SHorizontalBox > Box;
@@ -73,7 +82,7 @@ public:
 				[
 					SNew(STextBlock)
 					.Text( LabelText )
-					.ToolTipText( LabelText )
+					.ToolTipText( TooltipText )
 				];
 		}
 
diff --git a/Engine/Source/Runtime/Slate/Public/Widgets/Views/SHeaderRow.h b/Engine/Source/Runtime/Slate/Public/Widgets/Views/SHeaderRow.h
--- a/Engine/Source/Runtime/Slate/Public/Widgets/Views/SHeaderRow.h
+++ b/Engine/Source/Runtime/Slate/Public/Widgets/Views/SHeaderRow.h
@@ -82,6 +82,7 @@ public:
 		HACK_SLATE_SLOT_ARGS(FColumn)
 			: _ColumnId()
 			, _DefaultLabel()
+			, _DefaultTooltip()
 			, _FillWidth( 1.0f )
 			, _FixedWidth()
 			, _OnWidthChanged()
@@ -98,6 +99,7 @@ public:
 			{}
 			SLATE_ARGUMENT( FName, ColumnId )
 			SLATE_TEXT_ATTRIBUTE( DefaultLabel )
+			SLATE_TEXT_ATTRIBUTE( DefaultTooltip )
 			SLATE_ATTRIBUTE( float, FillWidth )
 			SLATE_ARGUMENT( TOptional< float >, FixedWidth )
 			SLATE_EVENT( FOnWidthChanged, OnWidthChanged )
@@ -120,6 +122,7 @@ public:
 		FColumn( const FArguments& InArgs )
 			: ColumnId( InArgs._ColumnId )
 			, DefaultText( InArgs._DefaultLabel )
+			, DefaultTooltip( InArgs._DefaultTooltip )
 			, Width( 1.0f )
 			, DefaultWidth( 1.0f )
 			, OnWidthChanged( InArgs._OnWidthChanged)
@@ -175,6 +178,9 @@ public:
 		/** Default text to use if no widget is passed in. */
 		TAttribute< FText > DefaultText;
 
+		/** Default tooltip to use if no widget is passed in */
+		TAttribute< FText > DefaultTooltip;
+
 		/** A column width in Slate Units */
 		TAttribute< float > Width;
 
