---- Merging with SlateDev branch ----

Introduces the concept of "Active Ticking" to allow Slate to go to sleep when there is no need to update the UI.

While asleep, Slate will skip the Tick & Paint pass for that frame entirely.
- There are TWO ways to "wake" Slate and cause a Tick/Paint pass:
    1. Provide some sort of input (mouse movement, clicks, and key presses). Slate will always tick when the user is active.
        - Therefore, if the logic in a given widget's Tick is only relevant in response to user action, there is no need to register an active tick.
    2. Register an Active Tick. Currently this is an all-or-nothing situation, so if a single active tick needs to execute, all of Slate will be ticked.

- The purpose of an Active Tick is to allow a widget to "drive" Slate and guarantee a Tick/Paint pass in the absence of any user action.
    - Examples include animation, async operations that update periodically, progress updates, loading bars, etc.

- An empty active tick is registered for viewports when they are real-time, so game project widgets are unaffected by this change and should continue to work as before.

- An Active Tick is registered by creating an FWidgetActiveTickDelegate and passing it to SWidget::RegisterActiveTick()
    - There are THREE ways to unregister an active tick:
        1. Return EActiveTickReturnType::StopTicking from the active tick function
        2. Pass the FActiveTickHandle returned by RegisterActiveTick() to SWidget::UnregisterActiveTick()
        3. Destroy the widget responsible for the active tick

- Sleeping is currently disabled, can be enabled with Slate.AllowSlateToSleep cvar
- There is currently a little buffer time during which Slate continues to tick following any input. Long-term, this is planned to be removed.
    - The duration of the buffer can be adjusted using Slate.SleepBufferPostInput cvar (defaults to 1.0f)

- The FCurveSequence API has been updated to work with the active tick system
    - Playing a curve sequence now requires that you pass the widget being animated by the sequence
    - The active tick will automatically be registered on behalf of the widget and unregister when the sequence is complete
    - GetLerpLooping() has been removed. Instead, pass true as the second param to Play() to indicate that the animation will loop. This causes the active tick to be registered indefinitely until paused or jumped to the start/end.

diff --git a/Engine/Plugins/2D/Paper2D/Source/Paper2DEditor/Private/FlipbookEditor/SFlipbookTrackHandle.h b/Engine/Plugins/2D/Paper2D/Source/Paper2DEditor/Private/FlipbookEditor/SFlipbookTrackHandle.h
--- a/Engine/Plugins/2D/Paper2D/Source/Paper2DEditor/Private/FlipbookEditor/SFlipbookTrackHandle.h
+++ b/Engine/Plugins/2D/Paper2D/Source/Paper2DEditor/Private/FlipbookEditor/SFlipbookTrackHandle.h
@@ -97,7 +97,7 @@ public:
 
 				if (!bDragging)
 				{
-					if (FMath::Abs(DistanceDragged) > FSlateApplication::Get().GetDragTriggerDistnace())
+					if (FMath::Abs(DistanceDragged) > FSlateApplication::Get().GetDragTriggerDistance())
 					{
 						const FPaperFlipbookKeyFrame& KeyFrame = Flipbook->GetKeyFrameChecked(KeyFrameIdx);
 						StartingFrameRun = KeyFrame.FrameRun;
diff --git a/Engine/Plugins/2D/Paper2D/Source/Paper2DEditor/Private/FlipbookEditor/STimelineHeader.h b/Engine/Plugins/2D/Paper2D/Source/Paper2DEditor/Private/FlipbookEditor/STimelineHeader.h
--- a/Engine/Plugins/2D/Paper2D/Source/Paper2DEditor/Private/FlipbookEditor/STimelineHeader.h
+++ b/Engine/Plugins/2D/Paper2D/Source/Paper2DEditor/Private/FlipbookEditor/STimelineHeader.h
@@ -36,10 +36,8 @@ public:
 		Rebuild();
 	}
 
-	virtual void Tick(const FGeometry& AllottedGeometry, const double InCurrentTime, const float InDeltaTime) override
+	virtual void Tick( const FGeometry& AllottedGeometry, const double InCurrentTime, const float InDeltaTime ) override
 	{
-		SCompoundWidget::Tick(AllottedGeometry, InCurrentTime, InDeltaTime);
-
 		UPaperFlipbook* Flipbook = FlipbookBeingEdited.Get();
 		int32 NewNumFrames = (Flipbook != nullptr) ? Flipbook->GetNumFrames() : 0;
 		if (NewNumFrames != NumFramesFromLastRebuild)
diff --git a/Engine/Plugins/2D/Paper2D/Source/Paper2DEditor/Private/FlipbookEditor/STimelineTrack.h b/Engine/Plugins/2D/Paper2D/Source/Paper2DEditor/Private/FlipbookEditor/STimelineTrack.h
--- a/Engine/Plugins/2D/Paper2D/Source/Paper2DEditor/Private/FlipbookEditor/STimelineTrack.h
+++ b/Engine/Plugins/2D/Paper2D/Source/Paper2DEditor/Private/FlipbookEditor/STimelineTrack.h
@@ -109,10 +109,8 @@ public:
 
 	void Construct(const FArguments& InArgs, TSharedPtr<const FUICommandList> InCommandList);
 
-	void Tick(const FGeometry& AllottedGeometry, const double InCurrentTime, const float InDeltaTime) override
+	virtual void Tick( const FGeometry& AllottedGeometry, const double InCurrentTime, const float InDeltaTime ) override
 	{
-		SCompoundWidget::Tick(AllottedGeometry, InCurrentTime, InDeltaTime);
-
 		UPaperFlipbook* Flipbook = FlipbookBeingEdited.Get();
 		int32 NewNumKeyframes = (Flipbook != nullptr) ? Flipbook->GetNumKeyFrames() : 0;
 		if (NewNumKeyframes != NumKeyframesFromLastRebuild)
diff --git a/Engine/Plugins/2D/Paper2D/Source/Paper2DEditor/Private/SPaperEditorViewport.cpp b/Engine/Plugins/2D/Paper2D/Source/Paper2DEditor/Private/SPaperEditorViewport.cpp
--- a/Engine/Plugins/2D/Paper2D/Source/Paper2DEditor/Private/SPaperEditorViewport.cpp
+++ b/Engine/Plugins/2D/Paper2D/Source/Paper2DEditor/Private/SPaperEditorViewport.cpp
@@ -123,16 +123,16 @@ void SPaperEditorViewport::Construct(const FArguments& InArgs, TSharedRef<FPaper
 	bIsPanning = false;
 
 	ZoomLevelFade = FCurveSequence( 0.0f, 0.75f );
-	ZoomLevelFade.Play();
+	ZoomLevelFade.Play(this->AsShared());
 
 	ZoomLevelGraphFade = FCurveSequence( 0.0f, 0.5f );
-	ZoomLevelGraphFade.Play();
+	ZoomLevelGraphFade.Play(this->AsShared());
 
 	DeferredPanPosition = FVector2D::ZeroVector;
 	bRequestDeferredPan = false;
 }
 
-void SPaperEditorViewport::Tick(const FGeometry& AllottedGeometry, const double InCurrentTime, const float InDeltaTime)
+void SPaperEditorViewport::Tick( const FGeometry& AllottedGeometry, const double InCurrentTime, const float InDeltaTime )
 {
 	// Handle any deferred panning
 	if (bRequestDeferredPan)
@@ -147,8 +147,7 @@ void SPaperEditorViewport::Tick(const FGeometry& AllottedGeometry, const double
 		bIsPanning = false;
 	}
 
-
-	ViewportClient->SetZoomPos(ViewOffset, GetZoomAmount());
+	ViewportClient->SetZoomPos( ViewOffset, GetZoomAmount() );
 	ViewportClient->bNeedsRedraw = true;
 
 	bool bSelectionModified = false;
@@ -177,9 +176,6 @@ void SPaperEditorViewport::Tick(const FGeometry& AllottedGeometry, const double
 		ViewportClient->Tick(InDeltaTime);
 		GEditor->UpdateSingleViewportClient(ViewportClient.Get(), /*bInAllowNonRealtimeViewportToDraw=*/ true, /*bLinkedOrthoMovement=*/ false);
 	}
-
-	//
-	SWidget::Tick(AllottedGeometry, InCurrentTime, InDeltaTime);
 }
 
 FReply SPaperEditorViewport::OnMouseButtonDown(const FGeometry& MyGeometry, const FPointerEvent& MouseEvent)
@@ -223,7 +219,7 @@ FReply SPaperEditorViewport::OnMouseButtonUp(const FGeometry& MyGeometry, const
 	// Did the user move the cursor sufficiently far, or is it in a dead zone?
 	// In Dead zone     - implies actions like summoning context menus and general clicking.
 	// Out of Dead Zone - implies dragging actions like moving nodes and marquee selection.
-	const bool bCursorInDeadZone = TotalMouseDelta <= FSlateApplication::Get().GetDragTriggerDistnace();
+	const bool bCursorInDeadZone = TotalMouseDelta <= FSlateApplication::Get().GetDragTriggerDistance();
 
 	if (MouseEvent.GetEffectingButton() == EKeys::RightMouseButton)
 	{
@@ -345,7 +341,7 @@ FReply SPaperEditorViewport::OnMouseMove(const FGeometry& MyGeometry, const FPoi
 			// Update the amount to pan panel
 			UpdateViewOffset(MyGeometry, MouseEvent.GetScreenSpacePosition());
 
-			const bool bCursorInDeadZone = TotalMouseDelta <= FSlateApplication::Get().GetDragTriggerDistnace();
+			const bool bCursorInDeadZone = TotalMouseDelta <= FSlateApplication::Get().GetDragTriggerDistance();
 
 // 			if (NodeBeingDragged.IsValid())
 // 			{
@@ -425,7 +421,7 @@ FReply SPaperEditorViewport::OnMouseWheel(const FGeometry& MyGeometry, const FPo
 		ZoomLevel = FMath::Clamp( ZoomLevel + ZoomLevelDelta, 0, DefaultZoomLevel );
 	}
 
-	ZoomLevelFade.Play();
+	ZoomLevelFade.Play(this->AsShared());
 
 
 	// Re-center the screen so that it feels like zooming around the cursor.
@@ -505,7 +501,7 @@ FVector2D SPaperEditorViewport::ComputeEdgePanAmount(const FGeometry& MyGeometry
 	// Never pan slower than this, it's just unpleasant.
 	static const float MinPanSpeed = 5.0f;
 
-	// Start panning before we rech the edge of the graph panel.
+	// Start panning before we reach the edge of the graph panel.
 	static const float EdgePanForgivenessZone = 30.0f;
 
 	const FVector2D LocalCursorPos = MyGeometry.AbsoluteToLocal( TargetPosition );
diff --git a/Engine/Plugins/2D/Paper2D/Source/Paper2DEditor/Private/SpriteEditor/SSpriteList.cpp b/Engine/Plugins/2D/Paper2D/Source/Paper2DEditor/Private/SpriteEditor/SSpriteList.cpp
--- a/Engine/Plugins/2D/Paper2D/Source/Paper2DEditor/Private/SpriteEditor/SSpriteList.cpp
+++ b/Engine/Plugins/2D/Paper2D/Source/Paper2DEditor/Private/SpriteEditor/SSpriteList.cpp
@@ -59,7 +59,7 @@ void SSpriteList::RebuildWidget(UTexture2D* NewTextureFilter)
 	];
 }
 
-void SSpriteList::Tick(const FGeometry& AllottedGeometry, const double InCurrentTime, const float InDeltaTime)
+void SSpriteList::Tick( const FGeometry& AllottedGeometry, const double InCurrentTime, const float InDeltaTime )
 {
 	TSharedPtr<FSpriteEditor> SpriteEditor = SpriteEditorPtr.Pin();
 	if (SpriteEditor.IsValid())
@@ -72,8 +72,6 @@ void SSpriteList::Tick(const FGeometry& AllottedGeometry, const double InCurrent
 			RebuildWidget(CurrentTexture);
 		}
 	}
-
-	SCompoundWidget::Tick(AllottedGeometry, InCurrentTime, InDeltaTime);
 }
 
 void SSpriteList::OnSpriteSelected(const FAssetData& AssetData)
diff --git a/Engine/Plugins/Editor/EpicSurvey/Source/EpicSurvey/Private/SSurvey.cpp b/Engine/Plugins/Editor/EpicSurvey/Source/EpicSurvey/Private/SSurvey.cpp
--- a/Engine/Plugins/Editor/EpicSurvey/Source/EpicSurvey/Private/SSurvey.cpp
+++ b/Engine/Plugins/Editor/EpicSurvey/Source/EpicSurvey/Private/SSurvey.cpp
@@ -12,7 +12,6 @@
 
 void SSurvey::Construct( const FArguments& Args, const TSharedRef< FEpicSurvey >& InEpicSurvey, const TSharedRef< FSurvey >& InSurvey )
 {
-	FinishedLoading = false;
 	EpicSurvey = InEpicSurvey;
 	Survey = InSurvey;
 	PageBox = NULL;
@@ -47,47 +46,52 @@ void SSurvey::Construct( const FArguments& Args, const TSharedRef< FEpicSurvey >
 		]
 	];
 
+	bool bIsLoaded = false;
 	if ( InSurvey->GetInitializationState() == EContentInitializationState::NotStarted )
 	{
 		InSurvey->Initialize();
 		ConstructLoadingLayout();
-		FinishedLoading = false;
 	}
 	else if ( InSurvey->GetInitializationState() == EContentInitializationState::Working )
 	{
 		ConstructLoadingLayout();
-		FinishedLoading = false;
 	}
 	else if ( InSurvey->GetInitializationState() == EContentInitializationState::Success )
 	{
 		ConstructSurveyLayout();
-		FinishedLoading = true;
+		bIsLoaded = true;
 	}
 	else if ( InSurvey->GetInitializationState() == EContentInitializationState::Failure )
 	{
 		ConstructFailureLayout();
-		FinishedLoading = true;
+		bIsLoaded = true;
+	}
+
+	if (!bIsLoaded)
+	{
+		//Construct the proper layout when finished loading
+		RegisterActiveTick( 0.1f, FWidgetActiveTickDelegate::CreateSP( this, &SSurvey::MonitorLoadStatePostConstruct ) );
 	}
 }
 
-void SSurvey::Tick( const FGeometry& AllottedGeometry, const double InCurrentTime, const float InDeltaTime )
+EActiveTickReturnType SSurvey::MonitorLoadStatePostConstruct( double InCurrentTime, float InDeltaTime )
 {
-	if ( !FinishedLoading )
-	{
-		EContentInitializationState::Type State = Survey->GetInitializationState();
+	EContentInitializationState::Type State = Survey->GetInitializationState();
 
-		if ( State == EContentInitializationState::Success )
-		{
-			FinishedLoading = true;
-			ConstructSurveyLayout();
-		}
-		else if ( State == EContentInitializationState::Failure )
-		{
-			//Show Error message
-			FinishedLoading = true;
-			ConstructFailureLayout();
-		}
+	bool bIsLoaded = false;
+	if ( State == EContentInitializationState::Success )
+	{
+		bIsLoaded = true;
+		ConstructSurveyLayout();
 	}
+	else if ( State == EContentInitializationState::Failure )
+	{
+		//Show Error message
+		bIsLoaded = true;
+		ConstructFailureLayout();
+	}
+
+	return bIsLoaded ? EActiveTickReturnType::StopTicking : EActiveTickReturnType::KeepTicking;
 }
 
 void SSurvey::ConstructLoadingLayout()
diff --git a/Engine/Plugins/Editor/EpicSurvey/Source/EpicSurvey/Private/SSurvey.h b/Engine/Plugins/Editor/EpicSurvey/Source/EpicSurvey/Private/SSurvey.h
--- a/Engine/Plugins/Editor/EpicSurvey/Source/EpicSurvey/Private/SSurvey.h
+++ b/Engine/Plugins/Editor/EpicSurvey/Source/EpicSurvey/Private/SSurvey.h
@@ -15,7 +15,7 @@ public:
 
 	/** Widget constructor */
 	void Construct( const FArguments& Args, const TSharedRef< FEpicSurvey >& InEpicSurvey, const TSharedRef< FSurvey >& InSurvey );
-	virtual void Tick( const FGeometry& AllottedGeometry, const double InCurrentTime, const float InDeltaTime ) override;
+	//virtual void Tick( const FGeometry& AllottedGeometry, const double InCurrentTime, const float InDeltaTime ) override;
 
 	virtual FReply OnMouseButtonDown( const FGeometry& MyGeometry, const FPointerEvent& MouseEvent ) override;
 
@@ -40,7 +40,8 @@ private:
 	void DisplayPage( int32 NewPageIndex );
 
 private:
-	bool FinishedLoading;
+	/** Initializes the survey once its done loading */
+	EActiveTickReturnType MonitorLoadStatePostConstruct( double InCurrentTime, float InDeltaTime );
 
 	TSharedPtr< FEpicSurvey > EpicSurvey;
 	TSharedPtr< FSurvey > Survey;
diff --git a/Engine/Plugins/Editor/EpicSurvey/Source/EpicSurvey/Private/SSurveyNotification.h b/Engine/Plugins/Editor/EpicSurvey/Source/EpicSurvey/Private/SSurveyNotification.h
--- a/Engine/Plugins/Editor/EpicSurvey/Source/EpicSurvey/Private/SSurveyNotification.h
+++ b/Engine/Plugins/Editor/EpicSurvey/Source/EpicSurvey/Private/SSurveyNotification.h
@@ -68,11 +68,11 @@ public:
 			{
 				if ( Sequence.IsAtStart() )
 				{
-					const_cast<SSurveyNotification*>(this)->Sequence.Play();
+					const_cast<SSurveyNotification*>(this)->Sequence.Play(this->AsShared());
 				}
 				else
 				{
-					const_cast<SSurveyNotification*>(this)->Sequence.PlayReverse();
+					const_cast<SSurveyNotification*>(this)->Sequence.PlayReverse(this->AsShared());
 				}
 			}
 
diff --git a/Engine/Plugins/Editor/PluginsEditor/Source/PluginsEditor/Private/SPluginCategories.cpp b/Engine/Plugins/Editor/PluginsEditor/Source/PluginsEditor/Private/SPluginCategories.cpp
--- a/Engine/Plugins/Editor/PluginsEditor/Source/PluginsEditor/Private/SPluginCategories.cpp
+++ b/Engine/Plugins/Editor/PluginsEditor/Source/PluginsEditor/Private/SPluginCategories.cpp
@@ -15,7 +15,6 @@ void SPluginCategories::Construct( const FArguments& Args, const TSharedRef< SPl
 {
 	OwnerWeak = Owner;
 
-	bNeedsRefresh = false;
 	RebuildAndFilterCategoryTree();
 
 	PluginCategoryTreeView =
@@ -179,23 +178,8 @@ void SPluginCategories::RebuildAndFilterCategoryTree()
 	{
 		PluginCategoryTreeView->RequestTreeRefresh();
 	}
-
-	bNeedsRefresh = false;
 }
 
-
-void SPluginCategories::Tick( const FGeometry& AllottedGeometry, const double InCurrentTime, const float InDeltaTime )
-{
-	// Call parent implementation
-	SCompoundWidget::Tick( AllottedGeometry, InCurrentTime, InDeltaTime );
-
-	if( bNeedsRefresh )
-	{
-		RebuildAndFilterCategoryTree();
-	}
-}
-
-
 TSharedRef<ITableRow> SPluginCategories::PluginCategoryTreeView_OnGenerateRow( FPluginCategoryTreeItemPtr Item, const TSharedRef<STableViewBase>& OwnerTable )
 {
 	return
diff --git a/Engine/Plugins/Editor/PluginsEditor/Source/PluginsEditor/Private/SPluginCategories.h b/Engine/Plugins/Editor/PluginsEditor/Source/PluginsEditor/Private/SPluginCategories.h
--- a/Engine/Plugins/Editor/PluginsEditor/Source/PluginsEditor/Private/SPluginCategories.h
+++ b/Engine/Plugins/Editor/PluginsEditor/Source/PluginsEditor/Private/SPluginCategories.h
@@ -55,10 +55,6 @@ private:
 	/** Rebuilds the category tree from scratch */
 	void RebuildAndFilterCategoryTree();
 
-	/** SWidget overrides */
-	virtual void Tick( const FGeometry& AllottedGeometry, const double InCurrentTime, const float InDeltaTime ) override;
-
-
 private:
 
 	/** Weak pointer back to its owner */
@@ -69,9 +65,5 @@ private:
 
 	/** Root list of categories */
 	TArray< FPluginCategoryTreeItemPtr > RootPluginCategories;
-
-	/** True if something has changed with either our filtering or the loaded plugin set, and we need to
-	    do a full refresh */
-	bool bNeedsRefresh;
 };
 
diff --git a/Engine/Plugins/Editor/PluginsEditor/Source/PluginsEditor/Private/SPluginList.cpp b/Engine/Plugins/Editor/PluginsEditor/Source/PluginsEditor/Private/SPluginList.cpp
--- a/Engine/Plugins/Editor/PluginsEditor/Source/PluginsEditor/Private/SPluginList.cpp
+++ b/Engine/Plugins/Editor/PluginsEditor/Source/PluginsEditor/Private/SPluginList.cpp
@@ -19,7 +19,7 @@ void SPluginList::Construct( const FArguments& Args, const TSharedRef< SPluginsE
 	// Find out when the plugin text filter changes
 	Owner->GetPluginTextFilter().OnChanged().AddSP( this, &SPluginList::OnPluginTextFilterChanged );
 
-	bNeedsRefresh = false;
+	bIsActiveTickRegistered = false;
 	RebuildAndFilterPluginList();
 
 	// @todo plugedit: Have optional compact version with only plugin icon + name + version?  Only expand selected?
@@ -129,23 +129,16 @@ void SPluginList::RebuildAndFilterPluginList()
 	{
 		PluginListView->RequestListRefresh();
 	}
-
-	bNeedsRefresh = false;
 }
 
-
-void SPluginList::Tick( const FGeometry& AllottedGeometry, const double InCurrentTime, const float InDeltaTime )
+EActiveTickReturnType SPluginList::TriggerListRebuild(double InCurrentTime, float InDeltaTime)
 {
-	// Call parent implementation
-	SCompoundWidget::Tick( AllottedGeometry, InCurrentTime, InDeltaTime );
+	RebuildAndFilterPluginList();
 
-	if( bNeedsRefresh )
-	{
-		RebuildAndFilterPluginList();
-	}
+	bIsActiveTickRegistered = false;
+	return EActiveTickReturnType::StopTicking;
 }
 
-
 void SPluginList::OnPluginTextFilterChanged()
 {
 	SetNeedsRefresh();
@@ -154,7 +147,11 @@ void SPluginList::OnPluginTextFilterChanged()
 
 void SPluginList::SetNeedsRefresh()
 {
-	bNeedsRefresh = true;
+	if (!bIsActiveTickRegistered)
+	{
+		bIsActiveTickRegistered = true;
+		RegisterActiveTick(0.f, FWidgetActiveTickDelegate::CreateSP(this, &SPluginList::TriggerListRebuild));
+	}
 }
 
 
diff --git a/Engine/Plugins/Editor/PluginsEditor/Source/PluginsEditor/Private/SPluginList.h b/Engine/Plugins/Editor/PluginsEditor/Source/PluginsEditor/Private/SPluginList.h
--- a/Engine/Plugins/Editor/PluginsEditor/Source/PluginsEditor/Private/SPluginList.h
+++ b/Engine/Plugins/Editor/PluginsEditor/Source/PluginsEditor/Private/SPluginList.h
@@ -45,9 +45,8 @@ private:
 	/** Rebuilds the list of plugins from scratch and applies filtering. */
 	void RebuildAndFilterPluginList();
 
-	/** SWidget overrides */
-	virtual void Tick( const FGeometry& AllottedGeometry, const double InCurrentTime, const float InDeltaTime ) override;
-
+	/** One-off active tick to trigger a full refresh when something has changed with either our filtering or the loaded plugin set */
+	EActiveTickReturnType TriggerListRebuild(double InCurrentTime, float InDeltaTime);
 
 private:
 
@@ -60,9 +59,7 @@ private:
 	/** List of everything that we want to display in the plugin list */
 	TArray< FPluginListItemPtr > PluginListItems;
 
-	/** True if something has changed with either our filtering or the loaded plugin set, and we need to
-	    do a full refresh */
-	bool bNeedsRefresh;
-
+	/** Whether the active tick to refresh the list is registered */
+	bool bIsActiveTickRegistered;
 };
 
diff --git a/Engine/Plugins/Editor/PluginsEditor/Source/PluginsEditor/Private/SPluginsEditor.cpp b/Engine/Plugins/Editor/PluginsEditor/Source/PluginsEditor/Private/SPluginsEditor.cpp
--- a/Engine/Plugins/Editor/PluginsEditor/Source/PluginsEditor/Private/SPluginsEditor.cpp
+++ b/Engine/Plugins/Editor/PluginsEditor/Source/PluginsEditor/Private/SPluginsEditor.cpp
@@ -26,7 +26,7 @@ void SPluginsEditor::Construct( const FArguments& Args )
 
 	// @todo plugedit: Should we save/restore category selection and other view settings?  Splitter position, etc.
 
-	bBreadcrumbNeedsRefresh = true;
+	RegisterActiveTick (0.f, FWidgetActiveTickDelegate::CreateSP (this, &SPluginsEditor::TriggerBreadcrumbRefresh));
 
 	// Setup text filtering
 	PluginTextFilter = MakeShareable( new FPluginTextFilter( FPluginTextFilter::FItemToStringArray::CreateStatic( &Local::PluginStatusToStringArray ) ) );
@@ -159,22 +159,15 @@ void SPluginsEditor::OnCategorySelectionChanged()
 	}
 
 	// Breadcrumbs will need to be refreshed
-	bBreadcrumbNeedsRefresh = true;
+	RegisterActiveTick (0.f, FWidgetActiveTickDelegate::CreateSP (this, &SPluginsEditor::TriggerBreadcrumbRefresh));
 }
 
-
-void SPluginsEditor::Tick( const FGeometry& AllottedGeometry, const double InCurrentTime, const float InDeltaTime )
+EActiveTickReturnType SPluginsEditor::TriggerBreadcrumbRefresh(double InCurrentTime, float InDeltaTime)
 {
-	// Call parent implementation
-	SCompoundWidget::Tick( AllottedGeometry, InCurrentTime, InDeltaTime );
-
-	if( bBreadcrumbNeedsRefresh )
-	{
-		RefreshBreadcrumbTrail();
-	}
+	RefreshBreadcrumbTrail();
+	return EActiveTickReturnType::StopTicking;
 }
 
-
 void SPluginsEditor::RefreshBreadcrumbTrail()
 {
 	// Update breadcrumb trail
@@ -202,8 +195,6 @@ void SPluginsEditor::RefreshBreadcrumbTrail()
 			BreadcrumbTrail->PushCrumb( Category->GetCategoryDisplayName(), NewBreadcrumb );
 		}
 	}
-
-	bBreadcrumbNeedsRefresh = false;
 }
 
 
diff --git a/Engine/Plugins/Editor/PluginsEditor/Source/PluginsEditor/Private/SPluginsEditor.h b/Engine/Plugins/Editor/PluginsEditor/Source/PluginsEditor/Private/SPluginsEditor.h
--- a/Engine/Plugins/Editor/PluginsEditor/Source/PluginsEditor/Private/SPluginsEditor.h
+++ b/Engine/Plugins/Editor/PluginsEditor/Source/PluginsEditor/Private/SPluginsEditor.h
@@ -82,9 +82,8 @@ private:
 	/** Called to refresh the breadcrumb trail immediately */
 	void RefreshBreadcrumbTrail();
 
-	/** SWidget overrides */
-	virtual void Tick( const FGeometry& AllottedGeometry, const double InCurrentTime, const float InDeltaTime ) override;
-
+	/** One-off active tick to trigger a refresh of the breadcrumb trail as needed */
+	EActiveTickReturnType TriggerBreadcrumbRefresh(double InCurrentTime, float InDeltaTime);
 
 private:
 
@@ -100,8 +99,5 @@ private:
 	/** Breadcrumb trail widget for the currently selected category */
 	TSharedPtr< SPluginCategoryBreadcrumbTrail > BreadcrumbTrail;
 
-	/** True if the breadcrumb trail needs to be refreshed */
-	bool bBreadcrumbNeedsRefresh;
-
 };
 
diff --git a/Engine/Plugins/Experimental/LiveEditor/Source/LiveEditor/Private/LiveEditorConfigWindow.cpp b/Engine/Plugins/Experimental/LiveEditor/Source/LiveEditor/Private/LiveEditorConfigWindow.cpp
--- a/Engine/Plugins/Experimental/LiveEditor/Source/LiveEditor/Private/LiveEditorConfigWindow.cpp
+++ b/Engine/Plugins/Experimental/LiveEditor/Source/LiveEditor/Private/LiveEditorConfigWindow.cpp
@@ -195,8 +195,6 @@ public:
 
 	virtual void Tick( const FGeometry& AllottedGeometry, const double InCurrentTime, const float InDeltaTime ) override
 	{
-		SCompoundWidget::Tick( AllottedGeometry, InCurrentTime, InDeltaTime );
-
 		bool success = FLiveEditorManager::Get().GetDeviceData( *DeviceID.Get(), DeviceData );
 		check(success);
 	}
diff --git a/Engine/Plugins/Messaging/MessagingDebugger/Source/MessagingDebugger/Private/Widgets/Endpoints/SMessagingEndpoints.cpp b/Engine/Plugins/Messaging/MessagingDebugger/Source/MessagingDebugger/Private/Widgets/Endpoints/SMessagingEndpoints.cpp
--- a/Engine/Plugins/Messaging/MessagingDebugger/Source/MessagingDebugger/Private/Widgets/Endpoints/SMessagingEndpoints.cpp
+++ b/Engine/Plugins/Messaging/MessagingDebugger/Source/MessagingDebugger/Private/Widgets/Endpoints/SMessagingEndpoints.cpp
@@ -113,12 +113,6 @@ void SMessagingEndpoints::Construct( const FArguments& InArgs, const FMessagingD
 	Filter->OnChanged().AddRaw(this, &SMessagingEndpoints::HandleFilterChanged);
 	Model->OnSelectedMessageChanged().AddRaw(this, &SMessagingEndpoints::HandleModelSelectedMessageChanged);
 }
-END_SLATE_FUNCTION_BUILD_OPTIMIZATION
-
-
-/* SCompoundWidget overrides
- *****************************************************************************/
-
 void SMessagingEndpoints::Tick( const FGeometry& AllottedGeometry, const double InCurrentTime, const float InDeltaTime )
 {
 	// @todo gmp: fix this
diff --git a/Engine/Plugins/Messaging/MessagingDebugger/Source/MessagingDebugger/Private/Widgets/History/SMessagingHistoryTableRow.h b/Engine/Plugins/Messaging/MessagingDebugger/Source/MessagingDebugger/Private/Widgets/History/SMessagingHistoryTableRow.h
--- a/Engine/Plugins/Messaging/MessagingDebugger/Source/MessagingDebugger/Private/Widgets/History/SMessagingHistoryTableRow.h
+++ b/Engine/Plugins/Messaging/MessagingDebugger/Source/MessagingDebugger/Private/Widgets/History/SMessagingHistoryTableRow.h
@@ -48,8 +48,6 @@ public:
 
 	virtual void Tick( const FGeometry& AllottedGeometry, const double InCurrentTime, const float InDeltaTime ) override
 	{
-		SWidget::Tick(AllottedGeometry, InCurrentTime, InDeltaTime);
-
 		MaxDispatchLatency = -1.0;
 		MaxHandlingTime = -1.0;
 
diff --git a/Engine/Source/Developer/AutomationWindow/Private/SAutomationTestItem.cpp b/Engine/Source/Developer/AutomationWindow/Private/SAutomationTestItem.cpp
--- a/Engine/Source/Developer/AutomationWindow/Private/SAutomationTestItem.cpp
+++ b/Engine/Source/Developer/AutomationWindow/Private/SAutomationTestItem.cpp
@@ -34,7 +34,7 @@ public:
 
 	// SWidget Interface
 
-	virtual void Tick(const FGeometry& AllottedGeometry, const double InCurrentTime, const float InDeltaTime) override;
+	virtual void Tick( const FGeometry& AllottedGeometry, const double InCurrentTime, const float InDeltaTime ) override;
 
 private:
 
@@ -104,10 +104,8 @@ void SAutomationHistoryCell::RebuildContentArea()
 		];
 	}
 }
-END_SLATE_FUNCTION_BUILD_OPTIMIZATION
-
 
-void SAutomationHistoryCell::Tick(const FGeometry& AllottedGeometry, const double InCurrentTime, const float InDeltaTime)
+void SAutomationHistoryCell::Tick( const FGeometry& AllottedGeometry, const double InCurrentTime, const float InDeltaTime )
 {
 	const TArray<TSharedPtr<FAutomationHistoryItem>> TestHistory = HistoryItem->GetHistory();
 
diff --git a/Engine/Source/Developer/DeviceManager/Private/Widgets/Browser/SDeviceBrowser.cpp b/Engine/Source/Developer/DeviceManager/Private/Widgets/Browser/SDeviceBrowser.cpp
--- a/Engine/Source/Developer/DeviceManager/Private/Widgets/Browser/SDeviceBrowser.cpp
+++ b/Engine/Source/Developer/DeviceManager/Private/Widgets/Browser/SDeviceBrowser.cpp
@@ -105,16 +105,9 @@ void SDeviceBrowser::Construct( const FArguments& InArgs, const FDeviceManagerMo
 
 	Filter->OnFilterChanged().AddSP(this, &SDeviceBrowser::HandleFilterChanged);
 }
-END_SLATE_FUNCTION_BUILD_OPTIMIZATION
-
-
-/* SCompoundWidget overrides
- *****************************************************************************/
-
 void SDeviceBrowser::Tick( const FGeometry& AllottedGeometry, const double InCurrentTime, const float InDeltaTime )
 {
-	SCompoundWidget::Tick(AllottedGeometry, InCurrentTime, InDeltaTime);
-
+	//@TODO Passive - Only happens in response to the addition or removal of a device to the device service manager
 	if (NeedsServiceListRefresh)
 	{
 		ReloadDeviceServiceList(true);
diff --git a/Engine/Source/Developer/DeviceManager/Private/Widgets/Processes/SDeviceProcesses.cpp b/Engine/Source/Developer/DeviceManager/Private/Widgets/Processes/SDeviceProcesses.cpp
--- a/Engine/Source/Developer/DeviceManager/Private/Widgets/Processes/SDeviceProcesses.cpp
+++ b/Engine/Source/Developer/DeviceManager/Private/Widgets/Processes/SDeviceProcesses.cpp
@@ -131,24 +131,18 @@ void SDeviceProcesses::Construct( const FArguments& InArgs, const FDeviceManager
 	Model->OnSelectedDeviceServiceChanged().AddRaw(this, &SDeviceProcesses::HandleModelSelectedDeviceServiceChanged);
 
 	ReloadProcessList(true);
-}
-END_SLATE_FUNCTION_BUILD_OPTIMIZATION
-
 
-/* SCompoundWidget overrides
- *****************************************************************************/
+	// Register for an active update every 2.5 seconds
+	RegisterActiveTick( 2.5f, FWidgetActiveTickDelegate::CreateSP( this, &SDeviceProcesses::UpdateProcessList ) );
+}
 
-void SDeviceProcesses::Tick( const FGeometry& AllottedGeometry, const double InCurrentTime, const float InDeltaTime )
+EActiveTickReturnType SDeviceProcesses::UpdateProcessList( double InCurrentTime, float InDeltaTime )
 {
-	SCompoundWidget::Tick(AllottedGeometry, InCurrentTime, InDeltaTime);
+	ReloadProcessList( true );
 
-	if (LastProcessListRefreshTime + FTimespan::FromSeconds(2.5) < FDateTime::UtcNow())
-	{
-		ReloadProcessList(true);
-	}
+	return EActiveTickReturnType::KeepTicking;
 }
 
-
 /* SDeviceDetails implementation
  *****************************************************************************/
 
diff --git a/Engine/Source/Developer/DeviceManager/Private/Widgets/Processes/SDeviceProcesses.h b/Engine/Source/Developer/DeviceManager/Private/Widgets/Processes/SDeviceProcesses.h
--- a/Engine/Source/Developer/DeviceManager/Private/Widgets/Processes/SDeviceProcesses.h
+++ b/Engine/Source/Developer/DeviceManager/Private/Widgets/Processes/SDeviceProcesses.h
@@ -31,12 +31,6 @@ public:
 	 */
 	void Construct( const FArguments& InArgs, const FDeviceManagerModelRef& InModel );
 
-public:
-
-	// SCompoundWidget overrides
-
-	virtual void Tick( const FGeometry& AllottedGeometry, const double InCurrentTime, const float InDeltaTime ) override;
-
 protected:
 
 	/**
@@ -45,6 +39,8 @@ protected:
 	void ReloadProcessList( bool FullyReload );
 
 private:
+	// Periodically refreshes the process list
+	EActiveTickReturnType UpdateProcessList( double InCurrentTime, float InDeltaTime );
 
 	// Callback for getting the text of the message overlay.
 	FText HandleMessageOverlayText( ) const;
@@ -77,7 +73,7 @@ private:
 	FReply HandleTerminateProcessButtonClicked( );
 
 private:
-
+	
 	// Holds the time at which the process list was last refreshed.
 	FDateTime LastProcessListRefreshTime;
 
diff --git a/Engine/Source/Developer/LogVisualizer/Private/TimeSliderController.cpp b/Engine/Source/Developer/LogVisualizer/Private/TimeSliderController.cpp
--- a/Engine/Source/Developer/LogVisualizer/Private/TimeSliderController.cpp
+++ b/Engine/Source/Developer/LogVisualizer/Private/TimeSliderController.cpp
@@ -380,7 +380,7 @@ FReply FSequencerTimeSliderController::OnMouseMove( TSharedRef<SWidget> WidgetOw
 			if (!bPanning)
 			{
 				DistanceDragged += FMath::Abs( MouseEvent.GetCursorDelta().X );
-				if ( DistanceDragged > FSlateApplication::Get().GetDragTriggerDistnace() )
+				if ( DistanceDragged > FSlateApplication::Get().GetDragTriggerDistance() )
 				{
 					FReply::Handled().CaptureMouse(WidgetOwner).UseHighPrecisionMouseMovement(WidgetOwner);
 					SoftwareCursorPosition = MyGeometry.AbsoluteToLocal(MouseEvent.GetLastScreenSpacePosition());
@@ -437,7 +437,7 @@ FReply FSequencerTimeSliderController::OnMouseMove( TSharedRef<SWidget> WidgetOw
 			if ( !bDraggingScrubber )
 			{
 				DistanceDragged += FMath::Abs( MouseEvent.GetCursorDelta().X );
-				if ( DistanceDragged > 0/*FSlateApplication::Get().GetDragTriggerDistnace()*/ )
+				if ( DistanceDragged > 0/*FSlateApplication::Get().GetDragTriggerDistance()*/ )
 				{
 					bDraggingScrubber = true;
 					TimeSliderArgs.OnBeginScrubberMovement.ExecuteIfBound();
diff --git a/Engine/Source/Developer/MovieSceneTools/Private/AnimationTrackEditor.cpp b/Engine/Source/Developer/MovieSceneTools/Private/AnimationTrackEditor.cpp
--- a/Engine/Source/Developer/MovieSceneTools/Private/AnimationTrackEditor.cpp
+++ b/Engine/Source/Developer/MovieSceneTools/Private/AnimationTrackEditor.cpp
@@ -122,12 +122,6 @@ int32 FAnimationSection::OnPaintSection( const FGeometry& AllottedGeometry, cons
 	return LayerId+3;
 }
 
-void FAnimationSection::Tick( const FGeometry& AllottedGeometry, const FGeometry& ParentGeometry, const double InCurrentTime, const float InDeltaTime )
-{
-}
-
-
-
 FAnimationTrackEditor::FAnimationTrackEditor( TSharedRef<ISequencer> InSequencer )
 	: FMovieSceneTrackEditor( InSequencer ) 
 {
diff --git a/Engine/Source/Developer/MovieSceneTools/Private/AnimationTrackEditor.h b/Engine/Source/Developer/MovieSceneTools/Private/AnimationTrackEditor.h
--- a/Engine/Source/Developer/MovieSceneTools/Private/AnimationTrackEditor.h
+++ b/Engine/Source/Developer/MovieSceneTools/Private/AnimationTrackEditor.h
@@ -56,7 +56,6 @@ public:
 	virtual float GetSectionHeight() const override;
 	virtual void GenerateSectionLayout( class ISectionLayoutBuilder& LayoutBuilder ) const override {}
 	virtual int32 OnPaintSection( const FGeometry& AllottedGeometry, const FSlateRect& SectionClippingRect, FSlateWindowElementList& OutDrawElements, int32 LayerId, bool bParentEnabled ) const override;
-	virtual void Tick( const FGeometry& AllottedGeometry, const FGeometry& ParentGeometry, const double InCurrentTime, const float InDeltaTime ) override;
 
 private:
 	/** The section we are visualizing */
diff --git a/Engine/Source/Developer/MovieSceneTools/Private/DirectorTrackEditor.cpp b/Engine/Source/Developer/MovieSceneTools/Private/DirectorTrackEditor.cpp
--- a/Engine/Source/Developer/MovieSceneTools/Private/DirectorTrackEditor.cpp
+++ b/Engine/Source/Developer/MovieSceneTools/Private/DirectorTrackEditor.cpp
@@ -26,8 +26,12 @@ namespace AnimatableShotToolConstants
 {
 	// @todo Sequencer Perhaps allow this to be customizable
 	const uint32 DirectorTrackHeight = 100;
+	const double ThumbnailFadeInDuration = 0.25f;
+
 }
 
+/////////////////////////////////////////////////
+// FShotThumbnailPool
 
 FShotThumbnailPool::FShotThumbnailPool(TSharedPtr<ISequencer> InSequencer, uint32 InMaxThumbnailsToDrawAtATime)
 	: Sequencer(InSequencer)
@@ -84,15 +88,15 @@ void FShotThumbnailPool::DrawThumbnails()
 }
 
 
+/////////////////////////////////////////////////
+// FShotThumbnail
 
 FShotThumbnail::FShotThumbnail(TSharedPtr<FShotSection> InSection, TRange<float> InTimeRange)
 	: OwningSection(InSection)
 	, Texture(NULL)
 	, TimeRange(InTimeRange)
+	, FadeInStartTime(0.0)
 {
-	ThumbnailFadeAnimation = FCurveSequence();
-	ThumbnailFadeCurve = ThumbnailFadeAnimation.AddCurve(0.f, 0.25f, ECurveEaseFunction::QuadOut);
-
 	Texture = new FSlateTexture2DRHIRef(GetSize().X, GetSize().Y, PF_B8G8R8A8, NULL, TexCreate_Dynamic, true);
 
 	BeginInitResource( Texture );
@@ -122,7 +126,7 @@ FSlateShaderResource* FShotThumbnail::GetViewportRenderTargetTexture() const
 void FShotThumbnail::DrawThumbnail()
 {
 	OwningSection.Pin()->DrawViewportThumbnail(SharedThis(this));
-	ThumbnailFadeAnimation.Play();
+	FadeInStartTime = FSlateApplication::Get().GetCurrentTime();
 }
 
 float FShotThumbnail::GetTime() const {return TimeRange.GetLowerBoundValue();}
@@ -137,7 +141,10 @@ void FShotThumbnail::CopyTextureIn(FSlateRenderTargetRHI* InTexture)
 	});
 }
 
-float FShotThumbnail::GetFadeInCurve() const {return 1.f - ThumbnailFadeCurve.GetLerp();}
+float FShotThumbnail::GetFadeInCurve() const 
+{
+	return (float)FMath::Clamp(FadeInStartTime / AnimatableShotToolConstants::ThumbnailFadeInDuration, 0.0, 1.0);
+}
 
 bool FShotThumbnail::IsVisible() const
 {
@@ -149,7 +156,8 @@ bool FShotThumbnail::IsValid() const
 	return OwningSection.IsValid();
 }
 
-
+/////////////////////////////////////////////////
+// FShotSection
 
 FShotSection::FShotSection( TSharedPtr<ISequencer> InSequencer, TSharedPtr<FShotThumbnailPool> InThumbnailPool, UMovieSceneSection& InSection, UObject* InTargetObject )
 	: Section( &InSection )
@@ -377,7 +385,8 @@ void FShotSection::CalculateThumbnailWidthAndResize()
 	}
 }
 
-
+/////////////////////////////////////////////////
+// FDirectorTrackEditor
 
 FDirectorTrackEditor::FDirectorTrackEditor( TSharedRef<ISequencer> InSequencer )
 	: FMovieSceneTrackEditor( InSequencer ) 
diff --git a/Engine/Source/Developer/MovieSceneTools/Private/DirectorTrackEditor.h b/Engine/Source/Developer/MovieSceneTools/Private/DirectorTrackEditor.h
--- a/Engine/Source/Developer/MovieSceneTools/Private/DirectorTrackEditor.h
+++ b/Engine/Source/Developer/MovieSceneTools/Private/DirectorTrackEditor.h
@@ -113,9 +113,8 @@ private:
 	/** Where in time this thumbnail is a rendering of */
 	TRange<float> TimeRange;
 
-	/** The fade in curve, so the thumbnails don't just pop in */
-	FCurveSequence ThumbnailFadeAnimation;
-	FCurveHandle ThumbnailFadeCurve;
+	/** The time when the thumbnail started to fade in. */
+	double FadeInStartTime;
 };
 
 
diff --git a/Engine/Source/Developer/OutputLog/Private/SDebugConsole.cpp b/Engine/Source/Developer/OutputLog/Private/SDebugConsole.cpp
--- a/Engine/Source/Developer/OutputLog/Private/SDebugConsole.cpp
+++ b/Engine/Source/Developer/OutputLog/Private/SDebugConsole.cpp
@@ -91,7 +91,7 @@ void SDebugConsole::Construct( const FArguments& InArgs, const EDebugConsoleStyl
 	AnimCurve = AnimCurveSequence.AddCurve( 0.0f, DebugConsoleDefs::IntroAnimationDuration, ECurveEaseFunction::QuadOut );
 	FlashCurve = AnimCurveSequence.AddCurve( DebugConsoleDefs::IntroAnimationDuration, .15f, ECurveEaseFunction::QuadInOut );
 
-	AnimCurveSequence.Play();
+	AnimCurveSequence.Play(this->AsShared());
 }
 END_SLATE_FUNCTION_BUILD_OPTIMIZATION
 
diff --git a/Engine/Source/Developer/OutputLog/Private/SOutputLog.cpp b/Engine/Source/Developer/OutputLog/Private/SOutputLog.cpp
--- a/Engine/Source/Developer/OutputLog/Private/SOutputLog.cpp
+++ b/Engine/Source/Developer/OutputLog/Private/SOutputLog.cpp
@@ -161,12 +161,8 @@ void SConsoleInputBox::Construct( const FArguments& InArgs )
 			)
 	];
 }
-END_SLATE_FUNCTION_BUILD_OPTIMIZATION
-
 void SConsoleInputBox::Tick( const FGeometry& AllottedGeometry, const double InCurrentTime, const float InDeltaTime )
 {
-	SCompoundWidget::Tick(AllottedGeometry, InCurrentTime, InDeltaTime);
-
 	if (!GIntraFrameDebuggingGameThread && !IsEnabled())
 	{
 		SetEnabled(true);
diff --git a/Engine/Source/Developer/Profiler/Private/Widgets/SDataGraph.cpp b/Engine/Source/Developer/Profiler/Private/Widgets/SDataGraph.cpp
--- a/Engine/Source/Developer/Profiler/Private/Widgets/SDataGraph.cpp
+++ b/Engine/Source/Developer/Profiler/Private/Widgets/SDataGraph.cpp
@@ -226,7 +226,6 @@ void SDataGraph::Construct( const FArguments& InArgs )
 
 void SDataGraph::Tick( const FGeometry& AllottedGeometry, const double InCurrentTime, const float InDeltaTime )
 {
-	SCompoundWidget::Tick(AllottedGeometry,InCurrentTime,InDeltaTime);
 	ThisGeometry = AllottedGeometry;
 
 	UpdateState();
@@ -1175,7 +1174,7 @@ FReply SDataGraph::OnMouseWheel( const FGeometry& MyGeometry, const FPointerEven
 	return FReply::Handled();
 }
 
-const FEventGraphDataHandlerRef SDataGraph::PrepareEventGraphDataHandler( const FGeometry& MyGeometry, const FVector2D& ScreenSpacePosition )
+const FEventGraphDataHandlerRef SDataGraph::PrepareEventGraphDataHandler( const FVector2D& ScreenSpacePosition )
 {
 	const FGraphDescription* GraphDescription = GetFirstGraph();
 
@@ -1282,7 +1281,7 @@ void SDataGraph::ShowContextMenu( const FVector2D& ScreenSpacePosition )
 	const FProfilerActionManager& ProfilerActionManager = FProfilerManager::GetActionManager();
 
 	// Build data required for opening event graph(s).
-	const FEventGraphDataHandlerRef EventGraphDataHandler = PrepareEventGraphDataHandler( ThisGeometry, ScreenSpacePosition );
+	const FEventGraphDataHandlerRef EventGraphDataHandler = PrepareEventGraphDataHandler( ScreenSpacePosition );
 
 	const bool bShouldCloseWindowAfterMenuSelection = true;
 	FMenuBuilder MenuBuilder( bShouldCloseWindowAfterMenuSelection, ProfilerCommandList );
diff --git a/Engine/Source/Developer/Profiler/Private/Widgets/SDataGraph.h b/Engine/Source/Developer/Profiler/Private/Widgets/SDataGraph.h
--- a/Engine/Source/Developer/Profiler/Private/Widgets/SDataGraph.h
+++ b/Engine/Source/Developer/Profiler/Private/Widgets/SDataGraph.h
@@ -280,7 +280,7 @@ public:
 protected:
 	/** The event to execute when the selected frames have been changed, for index based view mode. */
 	FSelectionChangedForIndexEvent SelectionChangedForIndexEvent;
-	
+
 public:
 	void AddInnerGraph
 	( 
@@ -349,7 +349,7 @@ protected:
 
 	void UpdateState();
 
-	const FEventGraphDataHandlerRef PrepareEventGraphDataHandler( const FGeometry& MyGeometry, const FVector2D& ScreenSpacePosition );
+	const FEventGraphDataHandlerRef PrepareEventGraphDataHandler( const FVector2D& ScreenSpacePosition );
 
 	/** Handles FExecuteAction for ViewMode_SetIndexBased. */
 	void ViewMode_SetIndexBased_Execute();
diff --git a/Engine/Source/Developer/Profiler/Private/Widgets/SFiltersAndPresets.cpp b/Engine/Source/Developer/Profiler/Private/Widgets/SFiltersAndPresets.cpp
--- a/Engine/Source/Developer/Profiler/Private/Widgets/SFiltersAndPresets.cpp
+++ b/Engine/Source/Developer/Profiler/Private/Widgets/SFiltersAndPresets.cpp
@@ -571,8 +571,8 @@ protected:
 SFiltersAndPresets::SFiltersAndPresets()
 	: GroupingMode( EStatGroupingOrSortingMode::GroupName )
 	, SortingMode( EStatGroupingOrSortingMode::StatName )
-	, bUpdateOnNextTick( false )
 	, bExpansionSaved( false )
+	, bIsActiveTickRegistered( false )
 {
 	FMemory::MemSet( bStatTypeIsVisible, 1 );
 }
@@ -755,23 +755,26 @@ void SFiltersAndPresets::Construct( const FArguments& InArgs )
 
 void SFiltersAndPresets::ProfilerManager_OnRequestFilterAndPresetsUpdate()
 {
-	bUpdateOnNextTick = true;
+	RegisterForUpdate();
 }
 
-void SFiltersAndPresets::Tick( const FGeometry& AllottedGeometry, const double InCurrentTime, const float InDeltaTime )
+void SFiltersAndPresets::RegisterForUpdate()
 {
-	if( bUpdateOnNextTick )
+	if (!bIsActiveTickRegistered)
 	{
-		for( auto It = FProfilerManager::Get()->GetProfilerInstancesIterator(); It; ++It )
-		{
-			UpdateGroupAndStatTree( It.Value() );
-			break;
-		}
-
-		bUpdateOnNextTick = false;
+		bIsActiveTickRegistered = true;
+		RegisterActiveTick(0.f, FWidgetActiveTickDelegate::CreateSP(this, &SFiltersAndPresets::UpdateData));
 	}
 }
 
+EActiveTickReturnType SFiltersAndPresets::UpdateData(double InCurrentTime, float InDeltaTime)
+{
+	auto It = FProfilerManager::Get()->GetProfilerInstancesIterator();
+	UpdateGroupAndStatTree(It.Value());
+
+	return EActiveTickReturnType::StopTicking;
+}
+
 void SFiltersAndPresets::UpdateGroupAndStatTree( const FProfilerSessionPtr InProfilerSession )
 {
 	const bool bRebuild = InProfilerSession != ProfilerSession;
diff --git a/Engine/Source/Developer/Profiler/Private/Widgets/SFiltersAndPresets.h b/Engine/Source/Developer/Profiler/Private/Widgets/SFiltersAndPresets.h
--- a/Engine/Source/Developer/Profiler/Private/Widgets/SFiltersAndPresets.h
+++ b/Engine/Source/Developer/Profiler/Private/Widgets/SFiltersAndPresets.h
@@ -284,15 +284,6 @@ public:
 	 */
 	void Construct( const FArguments& InArgs );
 
-	/**
-	 * Ticks this widget.  Override in derived classes, but always call the parent implementation.
-	 *
-	 * @param  AllottedGeometry The space allotted for this widget
-	 * @param  InCurrentTime  Current absolute real time
-	 * @param  InDeltaTime  Real time passed since last tick
-	 */
-	virtual void Tick( const FGeometry& AllottedGeometry, const double InCurrentTime, const float InDeltaTime ) override;
-
 protected:
 	/** Called when the filter and presets widget should be updated with the latest data. */
 	void ProfilerManager_OnRequestFilterAndPresetsUpdate();
@@ -415,12 +406,19 @@ protected:
 	/** How we sort the metadata?. */
 	EStatGroupingOrSortingMode::Type SortingMode;
 
-	/** If true, an update should occur on the next tick. */
-	bool bUpdateOnNextTick;
-
 	/** If true, the expanded nodes have been saved before applying a text filter. */
 	bool bExpansionSaved;
 
 	/** Holds the visibility of each stat type. */
 	bool bStatTypeIsVisible[EProfilerSampleTypes::InvalidOrMax];
+
+private:
+	/** Registers the active tick to trigger an update if not done already */
+	void RegisterForUpdate();
+
+	/** Simple one-off active tick to update the group and stat trees */
+	EActiveTickReturnType UpdateData(double InCurrentTime, float InDeltaTime);
+
+	/** Whether the active tick has already been registered this frame */
+	bool bIsActiveTickRegistered;
 };
diff --git a/Engine/Source/Developer/Profiler/Private/Widgets/SProfilerMiniView.cpp b/Engine/Source/Developer/Profiler/Private/Widgets/SProfilerMiniView.cpp
--- a/Engine/Source/Developer/Profiler/Private/Widgets/SProfilerMiniView.cpp
+++ b/Engine/Source/Developer/Profiler/Private/Widgets/SProfilerMiniView.cpp
@@ -8,6 +8,7 @@
 #undef LOCTEXT_NAMESPACE
 
 SProfilerMiniView::SProfilerMiniView()
+	: bIsActiveTickRegistered( false )
 {
 	Reset();
 }
@@ -51,11 +52,20 @@ void SProfilerMiniView::Reset()
 	CursorType = EMiniviewCursor::Default;
 }
 
+EActiveTickReturnType SProfilerMiniView::EnsureDataUpdateDuringPreview(double InCurrentTime, float InDeltaTime)
+{
+	if (RecentlyAddedFrames.Num() > 0)
+	{
+		bUpdateData = true;
+		return EActiveTickReturnType::KeepTicking;
+	}
+	
+	bIsActiveTickRegistered = false;
+	return EActiveTickReturnType::StopTicking;
+}
 
 void SProfilerMiniView::Tick( const FGeometry& AllottedGeometry, const double InCurrentTime, const float InDeltaTime )
 {
-	SCompoundWidget::Tick( AllottedGeometry, InCurrentTime, InDeltaTime );
-
 	if( ThisGeometry != AllottedGeometry )
 	{
 		// Refresh.
@@ -65,11 +75,6 @@ void SProfilerMiniView::Tick( const FGeometry& AllottedGeometry, const double In
 
 	ThisGeometry = AllottedGeometry;
 
-	if( RecentlyAddedFrames.Num() > 0 )
-	{
-		bUpdateData = true;
-	}
-
 	if( IsReady() )
 	{
 		NumPixelsPerFrame = (float)AllottedGeometry.Size.X / (float)AllFrames.Num();
@@ -581,6 +586,12 @@ void SProfilerMiniView::AddThreadTime( int32 InFrameIndex, const TMap<uint32, fl
 	FrameThreadTimes.ThreadTimes = InThreadMS;
 	RecentlyAddedFrames.Add( FrameThreadTimes );
 
+	if (!bIsActiveTickRegistered)
+	{
+		bIsActiveTickRegistered = true;
+		RegisterActiveTick(0.f, FWidgetActiveTickDelegate::CreateSP(this, &SProfilerMiniView::EnsureDataUpdateDuringPreview));
+	}
+
 	StatMetadata = InStatMetaData;
 }
 
diff --git a/Engine/Source/Developer/Profiler/Private/Widgets/SProfilerMiniView.h b/Engine/Source/Developer/Profiler/Private/Widgets/SProfilerMiniView.h
--- a/Engine/Source/Developer/Profiler/Private/Widgets/SProfilerMiniView.h
+++ b/Engine/Source/Developer/Profiler/Private/Widgets/SProfilerMiniView.h
@@ -321,7 +321,7 @@ protected:
 		return AllFrames.Num() > 0;
 	}
 
-	bool ShouldUpdateData()
+	bool ShouldUpdateData() const
 	{
 		return bUpdateData;
 	}
@@ -395,4 +395,11 @@ protected:
 
 	/** Cursor type. */
 	EMiniviewCursor::Type CursorType;
+
+private:
+	/** Actively ticks the widget to update to process new frame data as long as any frames have been added recently */
+	EActiveTickReturnType EnsureDataUpdateDuringPreview( double InCurrentTime, float InDeltaTime );
+
+	/** True if the active tick is currently registered */
+	bool bIsActiveTickRegistered;
 };
\ No newline at end of file
diff --git a/Engine/Source/Developer/Profiler/Private/Widgets/SProfilerThreadView.cpp b/Engine/Source/Developer/Profiler/Private/Widgets/SProfilerThreadView.cpp
--- a/Engine/Source/Developer/Profiler/Private/Widgets/SProfilerThreadView.cpp
+++ b/Engine/Source/Developer/Profiler/Private/Widgets/SProfilerThreadView.cpp
@@ -60,8 +60,6 @@ void SProfilerThreadView::Reset()
 
 void SProfilerThreadView::Tick( const FGeometry& AllottedGeometry, const double InCurrentTime, const float InDeltaTime )
 {
-	SCompoundWidget::Tick( AllottedGeometry, InCurrentTime, InDeltaTime );
-
 	if( AllottedGeometry.Size.X > 0.0f )
 	{
 		if( ThisGeometry.Size.X != AllottedGeometry.Size.X )
diff --git a/Engine/Source/Developer/Profiler/Private/Widgets/SProfilerWindow.cpp b/Engine/Source/Developer/Profiler/Private/Widgets/SProfilerWindow.cpp
--- a/Engine/Source/Developer/Profiler/Private/Widgets/SProfilerWindow.cpp
+++ b/Engine/Source/Developer/Profiler/Private/Widgets/SProfilerWindow.cpp
@@ -82,7 +82,6 @@ static FText GetTextForNotification( const EProfilerNotificationTypes::Type Noti
 
 SProfilerWindow::SProfilerWindow()
 	: DurationActive(0.0f)
-	, bIsActive(false)
 {}
 
 SProfilerWindow::~SProfilerWindow()
@@ -491,28 +490,34 @@ void SProfilerWindow::SendingServiceSideCapture_Load( const FString Filename )
 	}
 }
 
-void SProfilerWindow::Tick(const FGeometry& AllottedGeometry, const double InCurrentTime, const float InDeltaTime)
+EActiveTickReturnType SProfilerWindow::UpdateActiveDuration( double InCurrentTime, float InDeltaTime )
 {
-	SCompoundWidget::Tick(AllottedGeometry, InCurrentTime, InDeltaTime);
+	DurationActive += InDeltaTime;
 
-	if (bIsActive)
-	{
-		DurationActive += InDeltaTime;
-	}
+	// The profiler window will explicitly unregister this active tick when the mouse leaves
+	return EActiveTickReturnType::KeepTicking;
 }
 
 void SProfilerWindow::OnMouseEnter(const FGeometry& MyGeometry, const FPointerEvent& MouseEvent)
 {
 	SCompoundWidget::OnMouseEnter(MyGeometry, MouseEvent);
 
-	bIsActive = true;
+	if (!ActiveTickHandle.IsValid())
+	{
+		ActiveTickHandle = RegisterActiveTick(0.f, FWidgetActiveTickDelegate::CreateSP(this, &SProfilerWindow::UpdateActiveDuration));
+	}
 }
 
 void SProfilerWindow::OnMouseLeave(const FPointerEvent& MouseEvent)
 {
 	SCompoundWidget::OnMouseLeave(MouseEvent);
 
-	bIsActive = false;
+	auto PinnedActiveTickHandle = ActiveTickHandle.Pin();
+	if (PinnedActiveTickHandle.IsValid())
+	{
+		UnRegisterActiveTick(PinnedActiveTickHandle.ToSharedRef());
+	}
+	
 }
 
 FReply SProfilerWindow::OnKeyDown( const FGeometry& MyGeometry, const FKeyEvent& InKeyEvent )
diff --git a/Engine/Source/Developer/Profiler/Private/Widgets/SProfilerWindow.h b/Engine/Source/Developer/Profiler/Private/Widgets/SProfilerWindow.h
--- a/Engine/Source/Developer/Profiler/Private/Widgets/SProfilerWindow.h
+++ b/Engine/Source/Developer/Profiler/Private/Widgets/SProfilerWindow.h
@@ -68,7 +68,7 @@ private:
 	 * @param  InCurrentTime  Current absolute real time
 	 * @param  InDeltaTime  Real time passed since last tick
 	 */
-	virtual void Tick( const FGeometry& AllottedGeometry, const double InCurrentTime, const float InDeltaTime ) override;
+	//virtual void Tick( const FGeometry& AllottedGeometry, const double InCurrentTime, const float InDeltaTime ) override;
 
 	/**
 	 * The system will use this event to notify a widget that the cursor has entered it. This event is NOT bubbled.
@@ -115,10 +115,17 @@ private:
 	 */
 	virtual FReply OnDragOver( const FGeometry& MyGeometry, const FDragDropEvent& DragDropEvent )  override;
 
+	/** Updates the amount of time the profiler has been active */
+	EActiveTickReturnType UpdateActiveDuration( double InCurrentTime, float InDeltaTime );
+
+	/** The handle to the active update duration tick */
+	TWeakPtr<FActiveTickHandle> ActiveTickHandle;
+
 	/** Holds all widgets for the profiler window like menu bar, toolbar and tabs. */
 	TSharedPtr<SVerticalBox> MainContentPanel;
 
 public:
+	
 	/** Holds all event graphs. */
 	TSharedPtr<SVerticalBox> EventGraphPanel;
 
@@ -146,7 +153,4 @@ public:
 
 	/** The number of seconds the profiler has been active */
 	float	DurationActive;
-
-	/** Whether or not the profiler widget is being used */
-	bool	bIsActive;
 };
\ No newline at end of file
diff --git a/Engine/Source/Developer/ProjectLauncher/Private/Widgets/SProjectLauncherPreviewPage.cpp b/Engine/Source/Developer/ProjectLauncher/Private/Widgets/SProjectLauncherPreviewPage.cpp
--- a/Engine/Source/Developer/ProjectLauncher/Private/Widgets/SProjectLauncherPreviewPage.cpp
+++ b/Engine/Source/Developer/ProjectLauncher/Private/Widgets/SProjectLauncherPreviewPage.cpp
@@ -786,16 +786,8 @@ void SProjectLauncherPreviewPage::Construct( const FArguments& InArgs, const FPr
 			]
 	];
 }
-END_SLATE_FUNCTION_BUILD_OPTIMIZATION
-
-
-/* SCompoundWidget overrides
- *****************************************************************************/
-
 void SProjectLauncherPreviewPage::Tick( const FGeometry& AllottedGeometry, const double InCurrentTime, const float InDeltaTime )
 {
-	SCompoundWidget::Tick(AllottedGeometry, InCurrentTime, InDeltaTime);
-
 	RefreshDeviceProxyList();
 }
 
diff --git a/Engine/Source/Developer/RealtimeProfiler/Private/SRealtimeProfilerLineGraph.cpp b/Engine/Source/Developer/RealtimeProfiler/Private/SRealtimeProfilerLineGraph.cpp
--- a/Engine/Source/Developer/RealtimeProfiler/Private/SRealtimeProfilerLineGraph.cpp
+++ b/Engine/Source/Developer/RealtimeProfiler/Private/SRealtimeProfilerLineGraph.cpp
@@ -413,8 +413,6 @@ void SRealtimeProfilerLineGraph::Tick( const FGeometry& AllottedGeometry, const
 			LastGeometry = AllottedGeometry;
 		}
 	}
-
-	SWidget::Tick( AllottedGeometry, InCurrentTime, InDeltaTime );
 }
 
 FReply SRealtimeProfilerLineGraph::OnStartButtonDown()
diff --git a/Engine/Source/Developer/SharedSettingsWidgets/Private/SPlatformSetupMessage.cpp b/Engine/Source/Developer/SharedSettingsWidgets/Private/SPlatformSetupMessage.cpp
--- a/Engine/Source/Developer/SharedSettingsWidgets/Private/SPlatformSetupMessage.cpp
+++ b/Engine/Source/Developer/SharedSettingsWidgets/Private/SPlatformSetupMessage.cpp
@@ -117,10 +117,9 @@ void SPlatformSetupMessage::Construct(const FArguments& InArgs, const FString& I
 	UpdateCache(true);
 }
 
-void SPlatformSetupMessage::Tick(const FGeometry& AllottedGeometry, const double InCurrentTime, const float InDeltaTime)
+void SPlatformSetupMessage::Tick( const FGeometry& AllottedGeometry, const double InCurrentTime, const float InDeltaTime )
 {
 	UpdateCache(false);
-	SCompoundWidget::Tick(AllottedGeometry, InCurrentTime, InDeltaTime);
 }
 
 int32 SPlatformSetupMessage::GetSetupStateAsInt() const
diff --git a/Engine/Source/Developer/SharedSettingsWidgets/Private/SSettingsEditorCheckoutNotice.cpp b/Engine/Source/Developer/SharedSettingsWidgets/Private/SSettingsEditorCheckoutNotice.cpp
--- a/Engine/Source/Developer/SharedSettingsWidgets/Private/SSettingsEditorCheckoutNotice.cpp
+++ b/Engine/Source/Developer/SharedSettingsWidgets/Private/SSettingsEditorCheckoutNotice.cpp
@@ -21,7 +21,7 @@ void SSettingsEditorCheckoutNotice::Construct( const FArguments& InArgs )
 	OnFileProbablyModifiedExternally = InArgs._OnFileProbablyModifiedExternally;
 	ConfigFilePath = InArgs._ConfigFilePath;
 
-	DefaultConfigCheckOutTimer = 0.0f;
+	LastDefaultConfigCheckOutTime = 0.0;
 	DefaultConfigCheckOutNeeded = false;
 	DefaultConfigQueryInProgress = true;
 
@@ -226,14 +226,10 @@ FSlateColor SSettingsEditorCheckoutNotice::HandleBorderBackgroundColor() const
 	return FinalColor;
 }
 
-void SSettingsEditorCheckoutNotice::Tick(const FGeometry& AllottedGeometry, const double InCurrentTime, const float InDeltaTime)
+void SSettingsEditorCheckoutNotice::Tick( const FGeometry& AllottedGeometry, const double InCurrentTime, const float InDeltaTime )
 {
-	SCompoundWidget::Tick(AllottedGeometry, InCurrentTime, InDeltaTime);
-
-	// cache selected settings object's configuration file state
-	DefaultConfigCheckOutTimer += InDeltaTime;
-
-	if (DefaultConfigCheckOutTimer >= 1.0f)
+	// cache selected settings object's configuration file state (purposely not in an active tick to make sure it appears while the user is active)
+	if (InCurrentTime - LastDefaultConfigCheckOutTime >= 1.0f)
 	{
 		bool NewCheckOutNeeded = false;
 
@@ -265,7 +261,7 @@ void SSettingsEditorCheckoutNotice::Tick(const FGeometry& AllottedGeometry, cons
 		}
 
 		DefaultConfigCheckOutNeeded = NewCheckOutNeeded;
-		DefaultConfigCheckOutTimer = 0.0f;
+		LastDefaultConfigCheckOutTime = InCurrentTime;
 	}
 }
 
diff --git a/Engine/Source/Developer/SharedSettingsWidgets/Public/SSettingsEditorCheckoutNotice.h b/Engine/Source/Developer/SharedSettingsWidgets/Public/SSettingsEditorCheckoutNotice.h
--- a/Engine/Source/Developer/SharedSettingsWidgets/Public/SSettingsEditorCheckoutNotice.h
+++ b/Engine/Source/Developer/SharedSettingsWidgets/Public/SSettingsEditorCheckoutNotice.h
@@ -28,7 +28,7 @@ public:
 	void Construct( const FArguments& InArgs );
 
 	// SWidget interface
-	virtual void Tick(const FGeometry& AllottedGeometry, const double InCurrentTime, const float InDeltaTime) override;
+	virtual void Tick( const FGeometry& AllottedGeometry, const double InCurrentTime, const float InDeltaTime ) override;
 	// End of SWidget interface
 
 	// Return true if the watched file can be modified (as of the last update)
@@ -37,7 +37,7 @@ public:
 	// Force an update on the next tick
 	void Invalidate()
 	{
-		DefaultConfigCheckOutTimer = 1.0f;
+		LastDefaultConfigCheckOutTime = 0.f;
 	}
 
 private:
@@ -86,5 +86,5 @@ private:
 	bool DefaultConfigQueryInProgress;
 
 	// Holds a timer for checking whether the section's configuration file needs to be checked out.
-	float DefaultConfigCheckOutTimer;
+	double LastDefaultConfigCheckOutTime;
 };
diff --git a/Engine/Source/Developer/SourceControl/Private/SSourceControlLogin.cpp b/Engine/Source/Developer/SourceControl/Private/SSourceControlLogin.cpp
--- a/Engine/Source/Developer/SourceControl/Private/SSourceControlLogin.cpp
+++ b/Engine/Source/Developer/SourceControl/Private/SSourceControlLogin.cpp
@@ -55,6 +55,8 @@ private:
 	TWeakPtr<SWindow> ParentWindowPtr;
 };
 
+const float SSourceControlLogin::RefreshFrequency = 0.5f;
+
 void SSourceControlLogin::Construct(const FArguments& InArgs)
 {
 	ParentWindowPtr = InArgs._ParentWindow;
@@ -211,6 +213,12 @@ FReply SSourceControlLogin::OnAcceptSettings()
 {
 	ConnectionState = ELoginConnectionState::Connecting;
 
+	//Increase the tick frequency during login if needed
+	if ( ParentWindowPtr.IsValid() && ( FSlateApplication::Get().GetActiveModalWindow() == ParentWindowPtr.Pin() ) )
+	{
+		ActiveTickHandle = RegisterActiveTick( 0.f, FWidgetActiveTickDelegate::CreateSP( this, &SSourceControlLogin::TickSourceControlModule ) );
+	}
+
 	FSourceControlModule& SourceControlModule = FSourceControlModule::Get();
 	if(!SourceControlModule.GetProvider().Login(FString(), EConcurrency::Asynchronous, FSourceControlOperationComplete::CreateSP(this, &SSourceControlLogin::SourceControlOperationComplete)))
 	{
@@ -255,7 +263,7 @@ void SSourceControlLogin::SourceControlOperationComplete(const FSourceControlOpe
 	}
 }
 
-void SSourceControlLogin::DisplayConnectionError(const FText& InErrorText) const
+void SSourceControlLogin::DisplayConnectionError(const FText& InErrorText)
 {
 	FMessageLog SourceControlLog("SourceControl");
 	if(InErrorText.IsEmpty())
@@ -268,6 +276,13 @@ void SSourceControlLogin::DisplayConnectionError(const FText& InErrorText) const
 	}
 	
 	SourceControlLog.Notify();
+
+	// Suspend the active tick until there's another login attempt
+	auto PinnedActiveTick = ActiveTickHandle.Pin();
+	if ( PinnedActiveTick.IsValid() )
+	{
+		UnRegisterActiveTick( PinnedActiveTick.ToSharedRef() );
+	}
 }
 
 void SSourceControlLogin::DisplayConnectionSuccess() const
@@ -307,19 +322,11 @@ EVisibility SSourceControlLogin::GetDisabledTextVisibility() const
 	return SourceControlModule.GetProvider().GetName() == "None" ? EVisibility::Visible : EVisibility::Collapsed;
 }
 
-void SSourceControlLogin::Tick( const FGeometry& AllottedGeometry, const double InCurrentTime, const float InDeltaTime )
+EActiveTickReturnType SSourceControlLogin::TickSourceControlModule( double InCurrentTime, float InDeltaTime )
 {
-	// if we are modal then we need to tick the source control module
-	TSharedPtr<SWindow> ParentWindow = ParentWindowPtr.Pin();
-	if(ParentWindow.IsValid())
-	{
-		if(FSlateApplication::Get().GetActiveModalWindow() == ParentWindow)
-		{
-			FSourceControlModule& SourceControlModule = FSourceControlModule::Get();
-			SourceControlModule.Tick();
-		}
-	}
-};
+	FSourceControlModule::Get().Tick();
+	return EActiveTickReturnType::KeepTicking;
+}
 
 #undef LOCTEXT_NAMESPACE
 
diff --git a/Engine/Source/Developer/SourceControl/Private/SSourceControlLogin.h b/Engine/Source/Developer/SourceControl/Private/SSourceControlLogin.h
--- a/Engine/Source/Developer/SourceControl/Private/SSourceControlLogin.h
+++ b/Engine/Source/Developer/SourceControl/Private/SSourceControlLogin.h
@@ -34,9 +34,6 @@ public:
 
 	void Construct(const FArguments& InArgs);
 
-	/** SWidget interface */
-	virtual void Tick( const FGeometry& AllottedGeometry, const double InCurrentTime, const float InDeltaTime ) override;
-
 	/**
 	 * Refresh the displayed settings. Usually called when a provider is changed.
 	 */
@@ -44,6 +41,9 @@ public:
 
 private:
 
+	/** Ticks the source control module (only necessary if the login window is modal) */
+	EActiveTickReturnType TickSourceControlModule( double InCurrentTime, float InDeltaTime );
+
 	/** Delegate called when the user clicks the 'Accept Settings' button */
 	FReply OnAcceptSettings();
 
@@ -51,7 +51,7 @@ private:
 	FReply OnDisableSourceControl();
 
 	/** Called when a connection attempt fails */
-	void DisplayConnectionError(const FText& InErrorText = FText()) const;
+	void DisplayConnectionError(const FText& InErrorText = FText());
 
 	/** Called when a connection attempt succeeds */
 	void DisplayConnectionSuccess() const;
@@ -75,6 +75,8 @@ private:
 	EVisibility GetDisabledTextVisibility() const;
 
 private:
+	/** The frequency at which to tick to scc module when inside a modal window */
+	static const float RefreshFrequency;
 
 	/** The parent window of this widget */
 	TWeakPtr<SWindow> ParentWindowPtr;
@@ -90,6 +92,9 @@ private:
 
 	/** The currently displayed settings widget container */
 	TSharedPtr<class SBorder> SettingsBorder;
+
+	/** The handle to the active scc module tick */
+	TWeakPtr<FActiveTickHandle> ActiveTickHandle;
 };
 
 #endif // SOURCE_CONTROL_WITH_SLATE
diff --git a/Engine/Source/Developer/TaskGraph/Private/SGraphBar.cpp b/Engine/Source/Developer/TaskGraph/Private/SGraphBar.cpp
--- a/Engine/Source/Developer/TaskGraph/Private/SGraphBar.cpp
+++ b/Engine/Source/Developer/TaskGraph/Private/SGraphBar.cpp
@@ -115,8 +115,6 @@ void SGraphBar::Tick( const FGeometry& AllottedGeometry, const double InCurrentT
 			LastGeometry = AllottedGeometry;
 		}
 	}
-
-	SLeafWidget::Tick( AllottedGeometry, InCurrentTime, InDeltaTime );
 }
 
 FReply SGraphBar::OnMouseButtonDown( const FGeometry& MyGeometry, const FPointerEvent& MouseEvent )
diff --git a/Engine/Source/Developer/TaskGraph/Private/STimeline.cpp b/Engine/Source/Developer/TaskGraph/Private/STimeline.cpp
--- a/Engine/Source/Developer/TaskGraph/Private/STimeline.cpp
+++ b/Engine/Source/Developer/TaskGraph/Private/STimeline.cpp
@@ -165,8 +165,6 @@ int32 STimeline::OnPaint( const FPaintArgs& Args, const FGeometry& AllottedGeome
 void STimeline::Tick( const FGeometry& AllottedGeometry, const double InCurrentTime, const float InDeltaTime )
 {
 	DrawingOffsetX = DrawingGeometry.AbsolutePosition.X - AllottedGeometry.AbsolutePosition.X;
-
-	SCompoundWidget::Tick( AllottedGeometry, InCurrentTime, InDeltaTime );
 }
 
 FReply STimeline::OnMouseButtonDown( const FGeometry& MyGeometry, const FPointerEvent& MouseEvent )
diff --git a/Engine/Source/Developer/WidgetCarousel/Public/SWidgetCarousel.h b/Engine/Source/Developer/WidgetCarousel/Public/SWidgetCarousel.h
--- a/Engine/Source/Developer/WidgetCarousel/Public/SWidgetCarousel.h
+++ b/Engine/Source/Developer/WidgetCarousel/Public/SWidgetCarousel.h
@@ -74,7 +74,7 @@ struct FCarouselDisplayItem : public TSharedFromThis<FCarouselDisplayItem>
 		if (!bTransition)
 		{
 			SlideInCurve.JumpToEnd();
-			SlideInCurve.PlayReverse();
+			SlideInCurve.PlayReverse(FXWidget.ToSharedRef());
 			bPeak = true;
 			bFade = FadeRate != 0;
 			switch (ScrollDirection)
@@ -143,7 +143,7 @@ struct FCarouselDisplayItem : public TSharedFromThis<FCarouselDisplayItem>
 			if (FMath::Abs(BlendDif) > PeakDistance + (float)FLOAT_NORMAL_THRESH)
 			{
 				SlideInCurve.JumpToEnd();
-				SlideInCurve.PlayReverse();
+				SlideInCurve.PlayReverse(FXWidget.ToSharedRef());
 			}
 			float Lerp = FMath::Max(SlideInCurve.GetLerp(), 0.1f);
 			DesiredBlendSpeed = Speed * Lerp * DeltaTime;
diff --git a/Engine/Source/Editor/BehaviorTreeEditor/Private/SGraphNode_BehaviorTree.h b/Engine/Source/Editor/BehaviorTreeEditor/Private/SGraphNode_BehaviorTree.h
--- a/Engine/Source/Editor/BehaviorTreeEditor/Private/SGraphNode_BehaviorTree.h
+++ b/Engine/Source/Editor/BehaviorTreeEditor/Private/SGraphNode_BehaviorTree.h
@@ -44,7 +44,7 @@ public:
 	 * @param  InCurrentTime  Current absolute real time
 	 * @param  InDeltaTime  Real time passed since last tick
 	 */
-	void Tick( const FGeometry& AllottedGeometry, const double InCurrentTime, const float InDeltaTime ) override;
+	virtual void Tick( const FGeometry& AllottedGeometry, const double InCurrentTime, const float InDeltaTime ) override;
 	
 	/** gets decorator or service node if one is found under mouse cursor */
 	TSharedPtr<SGraphNode> GetSubNodeUnderCursor(const FGeometry& WidgetGeometry, const FPointerEvent& MouseEvent) const;
diff --git a/Engine/Source/Editor/ClassViewer/Private/SClassViewer.cpp b/Engine/Source/Editor/ClassViewer/Private/SClassViewer.cpp
--- a/Engine/Source/Editor/ClassViewer/Private/SClassViewer.cpp
+++ b/Engine/Source/Editor/ClassViewer/Private/SClassViewer.cpp
@@ -2962,9 +2962,7 @@ void SClassViewer::Refresh()
 }
 
 void SClassViewer::Tick( const FGeometry& AllottedGeometry, const double InCurrentTime, const float InDeltaTime )
-{
-	SCompoundWidget::Tick(AllottedGeometry, InCurrentTime, InDeltaTime);
-	
+{	
 	// Will populate the class hierarchy as needed.
 	ClassViewer::Helpers::PopulateClassHierarchy();
 
diff --git a/Engine/Source/Editor/ContentBrowser/Private/AssetViewWidgets.cpp b/Engine/Source/Editor/ContentBrowser/Private/AssetViewWidgets.cpp
--- a/Engine/Source/Editor/ContentBrowser/Private/AssetViewWidgets.cpp
+++ b/Engine/Source/Editor/ContentBrowser/Private/AssetViewWidgets.cpp
@@ -135,6 +135,7 @@ void SAssetViewItem::Construct( const FArguments& InArgs )
 	OnFilesDragDropped = InArgs._OnFilesDragDropped;
 	OnGetCustomAssetToolTip = InArgs._OnGetCustomAssetToolTip;
 	OnVisualizeAssetToolTip = InArgs._OnVisualizeAssetToolTip;
+	OnAssetToolTipClosing = InArgs._OnAssetToolTipClosing;
 
 	bDraggedOver = false;
 
@@ -164,8 +165,6 @@ void SAssetViewItem::Construct( const FArguments& InArgs )
 
 void SAssetViewItem::Tick( const FGeometry& AllottedGeometry, const double InCurrentTime, const float InDeltaTime )
 {
-	SCompoundWidget::Tick(AllottedGeometry, InCurrentTime, InDeltaTime);
-
 	const float PrevSizeX = LastGeometry.Size.X;
 
 	LastGeometry = AllottedGeometry;
@@ -941,6 +940,11 @@ bool SAssetViewItem::OnVisualizeTooltip(const TSharedPtr<SWidget>& TooltipConten
 	return false;
 }
 
+void SAssetViewItem::OnToolTipClosing()
+{
+	OnAssetToolTipClosing.ExecuteIfBound();
+}
+
 ///////////////////////////////
 // SAssetListItem
 ///////////////////////////////
@@ -967,6 +971,7 @@ void SAssetListItem::Construct( const FArguments& InArgs )
 		.OnFilesDragDropped(InArgs._OnFilesDragDropped)
 		.OnGetCustomAssetToolTip(InArgs._OnGetCustomAssetToolTip)
 		.OnVisualizeAssetToolTip(InArgs._OnVisualizeAssetToolTip)
+		.OnAssetToolTipClosing( InArgs._OnAssetToolTipClosing )
 		);
 
 	AssetThumbnail = InArgs._AssetThumbnail;
@@ -1172,6 +1177,7 @@ void SAssetTileItem::Construct( const FArguments& InArgs )
 		.OnFilesDragDropped(InArgs._OnFilesDragDropped)
 		.OnGetCustomAssetToolTip(InArgs._OnGetCustomAssetToolTip)
 		.OnVisualizeAssetToolTip(InArgs._OnVisualizeAssetToolTip)
+		.OnAssetToolTipClosing( InArgs._OnAssetToolTipClosing )
 		);
 
 	AssetThumbnail = InArgs._AssetThumbnail;
@@ -1366,6 +1372,49 @@ FSlateFontInfo SAssetTileItem::GetThumbnailFont() const
 // SAssetColumnItem
 ///////////////////////////////
 
+/** Custom box for the Name column of an asset */
+class SAssetColumnItemNameBox : public SCompoundWidget
+{
+public:
+	SLATE_BEGIN_ARGS( SAssetColumnItemNameBox ) {}
+
+		/** The color of the asset  */
+		SLATE_ATTRIBUTE( FMargin, Padding )
+
+		/** The widget content presented in the box */
+		SLATE_DEFAULT_SLOT(FArguments, Content)
+
+	SLATE_END_ARGS()
+
+	~SAssetColumnItemNameBox() {}
+
+	void Construct( const FArguments& InArgs, const TSharedRef<SAssetColumnItem>& InOwnerAssetColumnItem )
+	{
+		OwnerAssetColumnItem = InOwnerAssetColumnItem;
+
+		ChildSlot
+		[
+			SNew(SBox)
+			.Padding(InArgs._Padding)
+			[
+				InArgs._Content.Widget
+			]
+		];
+	}
+
+	/** Forward the event to the view item that this name box belongs to */
+	virtual void OnToolTipClosing() override
+	{
+		if ( OwnerAssetColumnItem.IsValid() )
+		{
+			OwnerAssetColumnItem.Pin()->OnToolTipClosing();
+		}
+	}
+
+private:
+	TWeakPtr<SAssetViewItem> OwnerAssetColumnItem;
+};
+
 void SAssetColumnItem::Construct( const FArguments& InArgs )
 {
 	SAssetViewItem::Construct( SAssetViewItem::FArguments()
@@ -1380,6 +1429,7 @@ void SAssetColumnItem::Construct( const FArguments& InArgs )
 		.OnFilesDragDropped(InArgs._OnFilesDragDropped)
 		.OnGetCustomAssetToolTip(InArgs._OnGetCustomAssetToolTip)
 		.OnVisualizeAssetToolTip(InArgs._OnVisualizeAssetToolTip)
+		.OnAssetToolTipClosing(InArgs._OnAssetToolTipClosing)
 		);
 	
 	HighlightText = InArgs._HighlightText;
@@ -1390,6 +1440,9 @@ TSharedRef<SWidget> SAssetColumnItem::GenerateWidgetForColumn( const FName& Colu
 {
 	TSharedPtr<SWidget> Content;
 
+	// A little right padding so text from this column does not run directly into text from the next.
+	static const FMargin ColumnItemPadding( 0, 0, 6, 0 );
+
 	if ( ColumnName == "Name" )
 	{
 		const FSlateBrush* IconBrush;
@@ -1476,6 +1529,13 @@ TSharedRef<SWidget> SAssetColumnItem::GenerateWidgetForColumn( const FName& Colu
 		{
 			AssetItem->RenamedRequestEvent.BindSP( InlineRenameWidget.Get(), &SInlineEditableTextBlock::EnterEditingMode );
 		}
+
+		return SNew( SAssetColumnItemNameBox, SharedThis(this) )
+			.Padding( ColumnItemPadding )
+			.ToolTip( CreateToolTipWidget() )
+			[
+				Content.ToSharedRef()
+			];
 	}
 	else if ( ColumnName == "Class" )
 	{
@@ -1499,8 +1559,7 @@ TSharedRef<SWidget> SAssetColumnItem::GenerateWidgetForColumn( const FName& Colu
 	}
 
 	return SNew(SBox)
-		.Padding(FMargin(0, 0, 6, 0)) // Add a little right padding so text from this column does not run directly into text from the next.
-		.ToolTip(CreateToolTipWidget())
+		.Padding( ColumnItemPadding )
 		[
 			Content.ToSharedRef()
 		];
@@ -1520,8 +1579,6 @@ void SAssetColumnItem::OnAssetDataChanged()
 	{
 		PathText->SetText( GetAssetPathText() );
 	}
-
-	SetToolTip( CreateToolTipWidget() );
 }
 
 FString SAssetColumnItem::GetAssetNameToolTipText() const
diff --git a/Engine/Source/Editor/ContentBrowser/Private/AssetViewWidgets.h b/Engine/Source/Editor/ContentBrowser/Private/AssetViewWidgets.h
--- a/Engine/Source/Editor/ContentBrowser/Private/AssetViewWidgets.h
+++ b/Engine/Source/Editor/ContentBrowser/Private/AssetViewWidgets.h
@@ -88,6 +88,9 @@ public:
 		/** Delegate for when an item is about to show a tool tip */
 		SLATE_EVENT( FOnVisualizeAssetToolTip, OnVisualizeAssetToolTip)
 
+		/** Delegate for when an item's tooltip is about to close */
+		SLATE_EVENT( FOnAssetToolTipClosing, OnAssetToolTipClosing )
+
 	SLATE_END_ARGS()
 
 	/** Virtual destructor */
@@ -115,9 +118,12 @@ public:
 	/** Delegate handling when an asset is loaded */
 	void HandleAssetLoaded(UObject* InAsset) const;
 
-	/* About to show a tool tip */
+	/** About to show a tool tip */
 	virtual bool OnVisualizeTooltip( const TSharedPtr<SWidget>& TooltipContent ) override;
 
+	/** Tooltip is closing */
+	virtual void OnToolTipClosing() override;
+
 protected:
 	/** Handles starting a name change */
 	virtual void HandleBeginNameChange( const FText& OriginalText );
@@ -221,6 +227,9 @@ protected:
 	/** Called if bound when about to show a tooltip */
 	FOnVisualizeAssetToolTip OnVisualizeAssetToolTip;
 
+	/** Called if bound when a tooltip is closing */
+	FOnAssetToolTipClosing OnAssetToolTipClosing;
+
 	/** The geometry last frame. Used when telling popup messages where to appear. */
 	FGeometry LastGeometry;
 
@@ -336,6 +345,9 @@ public:
 		/* Delegate to signal when the item is about to show a tooltip */
 		SLATE_EVENT(FOnVisualizeAssetToolTip, OnVisualizeAssetToolTip)
 
+		/** Delegate for when an item's tooltip is about to close */
+		SLATE_EVENT( FOnAssetToolTipClosing, OnAssetToolTipClosing )
+
 	SLATE_END_ARGS()
 
 	/** Destructor */
@@ -439,6 +451,9 @@ public:
 		/* Delegate to signal when the item is about to show a tooltip */
 		SLATE_EVENT(FOnVisualizeAssetToolTip, OnVisualizeAssetToolTip)
 
+		/** Delegate for when an item's tooltip is about to close */
+		SLATE_EVENT( FOnAssetToolTipClosing, OnAssetToolTipClosing )
+
 	SLATE_END_ARGS()
 
 	/** Destructor */
@@ -474,7 +489,7 @@ private:
 	float ThumbnailPadding;
 };
 
-/** An item in the asset tile view */
+/** An item in the asset column view */
 class SAssetColumnItem : public SAssetViewItem
 {
 public:
@@ -513,6 +528,9 @@ public:
 		/* Delegate to signal when the item is about to show a tooltip */
 		SLATE_EVENT( FOnVisualizeAssetToolTip, OnVisualizeAssetToolTip)
 
+		/** Delegate for when an item's tooltip is about to close */
+		SLATE_EVENT( FOnAssetToolTipClosing, OnAssetToolTipClosing )
+
 	SLATE_END_ARGS()
 
 	/** Constructs this widget with InArgs */
@@ -581,8 +599,6 @@ public:
 
 	virtual void Tick( const FGeometry& AllottedGeometry, const double InCurrentTime, const float InDeltaTime ) override
 	{
-		SMultiColumnTableRow< TSharedPtr<FAssetViewItem> >::Tick(AllottedGeometry, InCurrentTime, InDeltaTime);
-
 		this->AssetColumnItem->Tick(AllottedGeometry, InCurrentTime, InDeltaTime);
 	}
 
diff --git a/Engine/Source/Editor/ContentBrowser/Private/SAssetDialog.cpp b/Engine/Source/Editor/ContentBrowser/Private/SAssetDialog.cpp
--- a/Engine/Source/Editor/ContentBrowser/Private/SAssetDialog.cpp
+++ b/Engine/Source/Editor/ContentBrowser/Private/SAssetDialog.cpp
@@ -31,6 +31,8 @@ void SAssetDialog::Construct(const FArguments& InArgs, const FSharedAssetDialogC
 
 	const FString DefaultPath = InConfig.DefaultPath;
 
+	RegisterActiveTick( 0.f, FWidgetActiveTickDelegate::CreateSP( this, &SAssetDialog::SetFocusPostConstruct ) );
+
 	FPathPickerConfig PathPickerConfig;
 	PathPickerConfig.DefaultPath = DefaultPath;
 	PathPickerConfig.bFocusSearchBoxWhenOpened = false;
@@ -247,15 +249,10 @@ FReply SAssetDialog::OnKeyDown( const FGeometry& MyGeometry, const FKeyEvent& In
 	return SCompoundWidget::OnKeyDown(MyGeometry, InKeyEvent);
 }
 
-void SAssetDialog::Tick(const FGeometry& AllottedGeometry, const double InCurrentTime, const float InDeltaTime)
+EActiveTickReturnType SAssetDialog::SetFocusPostConstruct( double InCurrentTime, float InDeltaTime )
 {
-	if (bPendingFocusNextFrame)
-	{
-		FocusNameBox();
-		bPendingFocusNextFrame = false;
-	}
-
-	SCompoundWidget::Tick(AllottedGeometry, InCurrentTime, InDeltaTime);
+	FocusNameBox();
+	return EActiveTickReturnType::StopTicking;
 }
 
 void SAssetDialog::SetOnAssetsChosenForOpen(const FOnAssetsChosenForOpen& InOnAssetsChosenForOpen)
diff --git a/Engine/Source/Editor/ContentBrowser/Private/SAssetDialog.h b/Engine/Source/Editor/ContentBrowser/Private/SAssetDialog.h
--- a/Engine/Source/Editor/ContentBrowser/Private/SAssetDialog.h
+++ b/Engine/Source/Editor/ContentBrowser/Private/SAssetDialog.h
@@ -15,7 +15,6 @@ public:
 	virtual void Construct(const FArguments& InArgs, const FSharedAssetDialogConfig& InConfig);
 
 	virtual FReply OnKeyDown( const FGeometry& MyGeometry, const FKeyEvent& InKeyEvent ) override;
-	void Tick(const FGeometry& AllottedGeometry, const double InCurrentTime, const float InDeltaTime) override;
 
 	/** Sets the delegate handler for when an open operation is committed */
 	void SetOnAssetsChosenForOpen(const FOnAssetsChosenForOpen& InOnAssetsChosenForOpen);
@@ -28,6 +27,9 @@ public:
 
 private:
 
+	/** Used to focus the name box immediately following construction */
+	EActiveTickReturnType SetFocusPostConstruct( double InCurrentTime, float InDeltaTime );
+
 	/** Moves keyboard focus to the name box if this is a save dialog */
 	void FocusNameBox();
 
diff --git a/Engine/Source/Editor/ContentBrowser/Private/SAssetPicker.cpp b/Engine/Source/Editor/ContentBrowser/Private/SAssetPicker.cpp
--- a/Engine/Source/Editor/ContentBrowser/Private/SAssetPicker.cpp
+++ b/Engine/Source/Editor/ContentBrowser/Private/SAssetPicker.cpp
@@ -23,6 +23,11 @@ void SAssetPicker::Construct( const FArguments& InArgs )
 	SaveSettingsName = InArgs._AssetPickerConfig.SaveSettingsName;
 	OnFolderEnteredDelegate = InArgs._AssetPickerConfig.OnFolderEntered;
 
+	if ( InArgs._AssetPickerConfig.bFocusSearchBoxWhenOpened )
+	{
+		RegisterActiveTick( 0.f, FWidgetActiveTickDelegate::CreateSP( this, &SAssetPicker::SetFocusPostConstruct ) );
+	}
+
 	for (auto DelegateIt = InArgs._AssetPickerConfig.GetCurrentSelectionDelegates.CreateConstIterator(); DelegateIt; ++DelegateIt)
 	{
 		if ((*DelegateIt) != NULL)
@@ -202,6 +207,7 @@ void SAssetPicker::Construct( const FArguments& InArgs )
 		.OnGetAssetContextMenu(InArgs._AssetPickerConfig.OnGetAssetContextMenu)
 		.OnGetCustomAssetToolTip(InArgs._AssetPickerConfig.OnGetCustomAssetToolTip)
 		.OnVisualizeAssetToolTip(InArgs._AssetPickerConfig.OnVisualizeAssetToolTip)
+		.OnAssetToolTipClosing(InArgs._AssetPickerConfig.OnAssetToolTipClosing)
 		.AreRealTimeThumbnailsAllowed(this, &SAssetPicker::IsHovered)
 		.FrontendFilters(FrontendFilters)
 		.InitialSourcesData(CurrentSourcesData)
@@ -230,17 +236,18 @@ void SAssetPicker::Construct( const FArguments& InArgs )
 	AssetViewPtr->RequestSlowFullListRefresh();
 }
 
-void SAssetPicker::Tick(const FGeometry& AllottedGeometry, const double InCurrentTime, const float InDeltaTime)
+EActiveTickReturnType SAssetPicker::SetFocusPostConstruct( double InCurrentTime, float InDeltaTime )
 {
-	if ( bPendingFocusNextFrame && SearchBoxPtr.IsValid() )
+	if ( SearchBoxPtr.IsValid() )
 	{
 		FWidgetPath WidgetToFocusPath;
 		FSlateApplication::Get().GeneratePathToWidgetUnchecked( SearchBoxPtr.ToSharedRef(), WidgetToFocusPath );
 		FSlateApplication::Get().SetKeyboardFocus( WidgetToFocusPath, EFocusCause::SetDirectly );
-		bPendingFocusNextFrame = false;
+
+		return EActiveTickReturnType::StopTicking;
 	}
 
-	SCompoundWidget::Tick(AllottedGeometry, InCurrentTime, InDeltaTime);
+	return EActiveTickReturnType::KeepTicking;
 }
 
 FReply SAssetPicker::OnKeyDown(const FGeometry& MyGeometry, const FKeyEvent& InKeyEvent)
diff --git a/Engine/Source/Editor/ContentBrowser/Private/SAssetPicker.h b/Engine/Source/Editor/ContentBrowser/Private/SAssetPicker.h
--- a/Engine/Source/Editor/ContentBrowser/Private/SAssetPicker.h
+++ b/Engine/Source/Editor/ContentBrowser/Private/SAssetPicker.h
@@ -21,11 +21,13 @@ public:
 	void Construct( const FArguments& InArgs );
 
 	// SWidget implementation
-	virtual void Tick( const FGeometry& AllottedGeometry, const double InCurrentTime, const float InDeltaTime ) override;
 	virtual FReply OnKeyDown( const FGeometry& MyGeometry, const FKeyEvent& InKeyEvent ) override;
 	// End of SWidget implementation
 
 private:
+	/** Focuses the search box post-construct */
+	EActiveTickReturnType SetFocusPostConstruct( double InCurrentTime, float InDeltaTime );
+
 	void FolderEntered(const FString& FolderPath);
 
 	/** Called when the editable text needs to be set or cleared */
diff --git a/Engine/Source/Editor/ContentBrowser/Private/SAssetView.cpp b/Engine/Source/Editor/ContentBrowser/Private/SAssetView.cpp
--- a/Engine/Source/Editor/ContentBrowser/Private/SAssetView.cpp
+++ b/Engine/Source/Editor/ContentBrowser/Private/SAssetView.cpp
@@ -60,7 +60,7 @@ SAssetView::~SAssetView()
 BEGIN_SLATE_FUNCTION_BUILD_OPTIMIZATION
 void SAssetView::Construct( const FArguments& InArgs )
 {
-	IsWorking = false;
+	bIsWorking = false;
 	TotalAmortizeTime = 0;
 	AmortizeStartTime = 0;
 	MaxSecondsPerFrame = 0.015;
@@ -158,6 +158,7 @@ void SAssetView::Construct( const FArguments& InArgs )
 	OnAssetTagWantsToBeDisplayed = InArgs._OnAssetTagWantsToBeDisplayed;
 	OnGetCustomAssetToolTip = InArgs._OnGetCustomAssetToolTip;
 	OnVisualizeAssetToolTip = InArgs._OnVisualizeAssetToolTip;
+	OnAssetToolTipClosing = InArgs._OnAssetToolTipClosing;
 	HighlightedText = InArgs._HighlightedText;
 	ThumbnailLabel = InArgs._ThumbnailLabel;
 	AllowThumbnailHintLabel = InArgs._AllowThumbnailHintLabel;
@@ -213,7 +214,7 @@ void SAssetView::Construct( const FArguments& InArgs )
 			SNew( SBox )
 			.HeightOverride( 2 )
 			[
-				SNew(SProgressBar)
+				SNew( SProgressBar )
 				.Percent( this, &SAssetView::GetIsWorkingProgressBarState )
 				.Style( FEditorStyle::Get(), "WorkingBar" )
 				.BorderPadding( FVector2D(0,0) )
@@ -379,7 +380,7 @@ END_SLATE_FUNCTION_BUILD_OPTIMIZATION
 
 TOptional< float > SAssetView::GetIsWorkingProgressBarState() const
 {
-	return IsWorking ? TOptional< float >() : 0.0; 
+	return bIsWorking ? TOptional< float >() : 0.0f; 
 }
 
 void SAssetView::SetSourcesData(const FSourcesData& InSourcesData)
@@ -780,10 +781,8 @@ void SAssetView::ProcessRecentlyLoadedOrChangedAssets()
 	}
 }
 
-void SAssetView::Tick(const FGeometry& AllottedGeometry, const double InCurrentTime, const float InDeltaTime)
+void SAssetView::Tick( const FGeometry& AllottedGeometry, const double InCurrentTime, const float InDeltaTime )
 {
-	SCompoundWidget::Tick(AllottedGeometry, InCurrentTime, InDeltaTime);
-
 	CalculateFillScale( AllottedGeometry );
 
 	CurrentTime = InCurrentTime;
@@ -829,7 +828,7 @@ void SAssetView::Tick(const FGeometry& AllottedGeometry, const double InCurrentT
 		if ( AmortizeStartTime == 0 )
 		{
 			AmortizeStartTime = FPlatformTime::Seconds();
-			IsWorking = true;
+			bIsWorking = true;
 		}
 
 		ProcessQueriedItems( TickStartTime );
@@ -838,7 +837,7 @@ void SAssetView::Tick(const FGeometry& AllottedGeometry, const double InCurrentT
 		{
 			TotalAmortizeTime += FPlatformTime::Seconds() - AmortizeStartTime;
 			AmortizeStartTime = 0;
-			IsWorking = false;
+			bIsWorking = false;
 		}
 	}
 
@@ -897,8 +896,6 @@ void SAssetView::Tick(const FGeometry& AllottedGeometry, const double InCurrentT
 	DeferredCreateNewAsset();
 	DeferredCreateNewFolder();
 
-	AssetThumbnailPool->Tick(InDeltaTime);
-
 	// Do quick-jump last as the Tick function might have canceled it
 	if(QuickJumpData.bHasChangedSinceLastTick)
 	{
@@ -970,7 +967,7 @@ void SAssetView::CalculateThumbnailHintColorAndOpacity()
 		}
 		else if ( ThumbnailHintFadeInSequence.IsAtEnd() ) 
 		{
-			ThumbnailHintFadeInSequence.PlayReverse();
+			ThumbnailHintFadeInSequence.PlayReverse(this->AsShared());
 		}
 	}
 	else 
@@ -984,7 +981,7 @@ void SAssetView::CalculateThumbnailHintColorAndOpacity()
 		}
 		else if ( ThumbnailHintFadeInSequence.IsAtStart() ) 
 		{
-			ThumbnailHintFadeInSequence.Play();
+			ThumbnailHintFadeInSequence.Play(this->AsShared());
 		}
 	}
 
@@ -2663,7 +2660,8 @@ TSharedRef<ITableRow> SAssetView::MakeListViewWidget(TSharedPtr<FAssetViewItem>
 			.AllowThumbnailHintLabel( AllowThumbnailHintLabel )
 			.IsSelected( FIsSelected::CreateSP(TableRowWidget.Get(), &STableRow<TSharedPtr<FAssetViewItem>>::IsSelectedExclusively) )
 			.OnGetCustomAssetToolTip(OnGetCustomAssetToolTip)
-			.OnVisualizeAssetToolTip(OnVisualizeAssetToolTip);
+			.OnVisualizeAssetToolTip(OnVisualizeAssetToolTip)
+			.OnAssetToolTipClosing(OnAssetToolTipClosing);
 
 		TableRowWidget->SetContent(Item);
 
@@ -2750,7 +2748,8 @@ TSharedRef<ITableRow> SAssetView::MakeTileViewWidget(TSharedPtr<FAssetViewItem>
 			.AllowThumbnailHintLabel( AllowThumbnailHintLabel )
 			.IsSelected( FIsSelected::CreateSP(TableRowWidget.Get(), &STableRow<TSharedPtr<FAssetViewItem>>::IsSelectedExclusively) )
 			.OnGetCustomAssetToolTip(OnGetCustomAssetToolTip)
-			.OnVisualizeAssetToolTip(OnVisualizeAssetToolTip);
+			.OnVisualizeAssetToolTip( OnVisualizeAssetToolTip )
+			.OnAssetToolTipClosing( OnAssetToolTipClosing );
 
 		TableRowWidget->SetContent(Item);
 
@@ -2782,7 +2781,8 @@ TSharedRef<ITableRow> SAssetView::MakeColumnViewWidget(TSharedPtr<FAssetViewItem
 				.OnPathsDragDropped(this, &SAssetView::OnPathsDragDropped)
 				.OnFilesDragDropped(this, &SAssetView::OnFilesDragDropped)
 				.OnGetCustomAssetToolTip(OnGetCustomAssetToolTip)
-				.OnVisualizeAssetToolTip(OnVisualizeAssetToolTip)
+				.OnVisualizeAssetToolTip( OnVisualizeAssetToolTip )
+				.OnAssetToolTipClosing( OnAssetToolTipClosing )
 		);
 }
 
diff --git a/Engine/Source/Editor/ContentBrowser/Private/SCollectionView.cpp b/Engine/Source/Editor/ContentBrowser/Private/SCollectionView.cpp
--- a/Engine/Source/Editor/ContentBrowser/Private/SCollectionView.cpp
+++ b/Engine/Source/Editor/ContentBrowser/Private/SCollectionView.cpp
@@ -303,11 +303,6 @@ void SCollectionView::LoadSettings(const FString& IniFilename, const FString& In
 	}
 }
 
-void SCollectionView::Tick( const FGeometry& AllottedGeometry, const double InCurrentTime, const float InDeltaTime )
-{
-	SCompoundWidget::Tick(AllottedGeometry, InCurrentTime, InDeltaTime);
-}
-
 FReply SCollectionView::OnKeyDown( const FGeometry& MyGeometry, const FKeyEvent& InKeyEvent )
 {
 	if( Commands->ProcessCommandBindings( InKeyEvent ) )
diff --git a/Engine/Source/Editor/ContentBrowser/Private/SCollectionView.h b/Engine/Source/Editor/ContentBrowser/Private/SCollectionView.h
--- a/Engine/Source/Editor/ContentBrowser/Private/SCollectionView.h
+++ b/Engine/Source/Editor/ContentBrowser/Private/SCollectionView.h
@@ -55,8 +55,6 @@ public:
 	/** Loads any settings to config that should be persistent between editor sessions */
 	void LoadSettings(const FString& IniFilename, const FString& IniSection, const FString& SettingsString);
 
-	void Tick( const FGeometry& AllottedGeometry, const double InCurrentTime, const float InDeltaTime ) override;
-
 	virtual FReply OnKeyDown( const FGeometry& MyGeometry, const FKeyEvent& InKeyEvent ) override;
 
 private:
diff --git a/Engine/Source/Editor/ContentBrowser/Private/SPathView.cpp b/Engine/Source/Editor/ContentBrowser/Private/SPathView.cpp
--- a/Engine/Source/Editor/ContentBrowser/Private/SPathView.cpp
+++ b/Engine/Source/Editor/ContentBrowser/Private/SPathView.cpp
@@ -28,16 +28,20 @@ SPathView::~SPathView()
 
 void SPathView::Construct( const FArguments& InArgs )
 {
-	bNeedsRepopulate = true;
+	RegisterActiveTick(0.f, FWidgetActiveTickDelegate::CreateSP(this, &SPathView::TriggerRepopulate));
 
 	OnPathSelected = InArgs._OnPathSelected;
 	bAllowContextMenu = InArgs._AllowContextMenu;
 	OnGetFolderContextMenu = InArgs._OnGetFolderContextMenu;
 	OnGetPathContextMenuExtender = InArgs._OnGetPathContextMenuExtender;
-	bPendingFocusNextFrame = InArgs._FocusSearchBoxWhenOpened;
 	bAllowClassesFolder = InArgs._AllowClassesFolder;
 	PreventTreeItemChangedDelegateCount = 0;
 
+	if ( InArgs._FocusSearchBoxWhenOpened )
+	{
+		RegisterActiveTick( 0.f, FWidgetActiveTickDelegate::CreateSP( this, &SPathView::SetFocusPostConstruct ) );
+	}
+
 	// Listen for when view settings are changed
 	UContentBrowserSettings::OnSettingChanged().AddSP(this, &SPathView::HandleSettingChanged);
 
@@ -616,22 +620,19 @@ void SPathView::LoadSettings(const FString& IniFilename, const FString& IniSecti
 	}
 }
 
-void SPathView::Tick( const FGeometry& AllottedGeometry, const double InCurrentTime, const float InDeltaTime )
+EActiveTickReturnType SPathView::SetFocusPostConstruct( double InCurrentTime, float InDeltaTime )
 {
-	if ( bPendingFocusNextFrame )
-	{
-		FWidgetPath WidgetToFocusPath;
-		FSlateApplication::Get().GeneratePathToWidgetUnchecked( SearchBoxPtr.ToSharedRef(), WidgetToFocusPath );
-		FSlateApplication::Get().SetKeyboardFocus( WidgetToFocusPath, EFocusCause::SetDirectly );
-		bPendingFocusNextFrame = false;
-	}
+	FWidgetPath WidgetToFocusPath;
+	FSlateApplication::Get().GeneratePathToWidgetUnchecked( SearchBoxPtr.ToSharedRef(), WidgetToFocusPath );
+	FSlateApplication::Get().SetKeyboardFocus( WidgetToFocusPath, EFocusCause::SetDirectly );
 
-	SCompoundWidget::Tick(AllottedGeometry, InCurrentTime, InDeltaTime);
+	return EActiveTickReturnType::StopTicking;
+}
 
-	if( bNeedsRepopulate )
-	{
-		Populate();
-	}
+EActiveTickReturnType SPathView::TriggerRepopulate(double InCurrentTime, float InDeltaTime)
+{
+	Populate();
+	return EActiveTickReturnType::StopTicking;
 }
 
 TSharedPtr<SWidget> SPathView::MakePathViewContextMenu()
@@ -973,8 +974,6 @@ void SPathView::Populate()
 	}
 
 	TreeViewPtr->RequestTreeRefresh();
-
-	bNeedsRepopulate = false;
 }
 
 void SPathView::PopulateFolderSearchStrings( const FString& FolderName, OUT TArray< FString >& OutSearchStrings ) const
@@ -1337,7 +1336,7 @@ void SPathView::OnAssetRegistrySearchCompleted()
 void SPathView::OnContentPathMountedOrDismounted( const FString& AssetPath, const FString& FilesystemPath )
 {
 	// A new content path has appeared, so we should refresh out root set of paths
-	bNeedsRepopulate = true;
+	RegisterActiveTick(0.f, FWidgetActiveTickDelegate::CreateSP(this, &SPathView::TriggerRepopulate));
 }
 
 void SPathView::HandleSettingChanged(FName PropertyName)
diff --git a/Engine/Source/Editor/ContentBrowser/Private/SPathView.h b/Engine/Source/Editor/ContentBrowser/Private/SPathView.h
--- a/Engine/Source/Editor/ContentBrowser/Private/SPathView.h
+++ b/Engine/Source/Editor/ContentBrowser/Private/SPathView.h
@@ -109,14 +109,10 @@ public:
 	/** Loads any settings to config that should be persistent between editor sessions */
 	void LoadSettings(const FString& IniFilename, const FString& IniSection, const FString& SettingsString);
 
-	// SWidget implementation
-	virtual void Tick( const FGeometry& AllottedGeometry, const double InCurrentTime, const float InDeltaTime ) override;
-
 	/** Populates the tree with all folders that are not filtered out */
 	void Populate();
 
 private:
-
 	/** Gets the content for a context menu */
 	TSharedPtr<SWidget> MakePathViewContextMenu();
 
@@ -225,8 +221,13 @@ private:
 	/** Delegate called when an editor setting is changed */
 	void HandleSettingChanged(FName PropertyName);
 
-private:
+	/** One-off active tick to focus the widget post-construct */
+	EActiveTickReturnType SetFocusPostConstruct(double InCurrentTime, float InDeltaTime);
+
+	/** One-off active tick to repopulate the path view */
+	EActiveTickReturnType TriggerRepopulate(double InCurrentTime, float InDeltaTime);
 
+private:
 	/** A helper class to manage PreventTreeItemChangedDelegateCount by incrementing it when constructed (on the stack) and decrementing when destroyed */
 	class FScopedPreventTreeItemChangedDelegate
 	{
@@ -286,15 +287,9 @@ private:
 	/** If > 0, the selection or expansion changed delegate will not be called. Used to update the tree from an external source or in certain bulk operations. */
 	int32 PreventTreeItemChangedDelegateCount;
 
-	/** Next tick will focus the filter box. This is used to manage initial focus for popup menus. */
-	bool bPendingFocusNextFrame;
-
 	/** If false, the context menu will not open when right clicking an item in the tree */
 	bool bAllowContextMenu;
 
 	/** If false, the classes folder will not be added to the tree automatically */
 	bool bAllowClassesFolder;
-
-	/** True if we should repopulate the next opportunity we get.  We use this to allow a lazy refresh, for performance reasons. */
-	bool bNeedsRepopulate;
 };
diff --git a/Engine/Source/Editor/ContentBrowser/Private/SourcesViewWidgets.cpp b/Engine/Source/Editor/ContentBrowser/Private/SourcesViewWidgets.cpp
--- a/Engine/Source/Editor/ContentBrowser/Private/SourcesViewWidgets.cpp
+++ b/Engine/Source/Editor/ContentBrowser/Private/SourcesViewWidgets.cpp
@@ -255,8 +255,6 @@ FReply SAssetTreeItem::OnDrop( const FGeometry& MyGeometry, const FDragDropEvent
 
 void SAssetTreeItem::Tick( const FGeometry& AllottedGeometry, const double InCurrentTime, const float InDeltaTime )
 {
-	SCompoundWidget::Tick(AllottedGeometry, InCurrentTime, InDeltaTime);
-
 	LastGeometry = AllottedGeometry;
 }
 
@@ -466,8 +464,6 @@ void SCollectionListItem::Tick( const FGeometry& AllottedGeometry, const double
 {
 	// Cache this widget's geometry so it can pop up warnings over itself
 	CachedGeometry = AllottedGeometry;
-
-	SCompoundWidget::Tick(AllottedGeometry, InCurrentTime, InDeltaTime);
 }
 
 void SCollectionListItem::OnDragEnter( const FGeometry& MyGeometry, const FDragDropEvent& DragDropEvent )
diff --git a/Engine/Source/Editor/ContentBrowser/Public/ContentBrowserDelegates.h b/Engine/Source/Editor/ContentBrowser/Public/ContentBrowserDelegates.h
--- a/Engine/Source/Editor/ContentBrowser/Public/ContentBrowserDelegates.h
+++ b/Engine/Source/Editor/ContentBrowser/Public/ContentBrowserDelegates.h
@@ -84,6 +84,9 @@ DECLARE_DELEGATE_RetVal_TwoParams(bool, FOnVisualizeAssetToolTip, const TSharedP
 /** Called from the Asset Dialog when a non-modal dialog is closed or cancelled */
 DECLARE_DELEGATE(FOnAssetDialogCancelled);
 
+/** Called when an asset item's tooltip is closing */
+DECLARE_DELEGATE( FOnAssetToolTipClosing );
+
 /** Called from the Asset Dialog when assets are chosen in non-modal Open dialogs */
 DECLARE_DELEGATE_OneParam(FOnAssetsChosenForOpen, const TArray<FAssetData>& /*SelectedAssets*/);
 
diff --git a/Engine/Source/Editor/ContentBrowser/Public/IContentBrowserSingleton.h b/Engine/Source/Editor/ContentBrowser/Public/IContentBrowserSingleton.h
--- a/Engine/Source/Editor/ContentBrowser/Public/IContentBrowserSingleton.h
+++ b/Engine/Source/Editor/ContentBrowser/Public/IContentBrowserSingleton.h
@@ -80,9 +80,12 @@ struct FAssetPickerConfig
 	/** Fired when an asset item is constructed and a tooltip is requested. If unbound the item will use the default widget */
 	FOnGetCustomAssetToolTip OnGetCustomAssetToolTip;
 
-	/** Fired when an asset item is about to show it's tool tip */
+	/** Fired when an asset item is about to show its tool tip */
 	FOnVisualizeAssetToolTip OnVisualizeAssetToolTip;
 
+	/** Fired when an asset item's tooltip is closing */
+	FOnAssetToolTipClosing OnAssetToolTipClosing;
+
 	/** If more detailed filtering is required than simply Filter, this delegate will get fired for every asset to determine if it should be culled. */
 	FOnShouldFilterAsset OnShouldFilterAsset;
 
diff --git a/Engine/Source/Editor/DestructibleMeshEditor/Private/SDestructibleMeshEditorViewport.cpp b/Engine/Source/Editor/DestructibleMeshEditor/Private/SDestructibleMeshEditorViewport.cpp
--- a/Engine/Source/Editor/DestructibleMeshEditor/Private/SDestructibleMeshEditorViewport.cpp
+++ b/Engine/Source/Editor/DestructibleMeshEditor/Private/SDestructibleMeshEditorViewport.cpp
@@ -36,6 +36,7 @@ public:
 	// End of FGCObject interface
 
 	// FEditorViewportClient interface
+	virtual void Tick(float DeltaTime) override;
 	virtual void Draw(const FSceneView* View,FPrimitiveDrawInterface* PDI) override;
 	FLinearColor GetBackgroundColor() const { return FLinearColor::Black; }
 	virtual void ProcessClick(class FSceneView& View, class HHitProxy* HitProxy, FKey Key, EInputEvent Event, uint32 HitX, uint32 HitY);
@@ -337,6 +338,16 @@ void FDestructibleMeshEditorViewportClient::ProcessClick( class FSceneView& View
 #endif // WITH_APEX
 }
 
+void FDestructibleMeshEditorViewportClient::Tick(float DeltaTime)
+{
+	FEditorViewportClient::Tick(DeltaTime);
+
+	if (PreviewScene)
+	{
+		PreviewScene->GetWorld()->Tick(LEVELTICK_All, DeltaTime);
+	}
+}
+
 void FDestructibleMeshEditorViewportClient::Draw( const FSceneView* View,FPrimitiveDrawInterface* PDI )
 {
 	FEditorViewportClient::Draw(View, PDI);
@@ -597,13 +608,6 @@ void SDestructibleMeshEditorViewport::RefreshViewport()
 	Viewport->InvalidateDisplay();
 }
 
-void SDestructibleMeshEditorViewport::Tick(const FGeometry& AllottedGeometry, const double InCurrentTime, const float InDeltaTime)
-{
-	SCompoundWidget::Tick(AllottedGeometry, InCurrentTime, InDeltaTime);
-
-	PreviewScene.GetWorld()->Tick(LEVELTICK_All, InDeltaTime);
-}
-
 void SDestructibleMeshEditorViewport::SetPreviewMesh(UDestructibleMesh* InDestructibleMesh)
 {
 	DestructibleMesh = InDestructibleMesh;
diff --git a/Engine/Source/Editor/DestructibleMeshEditor/Private/SDestructibleMeshEditorViewport.h b/Engine/Source/Editor/DestructibleMeshEditor/Private/SDestructibleMeshEditorViewport.h
--- a/Engine/Source/Editor/DestructibleMeshEditor/Private/SDestructibleMeshEditorViewport.h
+++ b/Engine/Source/Editor/DestructibleMeshEditor/Private/SDestructibleMeshEditorViewport.h
@@ -26,7 +26,6 @@ public:
 	virtual void NotifyPostChange( const FPropertyChangedEvent& PropertyChangedEvent, class FEditPropertyChain* PropertyThatChanged) override;
 	// End of FNotifyHook interface
 
-	virtual void Tick(const FGeometry& AllottedGeometry, const double InCurrentTime, const float InDeltaTime) override;
 	void RefreshViewport();
 	
 	/** Component for the preview Destructible mesh. */
diff --git a/Engine/Source/Editor/DetailCustomizations/Private/AnimMontageSegmentDetails.cpp b/Engine/Source/Editor/DetailCustomizations/Private/AnimMontageSegmentDetails.cpp
--- a/Engine/Source/Editor/DetailCustomizations/Private/AnimMontageSegmentDetails.cpp
+++ b/Engine/Source/Editor/DetailCustomizations/Private/AnimMontageSegmentDetails.cpp
@@ -180,6 +180,7 @@ void SAnimationSegmentViewport::Construct(const FArguments& InArgs)
 			.PreviewInstance(this, &SAnimationSegmentViewport::GetPreviewInstance)
 			.DraggableBars(this, &SAnimationSegmentViewport::GetBars)
 			.OnBarDrag(this, &SAnimationSegmentViewport::OnBarDrag)
+			.OnTickPlayback(this, &SAnimationSegmentViewport::OnTickPreview)
 			.bAllowZoom(true)
 		]
 	];
@@ -243,13 +244,25 @@ void SAnimationSegmentViewport::InitSkeleton()
 	TargetSkeleton = Skeleton;
 }
 
-void SAnimationSegmentViewport::Tick(const FGeometry& AllottedGeometry, const double InCurrentTime, const float InDeltaTime)
+void SAnimationSegmentViewport::OnTickPreview( double InCurrentTime, float InDeltaTime )
 {
-	SCompoundWidget::Tick(AllottedGeometry, InCurrentTime, InDeltaTime);
+	// Clamp the sequence playing to the start/end values of the anim segment
+	float Start, End;
+	StartTimePropertyHandle->GetValue( Start );
+	EndTimePropertyHandle->GetValue( End );
+	if ( PreviewComponent->PreviewInstance->CurrentTime > End || PreviewComponent->PreviewInstance->CurrentTime < Start )
+	{
+		PreviewComponent->PreviewInstance->SetPosition( Start, false );
+	}
+
+	LevelViewportClient->Invalidate();
+}
 
+void SAnimationSegmentViewport::Tick( const FGeometry& AllottedGeometry, const double InCurrentTime, const float InDeltaTime )
+{
 	class UDebugSkelMeshComponent* Component = PreviewComponent;
 
-	FString TargetSkeletonName = TargetSkeleton ?TargetSkeleton->GetName() : FName(NAME_None).ToString();
+	FString TargetSkeletonName = TargetSkeleton ? TargetSkeleton->GetName() : FName(NAME_None).ToString();
 
 	if (Component != NULL)
 	{
@@ -258,20 +271,6 @@ void SAnimationSegmentViewport::Tick(const FGeometry& AllottedGeometry, const do
 
 		if (Component->IsPreviewOn())
 		{
-			if(PreviewComponent != NULL && PreviewComponent->PreviewInstance != NULL)
-			{
-				// Clamp the sequence playing to the start/end values of the anim segment
-				float Start, End;
-				StartTimePropertyHandle->GetValue(Start);
-				EndTimePropertyHandle->GetValue(End);
-				if(PreviewComponent->PreviewInstance->CurrentTime > End || PreviewComponent->PreviewInstance->CurrentTime < Start)
-				{
-					PreviewComponent->PreviewInstance->SetPosition(Start, false);
-					// reinvalidate the viewport
-					LevelViewportClient->Invalidate();
-				}
-			}
-
 			Description->SetText(FText::Format( LOCTEXT("Previewing", "Previewing {0}"), FText::FromString(Component->GetPreviewText()) ).ToString());
 		}
 		else if (Component->AnimBlueprintGeneratedClass)
@@ -402,6 +401,7 @@ void SAnimationSegmentScrubPanel::Construct( const SAnimationSegmentScrubPanel::
 				.IsRealtimeStreamingMode(this, &SAnimationSegmentScrubPanel::IsRealtimeStreamingMode)
 				.DraggableBars(InArgs._DraggableBars)
 				.OnBarDrag(InArgs._OnBarDrag)
+				.OnTickPlayback(InArgs._OnTickPlayback)
 			]
 		];
 
@@ -414,7 +414,6 @@ SAnimationSegmentScrubPanel::~SAnimationSegmentScrubPanel()
 
 FReply SAnimationSegmentScrubPanel::OnClick_Forward()
 {
-	
 	UAnimSingleNodeInstance* PreviewInst = GetPreviewInstance();
 	if (PreviewInst)
 	{
diff --git a/Engine/Source/Editor/DetailCustomizations/Private/AnimMontageSegmentDetails.h b/Engine/Source/Editor/DetailCustomizations/Private/AnimMontageSegmentDetails.h
--- a/Engine/Source/Editor/DetailCustomizations/Private/AnimMontageSegmentDetails.h
+++ b/Engine/Source/Editor/DetailCustomizations/Private/AnimMontageSegmentDetails.h
@@ -43,11 +43,14 @@ public:
 
 	void Construct(const FArguments& InArgs);
 
-	virtual void Tick(const FGeometry& AllottedGeometry, const double InCurrentTime, const float InDeltaTime) override;
+	virtual void Tick( const FGeometry& AllottedGeometry, const double InCurrentTime, const float InDeltaTime ) override;
 	void RefreshViewport();
 
 private:
 
+	/** Called to tick the preview during playback */
+	void OnTickPreview( double InCurrentTime, float InDeltaTime );
+
 	void InitSkeleton();
 
 	TSharedPtr<FEditorViewportClient> LevelViewportClient;
@@ -102,13 +105,14 @@ public:
 	{}
 		/** If you'd like to lock to one asset for this scrub control, give this**/
 		SLATE_ARGUMENT(UAnimSequenceBase*, LockedSequence)
-		SLATE_ATTRIBUTE(class UAnimSingleNodeInstance*, PreviewInstance);
+		SLATE_ATTRIBUTE(class UAnimSingleNodeInstance*, PreviewInstance)
 		/** View Input range **/
 		SLATE_ATTRIBUTE( float, ViewInputMin )
 		SLATE_ATTRIBUTE( float, ViewInputMax )
 		SLATE_ARGUMENT( bool, bAllowZoom )
 		SLATE_ATTRIBUTE( TArray<float>, DraggableBars )
-		SLATE_EVENT( FOnScrubBarDrag, OnBarDrag)
+		SLATE_EVENT( FOnScrubBarDrag, OnBarDrag )
+		SLATE_EVENT( FOnTickPlayback, OnTickPlayback )
 	SLATE_END_ARGS()
 	
 	void Construct( const FArguments& InArgs );
