Removed experimental flag on blueprint merging

diff --git a/Engine/Source/Developer/Merge/Private/Merge.cpp b/Engine/Source/Developer/Merge/Private/Merge.cpp
--- a/Engine/Source/Developer/Merge/Private/Merge.cpp
+++ b/Engine/Source/Developer/Merge/Private/Merge.cpp
@@ -196,11 +196,10 @@ TSharedRef<SDockTab> FMerge::GenerateMergeWidget(const UBlueprint* BaseBlueprint
 
 bool FMerge::PendingMerge(const UBlueprint& BlueprintObj) const
 {
-	bool bIsMergeEnabled = GetDefault<UEditorExperimentalSettings>()->bEnableBlueprintMergeTool;
 	ISourceControlProvider& SourceControlProvider = ISourceControlModule::Get().GetProvider();
 
 	bool bPendingMerge = false;
-	if( bIsMergeEnabled && SourceControlProvider.IsEnabled() )
+	if( SourceControlProvider.IsEnabled() )
 	{
 		FSourceControlStatePtr SourceControlState = SourceControlProvider.GetState(BlueprintObj.GetOutermost(), EStateCacheUsage::Use);
 		bPendingMerge = SourceControlState.IsValid() && SourceControlState->IsConflicted();
diff --git a/Engine/Source/Editor/ContentBrowser/Private/AssetContextMenu.cpp b/Engine/Source/Editor/ContentBrowser/Private/AssetContextMenu.cpp
--- a/Engine/Source/Editor/ContentBrowser/Private/AssetContextMenu.cpp
+++ b/Engine/Source/Editor/ContentBrowser/Private/AssetContextMenu.cpp
@@ -1649,29 +1649,26 @@ bool FAssetContextMenu::CanExecuteSCCMerge() const
 {
 	FAssetToolsModule& AssetToolsModule = FModuleManager::GetModuleChecked<FAssetToolsModule>("AssetTools");
 
-	bool bCanExecuteMerge = GetDefault<UEditorExperimentalSettings>()->bEnableBlueprintMergeTool;
-	if( bCanExecuteMerge )
+	bool bCanExecuteMerge = bCanExecuteSCCMerge;
+	for (int32 AssetIdx = 0; AssetIdx < SelectedAssets.Num() && bCanExecuteMerge; AssetIdx++)
 	{
-		bCanExecuteMerge = bCanExecuteSCCMerge;
-		for (int32 AssetIdx = 0; AssetIdx < SelectedAssets.Num() && bCanExecuteMerge; AssetIdx++)
+		// Get the actual asset (will load it)
+		const FAssetData& AssetData = SelectedAssets[AssetIdx];
+		UObject* CurrentObject = AssetData.GetAsset();
+		if (CurrentObject)
 		{
-			// Get the actual asset (will load it)
-			const FAssetData& AssetData = SelectedAssets[AssetIdx];
-			UObject* CurrentObject = AssetData.GetAsset();
-			if (CurrentObject)
-			{
-				auto AssetTypeActions = AssetToolsModule.Get().GetAssetTypeActionsForClass(CurrentObject->GetClass()).Pin();
-				if (AssetTypeActions.IsValid())
-				{
-					bCanExecuteMerge = AssetTypeActions->CanMerge();
-				}
-			}
-			else
+			auto AssetTypeActions = AssetToolsModule.Get().GetAssetTypeActionsForClass(CurrentObject->GetClass()).Pin();
+			if (AssetTypeActions.IsValid())
 			{
-				bCanExecuteMerge = false;
+				bCanExecuteMerge = AssetTypeActions->CanMerge();
 			}
 		}
+		else
+		{
+			bCanExecuteMerge = false;
+		}
 	}
+
 	return bCanExecuteMerge;
 }
 
diff --git a/Engine/Source/Editor/UnrealEd/Classes/Settings/EditorExperimentalSettings.h b/Engine/Source/Editor/UnrealEd/Classes/Settings/EditorExperimentalSettings.h
--- a/Engine/Source/Editor/UnrealEd/Classes/Settings/EditorExperimentalSettings.h
+++ b/Engine/Source/Editor/UnrealEd/Classes/Settings/EditorExperimentalSettings.h
@@ -77,10 +77,6 @@ public:
 	UPROPERTY(EditAnywhere, config, Category = AI, meta = (DisplayName = "Environment Querying System"))
 	bool bEQSEditor;
 
-	/** Enables the Blueprint merge tool */
-	UPROPERTY(EditAnywhere, config, Category = Blueprints, meta = (DisplayName = "Enable Blueprint Merge Tool"))
-	bool bEnableBlueprintMergeTool;
-
 	/**
 	 * Returns an event delegate that is executed when a setting has changed.
 	 *
