Sync option in the Content Browser context menu refreshes icons after operation completes.

diff --git a/Engine/Source/Editor/ContentBrowser/Private/AssetContextMenu.cpp b/Engine/Source/Editor/ContentBrowser/Private/AssetContextMenu.cpp
--- a/Engine/Source/Editor/ContentBrowser/Private/AssetContextMenu.cpp
+++ b/Engine/Source/Editor/ContentBrowser/Private/AssetContextMenu.cpp
@@ -1567,6 +1567,7 @@ void FAssetContextMenu::ExecuteSCCSync()
 		{
 			PackageTools::LoadPackage(*PackageIter);
 		}
+		ExecuteSCCRefresh();
 	}
 }
 
