Fixed code relying on SLATE_TEXT_ATTRIBUTE for tooltips.

UETOOL-213 - Minimize Slate FString -> FText conversion (remove SLATE_TEXT_ATTRIBUTE)

This fixes any editor/engine specific code that was passing text to Slate as FString rather than FText.

diff --git a/Engine/Plugins/Experimental/CodeView/Source/CodeView/Private/SCodeView.cpp b/Engine/Plugins/Experimental/CodeView/Source/CodeView/Private/SCodeView.cpp
--- a/Engine/Plugins/Experimental/CodeView/Source/CodeView/Private/SCodeView.cpp
+++ b/Engine/Plugins/Experimental/CodeView/Source/CodeView/Private/SCodeView.cpp
@@ -520,7 +520,7 @@ namespace CodeView
 								.HighlightText( this, &SCodeView::GetFilterText )
 
 								// Use the module name as the tool-tip
-								.ToolTipText( Item->ModuleName )
+								.ToolTipText( FText::FromString(Item->ModuleName) )
 						]
 				;
 
diff --git a/Engine/Source/Developer/ExternalImagePicker/Private/SExternalImagePicker.cpp b/Engine/Source/Developer/ExternalImagePicker/Private/SExternalImagePicker.cpp
--- a/Engine/Source/Developer/ExternalImagePicker/Private/SExternalImagePicker.cpp
+++ b/Engine/Source/Developer/ExternalImagePicker/Private/SExternalImagePicker.cpp
@@ -64,7 +64,7 @@ void SExternalImagePicker::Construct(const FArguments& InArgs)
 			[
 				SNew(SButton)
 				.ButtonStyle( FEditorStyle::Get(), "HoverHintOnly" )
-				.ToolTipText( LOCTEXT( "FileButtonToolTipText", "Choose a file from this computer").ToString() )
+				.ToolTipText( LOCTEXT( "FileButtonToolTipText", "Choose a file from this computer") )
 				.OnClicked( FOnClicked::CreateSP(this, &SExternalImagePicker::OnPickFile) )
 				.ContentPadding( 2.0f )
 				.ForegroundColor( FSlateColor::UseForeground() )
diff --git a/Engine/Source/Developer/ProjectLauncher/Private/Widgets/SProjectLauncherMessageListRow.h b/Engine/Source/Developer/ProjectLauncher/Private/Widgets/SProjectLauncherMessageListRow.h
--- a/Engine/Source/Developer/ProjectLauncher/Private/Widgets/SProjectLauncherMessageListRow.h
+++ b/Engine/Source/Developer/ProjectLauncher/Private/Widgets/SProjectLauncherMessageListRow.h
@@ -7,10 +7,10 @@
 
 struct FProjectLauncherMessage
 {
-	FString Message;
+	FText Message;
 	ELogVerbosity::Type Verbosity;
 
-	FProjectLauncherMessage(const FString& NewMessage, ELogVerbosity::Type InVerbosity)
+	FProjectLauncherMessage(const FText& NewMessage, ELogVerbosity::Type InVerbosity)
 		: Message(NewMessage)
 		, Verbosity(InVerbosity)
 	{ }
@@ -63,7 +63,7 @@ public:
 				[
 					SNew(STextBlock)
 						.ColorAndOpacity(HandleGetTextColor())
-						.Text(FText::FromString(Message->Message))
+						.Text(Message->Message)
 				];
 		}
 
diff --git a/Engine/Source/Developer/ProjectLauncher/Private/Widgets/SProjectLauncherProgress.h b/Engine/Source/Developer/ProjectLauncher/Private/Widgets/SProjectLauncherProgress.h
--- a/Engine/Source/Developer/ProjectLauncher/Private/Widgets/SProjectLauncherProgress.h
+++ b/Engine/Source/Developer/ProjectLauncher/Private/Widgets/SProjectLauncherProgress.h
@@ -261,7 +261,7 @@ private:
 		{
 			Verbosity = ELogVerbosity::Warning;
 		}
-		TSharedPtr<FProjectLauncherMessage> Message = MakeShareable(new FProjectLauncherMessage(InMessage, Verbosity));
+		TSharedPtr<FProjectLauncherMessage> Message = MakeShareable(new FProjectLauncherMessage(FText::FromString(InMessage), Verbosity));
 		PendingMessages.Add(Message);
 	}
 
@@ -488,9 +488,8 @@ private:
 
 			for( int32 Index = 0; Index < SelectedItems.Num(); ++Index )
 			{
-				SelectedText += FString::Printf(TEXT("%s"),
-					*SelectedItems[Index]->Message
-					) + LINE_TERMINATOR;
+				SelectedText += SelectedItems[Index]->Message.ToString();
+				SelectedText += LINE_TERMINATOR;
 			}
 
 			FPlatformMisc::ClipboardCopy( *SelectedText );
@@ -537,8 +536,7 @@ private:
 					{
 						for( int32 Index = 0; Index < MessageList.Num(); ++Index )
 						{
-							FString LogEntry = FString::Printf(TEXT("%s"),
-								*MessageList[Index]->Message) + LINE_TERMINATOR;
+							FString LogEntry = MessageList[Index]->Message.ToString() + LINE_TERMINATOR;
 
 							LogFile->Serialize(TCHAR_TO_ANSI(*LogEntry), LogEntry.Len());
 						}
diff --git a/Engine/Source/Editor/BehaviorTreeEditor/Private/SBehaviorTreeBlackboardView.cpp b/Engine/Source/Editor/BehaviorTreeEditor/Private/SBehaviorTreeBlackboardView.cpp
--- a/Engine/Source/Editor/BehaviorTreeEditor/Private/SBehaviorTreeBlackboardView.cpp
+++ b/Engine/Source/Editor/BehaviorTreeEditor/Private/SBehaviorTreeBlackboardView.cpp
@@ -171,9 +171,9 @@ private:
 		return DisplayWidget.ToSharedRef();
 	}
 
-	virtual FString GetItemTooltip() const override
+	virtual FText GetItemTooltip() const override
 	{
-		return ActionPtr.Pin()->TooltipDescription;
+		return FText::FromString(ActionPtr.Pin()->TooltipDescription);
 	}
 
 	virtual void OnNameTextCommitted(const FText& NewText, ETextCommit::Type InTextCommit) override
diff --git a/Engine/Source/Editor/BehaviorTreeEditor/Private/SGraphNode_BehaviorTree.cpp b/Engine/Source/Editor/BehaviorTreeEditor/Private/SGraphNode_BehaviorTree.cpp
--- a/Engine/Source/Editor/BehaviorTreeEditor/Private/SGraphNode_BehaviorTree.cpp
+++ b/Engine/Source/Editor/BehaviorTreeEditor/Private/SGraphNode_BehaviorTree.cpp
@@ -884,15 +884,20 @@ FText SGraphNode_BehaviorTree::GetDescription() const
 	return FText::FromString(StateNode->GetDescription());
 }
 
-FString SGraphNode_BehaviorTree::GetPinTooltip(UEdGraphPin* GraphPinObj) const
+FText SGraphNode_BehaviorTree::GetPinTooltip(UEdGraphPin* GraphPinObj) const
 {
-	FString HoverText;
+	FText HoverText = FText::GetEmpty();
 
 	check(GraphPinObj != nullptr);
 	UEdGraphNode* GraphNode = GraphPinObj->GetOwningNode();
 	if (GraphNode != nullptr)
 	{
-		GraphNode->GetPinHoverText(*GraphPinObj, /*out*/ HoverText);
+		FString HoverStr;
+		GraphNode->GetPinHoverText(*GraphPinObj, /*out*/HoverStr);
+		if (!HoverStr.IsEmpty())
+		{
+			HoverText = FText::FromString(HoverStr);
+		}
 	}
 
 	return HoverText;
diff --git a/Engine/Source/Editor/BehaviorTreeEditor/Private/SGraphNode_BehaviorTree.h b/Engine/Source/Editor/BehaviorTreeEditor/Private/SGraphNode_BehaviorTree.h
--- a/Engine/Source/Editor/BehaviorTreeEditor/Private/SGraphNode_BehaviorTree.h
+++ b/Engine/Source/Editor/BehaviorTreeEditor/Private/SGraphNode_BehaviorTree.h
@@ -125,5 +125,5 @@ protected:
 	/** Handle hover state changing for the index widget - we use this to highlight sibling nodes */
 	void OnIndexHoverStateChanged(bool bHovered);
 
-	FString GetPinTooltip(UEdGraphPin* GraphPinObj) const;
+	FText GetPinTooltip(UEdGraphPin* GraphPinObj) const;
 };
diff --git a/Engine/Source/Editor/Blutility/Private/BlutilityDetailsPanel.cpp b/Engine/Source/Editor/Blutility/Private/BlutilityDetailsPanel.cpp
--- a/Engine/Source/Editor/Blutility/Private/BlutilityDetailsPanel.cpp
+++ b/Engine/Source/Editor/Blutility/Private/BlutilityDetailsPanel.cpp
@@ -102,10 +102,10 @@ void FEditorUtilityInstanceDetails::CustomizeDetails(IDetailLayoutBuilder& Detai
 				const FText ButtonCaption = FText::FromString(FName::NameToDisplayString(*Function->GetName(), false));
 
 				//@TODO: Expose the code in UK2Node_CallFunction::GetUserFacingFunctionName / etc...
-				FString Tooltip = Function->GetToolTipText().ToString();
+				FText Tooltip = Function->GetToolTipText();
 				if (Tooltip.IsEmpty())
 				{
-					Tooltip = Function->GetName();
+					Tooltip = FText::FromString(Function->GetName());
 				}
 
 				TWeakObjectPtr<UFunction> WeakFunctionPtr(Function);
diff --git a/Engine/Source/Editor/ContentBrowser/Private/HistoryManager.cpp b/Engine/Source/Editor/ContentBrowser/Private/HistoryManager.cpp
--- a/Engine/Source/Editor/ContentBrowser/Private/HistoryManager.cpp
+++ b/Engine/Source/Editor/ContentBrowser/Private/HistoryManager.cpp
@@ -113,28 +113,22 @@ bool FHistoryManager::CanGoBack() const
 	return HistoryData.Num() > 0 && CurrentHistoryIndex > 0;
 }
 
-FString FHistoryManager::GetBackDesc() const
+FText FHistoryManager::GetBackDesc() const
 {
 	if ( CanGoBack() )
 	{
 		return HistoryData[CurrentHistoryIndex - 1].HistoryDesc;
 	}
-	else
-	{
-		return FString();
-	}
+	return FText::GetEmpty();
 }
 
-FString FHistoryManager::GetForwardDesc() const
+FText FHistoryManager::GetForwardDesc() const
 {
 	if ( CanGoForward() )
 	{
 		return HistoryData[CurrentHistoryIndex + 1].HistoryDesc;
 	}
-	else
-	{
-		return FString();
-	}
+	return FText::GetEmpty();
 }
 
 void FHistoryManager::GetAvailableHistoryMenuItems(bool bGetPrior, FMenuBuilder& MenuBuilder)
@@ -165,7 +159,7 @@ void FHistoryManager::GetAvailableHistoryMenuItems(bool bGetPrior, FMenuBuilder&
 				for (int32 HistoryIdx = StartIndex; HistoryIdx <= EndIndex; ++HistoryIdx)
 				{
 					MenuBuilder.AddMenuEntry(
-						FText::FromString( HistoryData[HistoryIdx].HistoryDesc ),
+						HistoryData[HistoryIdx].HistoryDesc,
 						FText(),
 						FSlateIcon(),
 						FUIAction(
diff --git a/Engine/Source/Editor/ContentBrowser/Private/HistoryManager.h b/Engine/Source/Editor/ContentBrowser/Private/HistoryManager.h
--- a/Engine/Source/Editor/ContentBrowser/Private/HistoryManager.h
+++ b/Engine/Source/Editor/ContentBrowser/Private/HistoryManager.h
@@ -6,7 +6,7 @@
 struct FHistoryData
 {
 	/** History Description */
-	FString HistoryDesc;
+	FText HistoryDesc;
 
 	/** The base set of filters on the asset view which includes selected paths and collections */
 	FSourcesData SourcesData;
@@ -53,10 +53,10 @@ public:
 	bool CanGoBack() const;
 
 	/** Gets the description of the previous history entry */
-	FString GetBackDesc() const;
+	FText GetBackDesc() const;
 
 	/** Gets the description of the next history entry */
-	FString GetForwardDesc() const;
+	FText GetForwardDesc() const;
 
 	/**
 	 * Populates a list of menu items that can be added to a context menu
diff --git a/Engine/Source/Editor/ContentBrowser/Private/SContentBrowser.cpp b/Engine/Source/Editor/ContentBrowser/Private/SContentBrowser.cpp
--- a/Engine/Source/Editor/ContentBrowser/Private/SContentBrowser.cpp
+++ b/Engine/Source/Editor/ContentBrowser/Private/SContentBrowser.cpp
@@ -1084,7 +1084,7 @@ void SContentBrowser::OnUpdateHistoryData(FHistoryData& HistoryData) const
 	const FSourcesData& SourcesData = AssetViewPtr->GetSourcesData();
 	const TArray<FAssetData>& SelectedAssets = AssetViewPtr->GetSelectedAssets();
 
-	const FString NewSource = SourcesData.PackagePaths.Num() > 0 ? SourcesData.PackagePaths[0].ToString() : (SourcesData.Collections.Num() > 0 ? SourcesData.Collections[0].Name.ToString() : LOCTEXT("AllAssets", "All Assets").ToString());
+	const FText NewSource = SourcesData.PackagePaths.Num() > 0 ? FText::FromName(SourcesData.PackagePaths[0]) : (SourcesData.Collections.Num() > 0 ? FText::FromName(SourcesData.Collections[0].Name) : LOCTEXT("AllAssets", "All Assets"));
 
 	HistoryData.HistoryDesc = NewSource;
 	HistoryData.SourcesData = SourcesData;
@@ -1331,19 +1331,17 @@ TSharedRef<SWidget> SContentBrowser::MakeCreateAssetContextMenu()
 		];
 }
 
-FString SContentBrowser::GetNewAssetToolTipText() const
+FText SContentBrowser::GetNewAssetToolTipText() const
 {
 	const FSourcesData& SourcesData = AssetViewPtr->GetSourcesData();
 
 	// At least one source is selected
 	if (SourcesData.PackagePaths.Num() > 0)
 	{
-		return FString::Printf( *LOCTEXT("CreateAssetToolTip", "Create an asset in %s.").ToString(), *SourcesData.PackagePaths[0].ToString() );
-	}
-	else
-	{
-		return FString();
+		return FText::Format( LOCTEXT("CreateAssetToolTipFmt", "Create an asset in {0}."), FText::FromName(SourcesData.PackagePaths[0]) );
 	}
+	
+	return FText::GetEmpty();
 }
 
 TSharedRef<SWidget> SContentBrowser::MakeAddFilterMenu()
@@ -1660,28 +1658,22 @@ bool SContentBrowser::CanExecuteDirectoryUp() const
 	return (SelectedPaths.Num() == 1 && !ContentBrowserUtils::IsAssetRootDir(SelectedPaths[0]));
 }
 
-FString SContentBrowser::GetHistoryBackTooltip() const
+FText SContentBrowser::GetHistoryBackTooltip() const
 {
 	if ( HistoryManager.CanGoBack() )
 	{
-		return FString::Printf( *LOCTEXT("HistoryBackTooltip", "Back to %s").ToString(), *HistoryManager.GetBackDesc() );
-	}
-	else
-	{
-		return FString();
+		return FText::Format( LOCTEXT("HistoryBackTooltipFmt", "Back to {0}"), HistoryManager.GetBackDesc() );
 	}
+	return FText::GetEmpty();
 }
 
-FString SContentBrowser::GetHistoryForwardTooltip() const
+FText SContentBrowser::GetHistoryForwardTooltip() const
 {
 	if ( HistoryManager.CanGoForward() )
 	{
-		return FString::Printf( *LOCTEXT("HistoryForwardTooltip", "Forward to %s").ToString(), *HistoryManager.GetForwardDesc() );
-	}
-	else
-	{
-		return FString();
+		return FText::Format( LOCTEXT("HistoryForwardTooltipFmt", "Forward to {0}"), HistoryManager.GetForwardDesc() );
 	}
+	return FText::GetEmpty();
 }
 
 FText SContentBrowser::GetDirectoryUpTooltip() const
diff --git a/Engine/Source/Editor/ContentBrowser/Private/SContentBrowser.h b/Engine/Source/Editor/ContentBrowser/Private/SContentBrowser.h
--- a/Engine/Source/Editor/ContentBrowser/Private/SContentBrowser.h
+++ b/Engine/Source/Editor/ContentBrowser/Private/SContentBrowser.h
@@ -153,7 +153,7 @@ private:
 	TSharedRef<SWidget> MakeCreateAssetContextMenu();
 
 	/** Gets the tool tip for the new asset button */
-	FString GetNewAssetToolTipText() const;
+	FText GetNewAssetToolTipText() const;
 
 	/** Makes the filters menu */
 	TSharedRef<SWidget> MakeAddFilterMenu();
@@ -237,10 +237,10 @@ private:
 	bool CanExecuteDirectoryUp() const;
 
 	/** Gets the tool tip text for the history back button */
-	FString GetHistoryBackTooltip() const;
+	FText GetHistoryBackTooltip() const;
 
 	/** Gets the tool tip text for the history forward button */
-	FString GetHistoryForwardTooltip() const;
+	FText GetHistoryForwardTooltip() const;
 
 	/** Gets the tool tip text for the directory up button */
 	FText GetDirectoryUpTooltip() const;
diff --git a/Engine/Source/Editor/DataTableEditor/Private/DataTableEditor.cpp b/Engine/Source/Editor/DataTableEditor/Private/DataTableEditor.cpp
--- a/Engine/Source/Editor/DataTableEditor/Private/DataTableEditor.cpp
+++ b/Engine/Source/Editor/DataTableEditor/Private/DataTableEditor.cpp
@@ -200,8 +200,8 @@ TSharedPtr<SUniformGridPanel> FDataTableEditor::CreateGridPanel()
 							SNew(STextBlock)
 							.Text(FText::FromString(Row[Column]))
 							.ToolTipText(bIsHeader 
-							?	(FString::Printf(TEXT("Column '%s"), *ColumnTitles[Column])) 
-							:	(FString::Printf(TEXT("%s: %s"), *ColumnTitles[Column], *Row[Column])))
+							?	(FText::Format(LOCTEXT("ColumnHeaderNameFmt", "Column '{0}'"), FText::FromString(ColumnTitles[Column]))) 
+							:	(FText::Format(LOCTEXT("ColumnRowNameFmt", "{0}: {1}"), FText::FromString(ColumnTitles[Column]), FText::FromString(Row[Column]))))
 						]
 					];
 				}
diff --git a/Engine/Source/Editor/Documentation/Private/UDNParser.cpp b/Engine/Source/Editor/Documentation/Private/UDNParser.cpp
--- a/Engine/Source/Editor/Documentation/Private/UDNParser.cpp
+++ b/Engine/Source/Editor/Documentation/Private/UDNParser.cpp
@@ -967,7 +967,7 @@ TSharedRef< SWidget > FUDNParser::GenerateExcerptContent( const FString& Link, F
 					AppendExcerpt(Box,
 						SNew( SImage )
 						.Image(DynamicBrush.Get())
-						.ToolTipText(Line.AdditionalContent[0])
+						.ToolTipText(FText::FromString(Line.AdditionalContent[0]))
 						);
 
 					DynamicBrushesUsed.AddUnique(DynamicBrush);
@@ -992,7 +992,7 @@ TSharedRef< SWidget > FUDNParser::GenerateExcerptContent( const FString& Link, F
 						[
 							SNew( SImage )
 							.Image(DynamicBrush.Get())
-							.ToolTipText(Line.AdditionalContent[0])
+							.ToolTipText(FText::FromString(Line.AdditionalContent[0]))
 						]
 					);
 
diff --git a/Engine/Source/Editor/EditorWidgets/Private/SObjectNameEditableTextBox.cpp b/Engine/Source/Editor/EditorWidgets/Private/SObjectNameEditableTextBox.cpp
--- a/Engine/Source/Editor/EditorWidgets/Private/SObjectNameEditableTextBox.cpp
+++ b/Engine/Source/Editor/EditorWidgets/Private/SObjectNameEditableTextBox.cpp
@@ -144,38 +144,38 @@ FText SObjectNameEditableTextBox::GetNameText() const
 	return FText::FromString(Result);
 }
 
-FString SObjectNameEditableTextBox::GetNameTooltipText() const
+FText SObjectNameEditableTextBox::GetNameTooltipText() const
 {
-	FString Result;
+	FText Result = FText::GetEmpty();
 
 	if (Objects.Num() == 0)
 	{
-		Result = LOCTEXT("EditableActorLabel_NoObjectsTooltip", "Nothing selected").ToString();
+		Result = LOCTEXT("EditableActorLabel_NoObjectsTooltip", "Nothing selected");
 	}
 	else if (Objects.Num() == 1 && Objects[0].IsValid())
 	{
 		if (CanEditNameText())
 		{
-			Result = FString::Printf(*LOCTEXT("EditableActorLabel_ActorTooltip", "Rename the selected %s").ToString(), *Objects[0].Get()->GetClass()->GetName());
+			Result = FText::Format(LOCTEXT("EditableActorLabel_ActorTooltipFmt", "Rename the selected {0}"), FText::FromString(Objects[0].Get()->GetClass()->GetName()));
 		}
 		else if (Objects[0].Get()->IsA(AActor::StaticClass()))
 		{
-			Result = LOCTEXT("EditableActorLabel_NoEditActorTooltip", "Can't rename selected actor (its label isn't editable)").ToString();
+			Result = LOCTEXT("EditableActorLabel_NoEditActorTooltip", "Can't rename selected actor (its label isn't editable)");
 		}
 		else
 		{
-			Result = LOCTEXT("EditableActorLabel_NoEditObjectTooltip", "Can't rename selected object (only actors can have editable labels)").ToString();
+			Result = LOCTEXT("EditableActorLabel_NoEditObjectTooltip", "Can't rename selected object (only actors can have editable labels)");
 		}
 	}
 	else if (Objects.Num() > 1)
 	{
 		if (CanEditNameText())
 		{
-			Result = LOCTEXT("EditableActorLabel_MultiActorTooltip", "Rename multiple selected actors at once").ToString();
+			Result = LOCTEXT("EditableActorLabel_MultiActorTooltip", "Rename multiple selected actors at once");
 		}
 		else
 		{
-			Result = LOCTEXT("EditableActorLabel_NoEditMultiObjectTooltip", "Can't rename selected objects (one or more aren't actors with editable labels)").ToString();
+			Result = LOCTEXT("EditableActorLabel_NoEditMultiObjectTooltip", "Can't rename selected objects (one or more aren't actors with editable labels)");
 		}
 	}
 
diff --git a/Engine/Source/Editor/EditorWidgets/Private/SObjectNameEditableTextBox.h b/Engine/Source/Editor/EditorWidgets/Private/SObjectNameEditableTextBox.h
--- a/Engine/Source/Editor/EditorWidgets/Private/SObjectNameEditableTextBox.h
+++ b/Engine/Source/Editor/EditorWidgets/Private/SObjectNameEditableTextBox.h
@@ -36,7 +36,7 @@ private:
 	EVisibility GetNameVisibility() const;
 
 	/** Getter for the ToolTipText attribute of the editable text inside this widget */
-	FString GetNameTooltipText() const;
+	FText GetNameTooltipText() const;
 
 	/** Getter for the HintText attribute of the editable text inside this widget */
 	FText GetNameHintText() const;
diff --git a/Engine/Source/Editor/GraphEditor/Private/AnimationNodes/SGraphNodeSequencePlayer.cpp b/Engine/Source/Editor/GraphEditor/Private/AnimationNodes/SGraphNodeSequencePlayer.cpp
--- a/Engine/Source/Editor/GraphEditor/Private/AnimationNodes/SGraphNodeSequencePlayer.cpp
+++ b/Engine/Source/Editor/GraphEditor/Private/AnimationNodes/SGraphNodeSequencePlayer.cpp
@@ -24,7 +24,7 @@ void SGraphNodeSequencePlayer::GetNodeInfoPopups(FNodeInfoContext* Context, TArr
 {
 }
 
-FString SGraphNodeSequencePlayer::GetPositionTooltip() const
+FText SGraphNodeSequencePlayer::GetPositionTooltip() const
 {
 	float Position;
 	float Length;
@@ -49,10 +49,10 @@ FString SGraphNodeSequencePlayer::GetPositionTooltip() const
 
 		const FString FramesStr = FString::Printf(TEXT("Frame %d"), CurrentFrame);
 
-		return FString::Printf(TEXT("%s (%s%s%s)"), *FramesStr, *MinuteStr, *SecondStr, *HundredthsStr);
+		return FText::FromString(FString::Printf(TEXT("%s (%s%s%s)"), *FramesStr, *MinuteStr, *SecondStr, *HundredthsStr));
 	}
 
-	return TEXT("Position");
+	return NSLOCTEXT("SGraphNodeSequencePlayer", "PositionToolTip_Default", "Position");
 }
 
 void SGraphNodeSequencePlayer::UpdateGraphNode()
diff --git a/Engine/Source/Editor/GraphEditor/Private/AnimationNodes/SGraphNodeSequencePlayer.h b/Engine/Source/Editor/GraphEditor/Private/AnimationNodes/SGraphNodeSequencePlayer.h
--- a/Engine/Source/Editor/GraphEditor/Private/AnimationNodes/SGraphNodeSequencePlayer.h
+++ b/Engine/Source/Editor/GraphEditor/Private/AnimationNodes/SGraphNodeSequencePlayer.h
@@ -26,7 +26,7 @@ protected:
 	float GetSequencePositionRatio() const;
 	void SetSequencePositionRatio(float NewRatio);
 
-	FString GetPositionTooltip() const;
+	FText GetPositionTooltip() const;
 
 	bool GetSequencePositionInfo(float& Out_Position, float& Out_Length, int32& FrameCount) const;
 };
diff --git a/Engine/Source/Editor/GraphEditor/Private/KismetPins/SGraphPinObject.cpp b/Engine/Source/Editor/GraphEditor/Private/KismetPins/SGraphPinObject.cpp
--- a/Engine/Source/Editor/GraphEditor/Private/KismetPins/SGraphPinObject.cpp
+++ b/Engine/Source/Editor/GraphEditor/Private/KismetPins/SGraphPinObject.cpp
@@ -85,7 +85,7 @@ TSharedRef<SWidget>	SGraphPinObject::GetDefaultValueWidget()
 				.ButtonColorAndOpacity( this, &SGraphPinObject::OnGetWidgetBackground )
 				.OnClicked( GetOnUseButtonDelegate() )
 				.ContentPadding(1.f)
-				.ToolTipText(NSLOCTEXT("GraphEditor", "ObjectGraphPin_Use", "Use asset browser selection").ToString())
+				.ToolTipText(NSLOCTEXT("GraphEditor", "ObjectGraphPin_Use", "Use asset browser selection"))
 				[
 					SNew(SImage)
 					.ColorAndOpacity( this, &SGraphPinObject::OnGetWidgetForeground )
@@ -103,7 +103,7 @@ TSharedRef<SWidget>	SGraphPinObject::GetDefaultValueWidget()
 				.ButtonColorAndOpacity( this, &SGraphPinObject::OnGetWidgetBackground )
 				.OnClicked( GetOnBrowseButtonDelegate() )
 				.ContentPadding(0)
-				.ToolTipText(NSLOCTEXT("GraphEditor", "ObjectGraphPin_Browse", "Browse").ToString())
+				.ToolTipText(NSLOCTEXT("GraphEditor", "ObjectGraphPin_Browse", "Browse"))
 				[
 					SNew(SImage)
 					.ColorAndOpacity( this, &SGraphPinObject::OnGetWidgetForeground )
diff --git a/Engine/Source/Editor/GraphEditor/Private/SGraphNode.cpp b/Engine/Source/Editor/GraphEditor/Private/SGraphNode.cpp
--- a/Engine/Source/Editor/GraphEditor/Private/SGraphNode.cpp
+++ b/Engine/Source/Editor/GraphEditor/Private/SGraphNode.cpp
@@ -1178,9 +1178,9 @@ void SGraphNode::PositionThisNodeBetweenOtherNodes(const FVector2D& PrevPos, con
 	GraphNode->NodePosY = NewCorner.Y;
 }
 
-FString SGraphNode::GetErrorMsgToolTip( ) const
+FText SGraphNode::GetErrorMsgToolTip( ) const
 {
-	return GraphNode->ErrorMsg;
+	return FText::FromString(GraphNode->ErrorMsg);
 }
 
 bool SGraphNode::IsNameReadOnly() const
diff --git a/Engine/Source/Editor/GraphEditor/Private/SGraphPalette.cpp b/Engine/Source/Editor/GraphEditor/Private/SGraphPalette.cpp
--- a/Engine/Source/Editor/GraphEditor/Private/SGraphPalette.cpp
+++ b/Engine/Source/Editor/GraphEditor/Private/SGraphPalette.cpp
@@ -160,9 +160,9 @@ FText SGraphPaletteItem::GetDisplayText() const
 	return ActionPtr.Pin()->MenuDescription;
 }
 
-FString SGraphPaletteItem::GetItemTooltip() const
+FText SGraphPaletteItem::GetItemTooltip() const
 {
-	return ActionPtr.Pin()->MenuDescription.ToString();
+	return ActionPtr.Pin()->MenuDescription;
 }
 
 
diff --git a/Engine/Source/Editor/GraphEditor/Private/SGraphPin.cpp b/Engine/Source/Editor/GraphEditor/Private/SGraphPin.cpp
--- a/Engine/Source/Editor/GraphEditor/Private/SGraphPin.cpp
+++ b/Engine/Source/Editor/GraphEditor/Private/SGraphPin.cpp
@@ -249,10 +249,10 @@ void SGraphPin::Construct(const FArguments& InArgs, UEdGraphPin* InPin)
 			];
 	}
 
-	TAttribute<FString> ToolTipAttribute;
+	TAttribute<FText> ToolTipAttribute;
 	if (InArgs._HasToolTip)
 	{
-		ToolTipAttribute = TAttribute<FString>::Create(TAttribute<FString>::FGetter::CreateSP(this, &SGraphPin::GetTooltip));
+		ToolTipAttribute = TAttribute<FText>::Create(TAttribute<FText>::FGetter::CreateSP(this, &SGraphPin::GetTooltip));
 	}
 
 	TSharedPtr<SWidget> PinContent;
@@ -948,15 +948,20 @@ void SGraphPin::SetShowLabel(bool bNewShowLabel)
 	bShowLabel = bNewShowLabel;
 }
 
-FString SGraphPin::GetTooltip() const
+FText SGraphPin::GetTooltip() const
 {
-	FString HoverText;
+	FText HoverText = FText::GetEmpty();
 
 	check(GraphPinObj != nullptr);
 	UEdGraphNode* GraphNode = GraphPinObj->GetOwningNodeUnchecked();
 	if (GraphNode != nullptr)
 	{
-		GraphNode->GetPinHoverText(*GraphPinObj, /*out*/ HoverText);
+		FString HoverStr;
+		GraphNode->GetPinHoverText(*GraphPinObj, /*out*/HoverStr);
+		if (!HoverStr.IsEmpty())
+		{
+			HoverText = FText::FromString(HoverStr);
+		}
 	}
 
 	return HoverText;
diff --git a/Engine/Source/Editor/GraphEditor/Public/SGraphNode.h b/Engine/Source/Editor/GraphEditor/Public/SGraphNode.h
--- a/Engine/Source/Editor/GraphEditor/Public/SGraphNode.h
+++ b/Engine/Source/Editor/GraphEditor/Public/SGraphNode.h
@@ -212,7 +212,7 @@ protected:
 	 *
 	 *	@return string to be displayed as tooltip.
 	 */
-	FString GetErrorMsgToolTip() const;
+	FText GetErrorMsgToolTip() const;
 
 	/**
 	 * Add a new pin to this graph node. The pin must be newly created.
diff --git a/Engine/Source/Editor/GraphEditor/Public/SGraphPalette.h b/Engine/Source/Editor/GraphEditor/Public/SGraphPalette.h
--- a/Engine/Source/Editor/GraphEditor/Public/SGraphPalette.h
+++ b/Engine/Source/Editor/GraphEditor/Public/SGraphPalette.h
@@ -51,7 +51,7 @@ protected:
 	virtual FText GetDisplayText() const;
 
 	/** Callback to get the tooltip */
-	virtual FString GetItemTooltip() const;
+	virtual FText GetItemTooltip() const;
 };
 
 //////////////////////////////////////////////////////////////////////////
diff --git a/Engine/Source/Editor/GraphEditor/Public/SGraphPin.h b/Engine/Source/Editor/GraphEditor/Public/SGraphPin.h
--- a/Engine/Source/Editor/GraphEditor/Public/SGraphPin.h
+++ b/Engine/Source/Editor/GraphEditor/Public/SGraphPin.h
@@ -140,7 +140,7 @@ protected:
 	virtual FSlateColor GetPinColor() const;
 
 	/** @return The tooltip to display for this pin */
-	FString GetTooltip() const;
+	FText GetTooltip() const;
 
 	TOptional<EMouseCursor::Type> GetPinCursor() const;
 
diff --git a/Engine/Source/Editor/Kismet/Private/Debugging/SKismetDebuggingView.cpp b/Engine/Source/Editor/Kismet/Private/Debugging/SKismetDebuggingView.cpp
--- a/Engine/Source/Editor/Kismet/Private/Debugging/SKismetDebuggingView.cpp
+++ b/Engine/Source/Editor/Kismet/Private/Debugging/SKismetDebuggingView.cpp
@@ -503,7 +503,7 @@ protected:
 	void OnNavigateToBreakpointLocation();
 
 	const FSlateBrush* GetStatusImage() const;
-	FString GetStatusTooltip() const;
+	FText GetStatusTooltip() const;
 };
 
 FText FBreakpointLineItem::GetLocationDescription() const
@@ -549,22 +549,22 @@ const FSlateBrush* FBreakpointLineItem::GetStatusImage() const
 	return FEditorStyle::GetDefaultBrush();
 }
 
-FString FBreakpointLineItem::GetStatusTooltip() const
+FText FBreakpointLineItem::GetStatusTooltip() const
 {
 	if (UBreakpoint* MyBreakpoint = BreakpointRef.Get())
 	{
 		if (!FKismetDebugUtilities::IsBreakpointValid(MyBreakpoint))
 		{
-			return LOCTEXT("Breakpoint_NoHit", "This breakpoint will not be hit because its node generated no code").ToString();
+			return LOCTEXT("Breakpoint_NoHit", "This breakpoint will not be hit because its node generated no code");
 		}
 		else
 		{
-			return MyBreakpoint->IsEnabledByUser() ? LOCTEXT("ActiveBreakpoint", "Active breakpoint").ToString() : LOCTEXT("InactiveBreakpoint", "Inactive breakpoint").ToString();
+			return MyBreakpoint->IsEnabledByUser() ? LOCTEXT("ActiveBreakpoint", "Active breakpoint") : LOCTEXT("InactiveBreakpoint", "Inactive breakpoint");
 		}
 	}
 	else
 	{
-		return LOCTEXT("NoBreakpoint", "No Breakpoint").ToString();
+		return LOCTEXT("NoBreakpoint", "No Breakpoint");
 	}
 }
 
diff --git a/Engine/Source/Editor/Kismet/Private/SBlueprintPalette.cpp b/Engine/Source/Editor/Kismet/Private/SBlueprintPalette.cpp
--- a/Engine/Source/Editor/Kismet/Private/SBlueprintPalette.cpp
+++ b/Engine/Source/Editor/Kismet/Private/SBlueprintPalette.cpp
@@ -831,12 +831,12 @@ private:
 	 * 
 	 * @return Tooltip text for this toggle.
 	 */
-	FString GetVisibilityToggleToolTip() const
+	FText GetVisibilityToggleToolTip() const
 	{
-		FString ToolTip;
+		FText ToolTip = FText::GetEmpty();
 		if(GetVisibilityToggleState() != ECheckBoxState::Checked)
 		{
-			ToolTip =  LOCTEXT("VariablePrivacy_not_public_Tooltip", "Variable is not public and will not be editable on an instance of this Blueprint.").ToString();
+			ToolTip = LOCTEXT("VariablePrivacy_not_public_Tooltip", "Variable is not public and will not be editable on an instance of this Blueprint.");
 		}
 		else
 		{
@@ -846,11 +846,11 @@ private:
 			FBlueprintEditorUtils::GetBlueprintVariableMetaData( BlueprintObj, VarAction->GetVariableName(), NULL, TEXT("tooltip"), Result);
 			if(!Result.IsEmpty())
 			{
-				ToolTip = LOCTEXT("VariablePrivacy_is_public_Tooltip", "Variable is public and is editable on each instance of this Blueprint.").ToString();
+				ToolTip = LOCTEXT("VariablePrivacy_is_public_Tooltip", "Variable is public and is editable on each instance of this Blueprint.");
 			}
 			else
 			{
-				ToolTip = LOCTEXT("VariablePrivacy_is_public_no_tooltip_Tooltip", "Variable is public but MISSING TOOLTIP.").ToString();
+				ToolTip = LOCTEXT("VariablePrivacy_is_public_no_tooltip_Tooltip", "Variable is public but MISSING TOOLTIP.");
 			}
 		}
 		return ToolTip;
diff --git a/Engine/Source/Editor/Kismet/Private/STimelineEditor.cpp b/Engine/Source/Editor/Kismet/Private/STimelineEditor.cpp
--- a/Engine/Source/Editor/Kismet/Private/STimelineEditor.cpp
+++ b/Engine/Source/Editor/Kismet/Private/STimelineEditor.cpp
@@ -131,7 +131,7 @@ void STimelineEdTrack::Construct(const FArguments& InArgs, TSharedPtr<FTimelineE
 						// Name of track
 						SAssignNew(InlineTextBlock, SInlineEditableTextBlock)
 							.Text(FText::FromName(TrackBase->TrackName))
-							.ToolTipText(FString(TEXT("Enter track name")))
+							.ToolTipText(LOCTEXT("TrackNameTooltip", "Enter track name"))
 							.OnVerifyTextChanged(TimelineRef, &STimelineEditor::OnVerifyTrackNameCommit,TrackBase, this)
 							.OnTextCommitted(TimelineRef, &STimelineEditor::OnTrackNameCommitted,TrackBase, this)
 					]
@@ -524,9 +524,9 @@ FText STimelineEdTrack::GetExternalCurveName( ) const
 	return FText::FromString(ExternalCurveName);
 }
 
-FString STimelineEdTrack::GetExternalCurvePath( ) const
+FText STimelineEdTrack::GetExternalCurvePath( ) const
 {
-	return ExternalCurvePath;
+	return FText::FromString(ExternalCurvePath);
 }
 
 UCurveBase* STimelineEdTrack::CreateCurveAsset()
diff --git a/Engine/Source/Editor/Kismet/Private/STimelineEditor.h b/Engine/Source/Editor/Kismet/Private/STimelineEditor.h
--- a/Engine/Source/Editor/Kismet/Private/STimelineEditor.h
+++ b/Engine/Source/Editor/Kismet/Private/STimelineEditor.h
@@ -95,7 +95,7 @@ private:
 	FText GetExternalCurveName( ) const;
 
 	/** Get the current external curve path*/
-	FString GetExternalCurvePath( ) const;
+	FText GetExternalCurvePath( ) const;
 
 	/** Function to replace internal curve with an external curve*/
 	void SwitchToExternalCurve(UCurveBase* AssetCurvePtr);
diff --git a/Engine/Source/Editor/Kismet/Private/UserDefinedStructureEditor.cpp b/Engine/Source/Editor/Kismet/Private/UserDefinedStructureEditor.cpp
--- a/Engine/Source/Editor/Kismet/Private/UserDefinedStructureEditor.cpp
+++ b/Engine/Source/Editor/Kismet/Private/UserDefinedStructureEditor.cpp
@@ -350,7 +350,7 @@ public:
 		return NULL;
 	}
 
-	FString GetStatusTooltip() const
+	FText GetStatusTooltip() const
 	{
 		auto StructureDetailsSP = StructureDetails.Pin();
 		if (StructureDetailsSP.IsValid())
@@ -360,11 +360,11 @@ public:
 				switch (Struct->Status.GetValue())
 				{
 				case EUserDefinedStructureStatus::UDSS_Error:
-					return Struct->ErrorMessage;
+					return FText::FromString(Struct->ErrorMessage);
 				}
 			}
 		}
-		return FString();
+		return FText::GetEmpty();
 	}
 
 	FText OnGetTooltipText() const
diff --git a/Engine/Source/Editor/Layers/Private/SLayerStats.h b/Engine/Source/Editor/Layers/Private/SLayerStats.h
--- a/Engine/Source/Editor/Layers/Private/SLayerStats.h
+++ b/Engine/Source/Editor/Layers/Private/SLayerStats.h
@@ -158,9 +158,9 @@ private:
 	 *
 	 *	@param Class	The Actor class which the stats are associated with
 	 */
-	FString GetStatButtonToolTipText( const TWeakObjectPtr< UClass > Class ) const
+	FText GetStatButtonToolTipText( const TWeakObjectPtr< UClass > Class ) const
 	{
-		return FString::Printf( *LOCTEXT("StatButtonToolTip", "Select All %s Actors in %s").ToString(), *Class->GetFName().ToString(), *ViewModel->GetName() );
+		return FText::Format( LOCTEXT("StatButtonToolTipFmt", "Select All {0} Actors in {1}"), FText::FromName(Class->GetFName()), FText::FromString(ViewModel->GetName()) );
 	}
 
 	/**
@@ -168,9 +168,9 @@ private:
 	 *
 	 *	@param Class	The Actor class which the stats are associated with
 	 */
-	FString GetRemoveStatButtonToolTipText( const TWeakObjectPtr< UClass > Class ) const
+	FText GetRemoveStatButtonToolTipText( const TWeakObjectPtr< UClass > Class ) const
 	{
-		return FString::Printf( *LOCTEXT("RemoveStatButtonToolTip", "Remove All").ToString(), *Class->GetFName().ToString(), *ViewModel->GetName() );
+		return LOCTEXT("RemoveStatButtonToolTip", "Remove All");
 	}
 
 
diff --git a/Engine/Source/Editor/Layers/Private/SceneOutlinerLayerContentsColumn.cpp b/Engine/Source/Editor/Layers/Private/SceneOutlinerLayerContentsColumn.cpp
--- a/Engine/Source/Editor/Layers/Private/SceneOutlinerLayerContentsColumn.cpp
+++ b/Engine/Source/Editor/Layers/Private/SceneOutlinerLayerContentsColumn.cpp
@@ -41,7 +41,7 @@ TSharedRef<SWidget> FSceneOutlinerLayerContentsColumn::ConstructRowWidget(const
 		.ButtonStyle( FEditorStyle::Get(), "LayerBrowserButton" )
 		.ContentPadding( 0 )
 		.OnClicked( this, &FSceneOutlinerLayerContentsColumn::OnRemoveFromLayerClicked, Actor )
-		.ToolTipText( LOCTEXT("RemoveFromLayerButtonText", "Remove from Layer").ToString() )
+		.ToolTipText( LOCTEXT("RemoveFromLayerButtonText", "Remove from Layer") )
 		[
 			SNew( SImage )
 			.Image( FEditorStyle::GetBrush( TEXT( "LayerBrowser.Actor.RemoveFromLayer" ) ) )
diff --git a/Engine/Source/Editor/LevelEditor/Private/LevelEditorActions.cpp b/Engine/Source/Editor/LevelEditor/Private/LevelEditorActions.cpp
--- a/Engine/Source/Editor/LevelEditor/Private/LevelEditorActions.cpp
+++ b/Engine/Source/Editor/LevelEditor/Private/LevelEditorActions.cpp
@@ -2114,14 +2114,14 @@ bool FLevelEditorActionCallbacks::OnIsLevelStreamingVolumePrevisEnabled()
 	return GetDefault<ULevelEditorViewportSettings>()->bLevelStreamingVolumePrevis;
 }
 
-FString FLevelEditorActionCallbacks::GetAudioVolumeToolTip()
+FText FLevelEditorActionCallbacks::GetAudioVolumeToolTip()
 {
 	if ( !GEditor->IsRealTimeAudioMuted() )
 	{
 		const float Volume = GEditor->GetRealTimeAudioVolume() * 100.0f;
-		return FString::Printf( TEXT( "%.0f" ), Volume );
+		return FText::AsNumber( FMath::RoundToInt(Volume) );
 	}
-	return NSLOCTEXT("UnrealEd", "Muted", "Muted" ).ToString();
+	return NSLOCTEXT("UnrealEd", "Muted", "Muted" );
 }
 
 float FLevelEditorActionCallbacks::GetAudioVolume()
@@ -2172,14 +2172,17 @@ bool FLevelEditorActionCallbacks::OnIsVertexSnapEnabled()
 	return GetDefault<ULevelEditorViewportSettings>()->bSnapVertices;
 }
 
-FString FLevelEditorActionCallbacks::GetActorSnapTooltip()
+FText FLevelEditorActionCallbacks::GetActorSnapTooltip()
 {
 	// If the setting is enabled, return the distance, otherwise say disabled
 	if ( FSnappingUtils::IsSnapToActorEnabled() )
 	{
-		return FString::Printf( TEXT( "%.2f" ), FSnappingUtils::GetActorSnapDistance() );
+		static const FNumberFormattingOptions FormatOptions = FNumberFormattingOptions()
+			.SetMinimumFractionalDigits(2)
+			.SetMaximumFractionalDigits(2);
+		return FText::AsNumber( FSnappingUtils::GetActorSnapDistance(), &FormatOptions );
 	}
-	return NSLOCTEXT("UnrealEd", "Disabled", "Disabled" ).ToString();
+	return NSLOCTEXT("UnrealEd", "Disabled", "Disabled" );
 }
 
 float FLevelEditorActionCallbacks::GetActorSnapSetting()
diff --git a/Engine/Source/Editor/LevelEditor/Private/LevelEditorContextMenu.cpp b/Engine/Source/Editor/LevelEditor/Private/LevelEditorContextMenu.cpp
--- a/Engine/Source/Editor/LevelEditor/Private/LevelEditorContextMenu.cpp
+++ b/Engine/Source/Editor/LevelEditor/Private/LevelEditorContextMenu.cpp
@@ -808,7 +808,7 @@ void FLevelEditorContextMenuImpl::FillActorMenu( FMenuBuilder& MenuBuilder )
 			.Padding(4.0f, 0.0f, 0.0f, 0.0f)
 			[
 				SNew(SButton)
-				.ToolTipText( LOCTEXT( "PickButtonLabel", "Pick a parent actor to attach to").ToString() )
+				.ToolTipText( LOCTEXT( "PickButtonLabel", "Pick a parent actor to attach to") )
 				.ButtonStyle(FEditorStyle::Get(), "HoverHintOnly")
 				.OnClicked(FOnClicked::CreateStatic(&Local::OnInteractiveActorPickerClicked))
 				.ContentPadding(4.0f)
diff --git a/Engine/Source/Editor/LevelEditor/Private/SToolkitDisplay.cpp b/Engine/Source/Editor/LevelEditor/Private/SToolkitDisplay.cpp
--- a/Engine/Source/Editor/LevelEditor/Private/SToolkitDisplay.cpp
+++ b/Engine/Source/Editor/LevelEditor/Private/SToolkitDisplay.cpp
@@ -83,7 +83,7 @@ void SLevelEditorActiveToolkit::Construct( const FArguments&, const TSharedPtr<
 												SNew( SImage )
 												.Visibility( this, &SLevelEditorActiveToolkit::GetVisibilityForUnsavedChangeIcon )
 												.Image( FEditorStyle::GetBrush( "ToolkitDisplay.UnsavedChangeIcon" ) )
-												.ToolTipText( LOCTEXT("UnsavedChangeToolTip", "This asset has unsaved changes").ToString() )
+												.ToolTipText( LOCTEXT("UnsavedChangeToolTip", "This asset has unsaved changes") )
 											]
 									]
 						]
@@ -124,7 +124,7 @@ void SLevelEditorActiveToolkit::Construct( const FArguments&, const TSharedPtr<
 		[
 			SNew( SButton )
 				.OnClicked( this, &SLevelEditorActiveToolkit::OnToolkitCloseButtonClicked )
-				.ToolTipText( LOCTEXT("CloseToolkitButton", "Close").ToString() )
+				.ToolTipText( LOCTEXT("CloseToolkitButton", "Close") )
 				.ButtonStyle( FEditorStyle::Get(), "NoBorder" )
 				.ContentPadding(0)
 				.HAlign(HAlign_Center)
diff --git a/Engine/Source/Editor/LevelEditor/Public/LevelEditorActions.h b/Engine/Source/Editor/LevelEditor/Public/LevelEditorActions.h
--- a/Engine/Source/Editor/LevelEditor/Public/LevelEditorActions.h
+++ b/Engine/Source/Editor/LevelEditor/Public/LevelEditorActions.h
@@ -1080,7 +1080,7 @@ public:
 	static void OnToggleLevelStreamingVolumePrevis();
 	static bool OnIsLevelStreamingVolumePrevisEnabled(); 
 	
-	static FString GetAudioVolumeToolTip();
+	static FText GetAudioVolumeToolTip();
 	static float GetAudioVolume();
 	static void OnAudioVolumeChanged(float Volume);
 	static bool GetAudioMuted();
@@ -1088,7 +1088,7 @@ public:
 
 	static void OnEnableActorSnap();
 	static bool OnIsActorSnapEnabled();
-	static FString GetActorSnapTooltip();
+	static FText GetActorSnapTooltip();
 	static float GetActorSnapSetting();
 	static void SetActorSnapSetting(float Distance);
 	static void OnEnableVertexSnap();
diff --git a/Engine/Source/Editor/Levels/Private/SLevelBrowser.cpp b/Engine/Source/Editor/Levels/Private/SLevelBrowser.cpp
--- a/Engine/Source/Editor/Levels/Private/SLevelBrowser.cpp
+++ b/Engine/Source/Editor/Levels/Private/SLevelBrowser.cpp
@@ -40,7 +40,7 @@ void SLevelBrowser::Construct(const FArguments& InArgs)
 				.AutoWidth()
 				[
 					SNew(SButton)
-					.ToolTipText(LOCTEXT("ShiftUpHint", "Shift Streaming Level Up").ToString())
+					.ToolTipText(LOCTEXT("ShiftUpHint", "Shift Streaming Level Up"))
 					.ContentPadding(0)
 					.OnClicked(this, &SLevelBrowser::OnClickUp)
 					.IsEnabled(this, &SLevelBrowser::CanShiftSelection)
@@ -54,7 +54,7 @@ void SLevelBrowser::Construct(const FArguments& InArgs)
 					.AutoWidth()
 					[
 						SNew(SButton)
-						.ToolTipText(LOCTEXT("ShiftDownHint", "Shift Streaming Level Down").ToString())
+						.ToolTipText(LOCTEXT("ShiftDownHint", "Shift Streaming Level Down"))
 						.ContentPadding(0)
 						.OnClicked(this, &SLevelBrowser::OnClickDown)
 						.IsEnabled(this, &SLevelBrowser::CanShiftSelection)
@@ -67,7 +67,7 @@ void SLevelBrowser::Construct(const FArguments& InArgs)
 					.FillWidth(1.f)
 					[
 						SNew(SSearchBox)
-						.ToolTipText(LOCTEXT("FilterSearchToolTip", "Type here to search Levels").ToString())
+						.ToolTipText(LOCTEXT("FilterSearchToolTip", "Type here to search Levels"))
 						.HintText(LOCTEXT("FilterSearchHint", "Search Levels"))
 						.OnTextChanged(SearchBoxLevelFilter.Get(), &LevelTextFilter::SetRawFilterText)
 					]
diff --git a/Engine/Source/Editor/Levels/Private/SLevelsViewRow.h b/Engine/Source/Editor/Levels/Private/SLevelsViewRow.h
--- a/Engine/Source/Editor/Levels/Private/SLevelsViewRow.h
+++ b/Engine/Source/Editor/Levels/Private/SLevelsViewRow.h
@@ -438,7 +438,7 @@ private:
 	 *
 	 *	@return	The tooltip representing the Level's lock state
 	 */
-	FString GetLockToolTipForLevel() const
+	FText GetLockToolTipForLevel() const
 	{
 		//Non-Persistent
 
@@ -446,14 +446,14 @@ private:
 		{
 			if(ViewModel->IsReadOnly())
 			{
-				return LOCTEXT("ReadOnly_LockButtonToolTip", "Read-Only levels are locked!").ToString();
+				return LOCTEXT("ReadOnly_LockButtonToolTip", "Read-Only levels are locked!");
 			}
 		}
 
-		return LOCTEXT("LockButtonToolTip", "Toggle Level Lock").ToString();
+		return LOCTEXT("LockButtonToolTip", "Toggle Level Lock");
 	}
 
-	FString GetSCCStateTooltip() const
+	FText GetSCCStateTooltip() const
 	{
 		ULevel* Level = ViewModel->GetLevel().Get();
 		if (Level != NULL)
@@ -461,11 +461,11 @@ private:
 			FSourceControlStatePtr SourceControlState = ISourceControlModule::Get().GetProvider().GetState(Level->GetOutermost(), EStateCacheUsage::Use);
 			if(SourceControlState.IsValid())
 			{
-				return SourceControlState->GetDisplayTooltip().ToString();
+				return SourceControlState->GetDisplayTooltip();
 			}
 		}
 
-		return FString();
+		return FText::GetEmpty();
 	}
 
 	const FSlateBrush* GetSCCStateImage() const
diff --git a/Engine/Source/Editor/Matinee/Private/SMatineeRecorder.cpp b/Engine/Source/Editor/Matinee/Private/SMatineeRecorder.cpp
--- a/Engine/Source/Editor/Matinee/Private/SMatineeRecorder.cpp
+++ b/Engine/Source/Editor/Matinee/Private/SMatineeRecorder.cpp
@@ -72,7 +72,7 @@ void SMatineeRecorder::Construct(const FArguments& InArgs)
 				[
 					SAssignNew( RecordButton, SButton )
 					.OnClicked( this, &SMatineeRecorder::ToggleRecord )
-					.ToolTipText( LOCTEXT("StartStopRecording", "Start/Stop Recording").ToString() )
+					.ToolTipText( LOCTEXT("StartStopRecording", "Start/Stop Recording") )
 					[
 						SNew(SHorizontalBox)
 						+ SHorizontalBox::Slot()
@@ -90,7 +90,7 @@ void SMatineeRecorder::Construct(const FArguments& InArgs)
 					SAssignNew( CameraModeComboBox, STextComboBox )
 					.OptionsSource( &CameraModeOptions )
 					.OnSelectionChanged( this, &SMatineeRecorder::SelectCameraMode )
-					.ToolTipText( LOCTEXT("ChangeCameraMode", "Change Camera Mode").ToString() )
+					.ToolTipText( LOCTEXT("ChangeCameraMode", "Change Camera Mode") )
 					.InitiallySelectedItem( CameraModeOptions[ 0 ] )
 				]
 			]
diff --git a/Engine/Source/Editor/Persona/Private/SAnimMontagePanel.cpp b/Engine/Source/Editor/Persona/Private/SAnimMontagePanel.cpp
--- a/Engine/Source/Editor/Persona/Private/SAnimMontagePanel.cpp
+++ b/Engine/Source/Editor/Persona/Private/SAnimMontagePanel.cpp
@@ -185,7 +185,7 @@ void SAnimMontagePanel::Update()
 						.OnComboBoxOpening(this, &SAnimMontagePanel::OnSlotListOpening, SlotAnimIdx)
 						.InitiallySelectedItem(ComboItem)
 						.ContentPadding(2)
-						.ToolTipText(*ComboItem)
+						.ToolTipText(FText::FromString(*ComboItem))
 					]
 
 					+ SVerticalBox::Slot()
diff --git a/Engine/Source/Editor/Persona/Private/SAnimSegmentsPanel.cpp b/Engine/Source/Editor/Persona/Private/SAnimSegmentsPanel.cpp
--- a/Engine/Source/Editor/Persona/Private/SAnimSegmentsPanel.cpp
+++ b/Engine/Source/Editor/Persona/Private/SAnimSegmentsPanel.cpp
@@ -126,7 +126,7 @@ FString	SAnimSegmentsPanel::GetAnimSegmentName(int32 AnimSegmentIndex) const
 	return FString();
 }
 
-FString	SAnimSegmentsPanel::GetAnimSegmentDetailedInfo(int32 AnimSegmentIndex) const
+FText SAnimSegmentsPanel::GetAnimSegmentDetailedInfo(int32 AnimSegmentIndex) const
 {
 	if (ValidIndex(AnimSegmentIndex))
 	{
@@ -134,10 +134,13 @@ FString	SAnimSegmentsPanel::GetAnimSegmentDetailedInfo(int32 AnimSegmentIndex) c
 		UAnimSequenceBase * Anim = AnimSegment.AnimReference;
 		if ( Anim != NULL )
 		{
-			return FString::Printf(TEXT("%s %.2f"), *Anim->GetName(), AnimSegment.GetLength() );
+			static const FNumberFormattingOptions FormatOptions = FNumberFormattingOptions()
+				.SetMinimumFractionalDigits(2)
+				.SetMaximumFractionalDigits(2);
+			return FText::Format(LOCTEXT("AnimSegmentPanel_GetAnimSegmentDetailedInfoFmt", "{0} {1}"), FText::FromString(Anim->GetName()), FText::AsNumber(AnimSegment.GetLength(), &FormatOptions) );
 		}
 	}
-	return FString();
+	return FText::GetEmpty();
 }
 
 void SAnimSegmentsPanel::SetSegmentStartPos(float NewStartPos, int32 AnimSegmentIndex)
diff --git a/Engine/Source/Editor/Persona/Private/SAnimSegmentsPanel.h b/Engine/Source/Editor/Persona/Private/SAnimSegmentsPanel.h
--- a/Engine/Source/Editor/Persona/Private/SAnimSegmentsPanel.h
+++ b/Engine/Source/Editor/Persona/Private/SAnimSegmentsPanel.h
@@ -72,7 +72,7 @@ private:
 	float				GetSegmentLength(int32 AnimSegmentIndex) const;
 	float				GetSegmentStartPos(int32 AnimSegmentIndex) const;
 	FString				GetAnimSegmentName(int32 AnimSegmentIndex) const;
-	FString				GetAnimSegmentDetailedInfo(int32 AnimSegmentIndex) const;
+	FText				GetAnimSegmentDetailedInfo(int32 AnimSegmentIndex) const;
 
 	void				SetSegmentStartPos(float NewStartPos, int32 AnimSegmentIndex);
 	void				OnSegmentDropped(int32 AnimSegmentIndex);
diff --git a/Engine/Source/Editor/PropertyEditor/Private/PropertyCustomizationHelpers.cpp b/Engine/Source/Editor/PropertyEditor/Private/PropertyCustomizationHelpers.cpp
--- a/Engine/Source/Editor/PropertyEditor/Private/PropertyCustomizationHelpers.cpp
+++ b/Engine/Source/Editor/PropertyEditor/Private/PropertyCustomizationHelpers.cpp
@@ -241,7 +241,7 @@ namespace PropertyCustomizationHelpers
 	{
 		return 
 			SNew( SPropertyEditorInteractiveActorPicker )
-			.ToolTipText( LOCTEXT( "PickButtonLabel", "Pick Actor from scene").ToString() )
+			.ToolTipText( LOCTEXT( "PickButtonLabel", "Pick Actor from scene") )
 			.OnGetAllowedClasses( OnGetAllowedClasses )
 			.OnShouldFilterActor( OnShouldFilterActor )
 			.OnActorSelected( OnActorSelectedFromPicker );
diff --git a/Engine/Source/Editor/PropertyEditor/Private/UserInterface/PropertyEditor/SPropertyEditor.h b/Engine/Source/Editor/PropertyEditor/Private/UserInterface/PropertyEditor/SPropertyEditor.h
--- a/Engine/Source/Editor/PropertyEditor/Private/UserInterface/PropertyEditor/SPropertyEditor.h
+++ b/Engine/Source/Editor/PropertyEditor/Private/UserInterface/PropertyEditor/SPropertyEditor.h
@@ -26,7 +26,7 @@ public:
 				// Make a read only text box so that copy still works
 				SNew( SEditableTextBox )
 				.Text( InPropertyEditor, &FPropertyEditor::GetValueAsText )
-				.ToolTipText( InPropertyEditor, &FPropertyEditor::GetValueAsString )
+				.ToolTipText( InPropertyEditor, &FPropertyEditor::GetValueAsText )
 				.Font( InArgs._Font )
 				.IsReadOnly( true )
 			];
diff --git a/Engine/Source/Editor/PropertyEditor/Private/UserInterface/PropertyEditor/SPropertyEditorCombo.cpp b/Engine/Source/Editor/PropertyEditor/Private/UserInterface/PropertyEditor/SPropertyEditorCombo.cpp
--- a/Engine/Source/Editor/PropertyEditor/Private/UserInterface/PropertyEditor/SPropertyEditorCombo.cpp
+++ b/Engine/Source/Editor/PropertyEditor/Private/UserInterface/PropertyEditor/SPropertyEditorCombo.cpp
@@ -80,7 +80,7 @@ void SPropertyEditorCombo::Construct( const FArguments& InArgs, const TSharedRef
 		.OnSelectionChanged( this, &SPropertyEditorCombo::OnComboSelectionChanged )
 		.OnComboBoxOpening( this, &SPropertyEditorCombo::OnComboOpening )
 		.VisibleText( this, &SPropertyEditorCombo::GetDisplayValueAsString )
-		.ToolTipText( InPropertyEditor, &FPropertyEditor::GetValueAsString );
+		.ToolTipText( InPropertyEditor, &FPropertyEditor::GetValueAsText );
 
 	ChildSlot
 	[
diff --git a/Engine/Source/Editor/PropertyEditor/Private/UserInterface/PropertyEditor/SPropertyEditorEditInline.cpp b/Engine/Source/Editor/PropertyEditor/Private/UserInterface/PropertyEditor/SPropertyEditorEditInline.cpp
--- a/Engine/Source/Editor/PropertyEditor/Private/UserInterface/PropertyEditor/SPropertyEditorEditInline.cpp
+++ b/Engine/Source/Editor/PropertyEditor/Private/UserInterface/PropertyEditor/SPropertyEditorEditInline.cpp
@@ -72,7 +72,7 @@ void SPropertyEditorEditInline::Construct( const FArguments& InArgs, const TShar
 		SAssignNew(ComboButton, SComboButton)
 		.OnGetMenuContent(this, &SPropertyEditorEditInline::GenerateClassPicker)
 		.ContentPadding(0)
-		.ToolTipText(InPropertyEditor, &FPropertyEditor::GetValueAsString )
+		.ToolTipText(InPropertyEditor, &FPropertyEditor::GetValueAsText )
 		.ButtonContent()
 		[
 			SNew(SHorizontalBox)
diff --git a/Engine/Source/Editor/ReferenceViewer/Private/HistoryManager.cpp b/Engine/Source/Editor/ReferenceViewer/Private/HistoryManager.cpp
--- a/Engine/Source/Editor/ReferenceViewer/Private/HistoryManager.cpp
+++ b/Engine/Source/Editor/ReferenceViewer/Private/HistoryManager.cpp
@@ -113,28 +113,22 @@ bool FReferenceViewerHistoryManager::CanGoBack() const
 	return HistoryData.Num() > 0 && CurrentHistoryIndex > 0;
 }
 
-FString FReferenceViewerHistoryManager::GetBackDesc() const
+FText FReferenceViewerHistoryManager::GetBackDesc() const
 {
 	if ( CanGoBack() )
 	{
 		return HistoryData[CurrentHistoryIndex - 1].HistoryDesc;
 	}
-	else
-	{
-		return FString();
-	}
+	return FText::GetEmpty();
 }
 
-FString FReferenceViewerHistoryManager::GetForwardDesc() const
+FText FReferenceViewerHistoryManager::GetForwardDesc() const
 {
 	if ( CanGoForward() )
 	{
 		return HistoryData[CurrentHistoryIndex + 1].HistoryDesc;
 	}
-	else
-	{
-		return FString();
-	}
+	return FText::GetEmpty();
 }
 
 void FReferenceViewerHistoryManager::GetAvailableHistoryMenuItems(bool bGetPrior, FMenuBuilder& MenuBuilder)
@@ -165,7 +159,7 @@ void FReferenceViewerHistoryManager::GetAvailableHistoryMenuItems(bool bGetPrior
 				for (int32 HistoryIdx = StartIndex; HistoryIdx <= EndIndex; ++HistoryIdx)
 				{
 					MenuBuilder.AddMenuEntry(
-						FText::FromString( HistoryData[HistoryIdx].HistoryDesc ),
+						HistoryData[HistoryIdx].HistoryDesc,
 						FText(),
 						FSlateIcon(),
 						FUIAction(
diff --git a/Engine/Source/Editor/ReferenceViewer/Private/HistoryManager.h b/Engine/Source/Editor/ReferenceViewer/Private/HistoryManager.h
--- a/Engine/Source/Editor/ReferenceViewer/Private/HistoryManager.h
+++ b/Engine/Source/Editor/ReferenceViewer/Private/HistoryManager.h
@@ -6,7 +6,7 @@
 struct FReferenceViewerHistoryData
 {
 	/** History Description */
-	FString HistoryDesc;
+	FText HistoryDesc;
 
 	/** The list of package names to serve as the root */
 	TArray<FName> PackageNames;
@@ -50,10 +50,10 @@ public:
 	bool CanGoBack() const;
 
 	/** Gets the description of the previous history entry */
-	FString GetBackDesc() const;
+	FText GetBackDesc() const;
 
 	/** Gets the description of the next history entry */
-	FString GetForwardDesc() const;
+	FText GetForwardDesc() const;
 
 	/**
 	 * Populates a list of menu items that can be added to a context menu
diff --git a/Engine/Source/Editor/ReferenceViewer/Private/SReferenceViewer.cpp b/Engine/Source/Editor/ReferenceViewer/Private/SReferenceViewer.cpp
--- a/Engine/Source/Editor/ReferenceViewer/Private/SReferenceViewer.cpp
+++ b/Engine/Source/Editor/ReferenceViewer/Private/SReferenceViewer.cpp
@@ -294,28 +294,22 @@ void SReferenceViewer::GraphNavigateHistoryForward()
 	ForwardClicked();
 }
 
-FString SReferenceViewer::GetHistoryBackTooltip() const
+FText SReferenceViewer::GetHistoryBackTooltip() const
 {
 	if ( HistoryManager.CanGoBack() )
 	{
-		return FText::Format( LOCTEXT("HistoryBackTooltip", "Back to {0}"), FText::FromString(HistoryManager.GetBackDesc()) ).ToString();
-	}
-	else
-	{
-		return FString();
+		return FText::Format( LOCTEXT("HistoryBackTooltip", "Back to {0}"), HistoryManager.GetBackDesc() );
 	}
+	return FText::GetEmpty();
 }
 
-FString SReferenceViewer::GetHistoryForwardTooltip() const
+FText SReferenceViewer::GetHistoryForwardTooltip() const
 {
 	if ( HistoryManager.CanGoForward() )
 	{
-		return FText::Format( LOCTEXT("HistoryForwardTooltip", "Forward to {0}"), FText::FromString(HistoryManager.GetForwardDesc()) ).ToString();
-	}
-	else
-	{
-		return FString();
+		return FText::Format( LOCTEXT("HistoryForwardTooltip", "Forward to {0}"), HistoryManager.GetForwardDesc() );
 	}
+	return FText::GetEmpty();
 }
 
 FText SReferenceViewer::GetAddressBarText() const
@@ -357,22 +351,22 @@ void SReferenceViewer::OnUpdateHistoryData(FReferenceViewerHistoryData& HistoryD
 		const TArray<FName>& CurrentGraphRootPackageNames = GraphObj->GetCurrentGraphRootPackageNames();
 		if ( CurrentGraphRootPackageNames.Num() == 1 )
 		{
-			HistoryData.HistoryDesc = FText::FromName(CurrentGraphRootPackageNames[0]).ToString();
+			HistoryData.HistoryDesc = FText::FromName(CurrentGraphRootPackageNames[0]);
 		}
 		else if ( CurrentGraphRootPackageNames.Num() > 1 )
 		{
-			HistoryData.HistoryDesc = FText::Format(LOCTEXT("AddressBarMultiplePackagesText", "{0} and {1} others"), FText::FromName(CurrentGraphRootPackageNames[0]), FText::AsNumber(CurrentGraphRootPackageNames.Num())).ToString();
+			HistoryData.HistoryDesc = FText::Format(LOCTEXT("AddressBarMultiplePackagesText", "{0} and {1} others"), FText::FromName(CurrentGraphRootPackageNames[0]), FText::AsNumber(CurrentGraphRootPackageNames.Num()));
 		}
 		else
 		{
-			HistoryData.HistoryDesc = TEXT("");
+			HistoryData.HistoryDesc = FText::GetEmpty();
 		}
 
 		HistoryData.PackageNames = CurrentGraphRootPackageNames;
 	}
 	else
 	{
-		HistoryData.HistoryDesc = TEXT("");
+		HistoryData.HistoryDesc = FText::GetEmpty();
 		HistoryData.PackageNames.Empty();
 	}
 }
diff --git a/Engine/Source/Editor/ReferenceViewer/Private/SReferenceViewer.h b/Engine/Source/Editor/ReferenceViewer/Private/SReferenceViewer.h
--- a/Engine/Source/Editor/ReferenceViewer/Private/SReferenceViewer.h
+++ b/Engine/Source/Editor/ReferenceViewer/Private/SReferenceViewer.h
@@ -45,10 +45,10 @@ private:
 	void GraphNavigateHistoryForward();
 
 	/** Gets the tool tip text for the history back button */
-	FString GetHistoryBackTooltip() const;
+	FText GetHistoryBackTooltip() const;
 
 	/** Gets the tool tip text for the history forward button */
-	FString GetHistoryForwardTooltip() const;
+	FText GetHistoryForwardTooltip() const;
 
 	/** Gets the text to be displayed in the address bar */
 	FText GetAddressBarText() const;
diff --git a/Engine/Source/Editor/SceneOutliner/Private/SSceneOutliner.cpp b/Engine/Source/Editor/SceneOutliner/Private/SSceneOutliner.cpp
--- a/Engine/Source/Editor/SceneOutliner/Private/SSceneOutliner.cpp
+++ b/Engine/Source/Editor/SceneOutliner/Private/SSceneOutliner.cpp
@@ -216,7 +216,7 @@ namespace SceneOutliner
 			SAssignNew( FilterTextBoxWidget, SSearchBox )
 			.Visibility( InInitOptions.bShowSearchBox ? EVisibility::Visible : EVisibility::Collapsed )
 			.HintText( LOCTEXT( "FilterSearch", "Search..." ) )
-			.ToolTipText( LOCTEXT("FilterSearchHint", "Type here to search (pressing enter selects the results)").ToString() )
+			.ToolTipText( LOCTEXT("FilterSearchHint", "Type here to search (pressing enter selects the results)") )
 			.OnTextChanged( this, &SSceneOutliner::OnFilterTextChanged )
 			.OnTextCommitted( this, &SSceneOutliner::OnFilterTextCommitted )
 		];
diff --git a/Engine/Source/Editor/TranslationEditor/Private/TranslationEditor.cpp b/Engine/Source/Editor/TranslationEditor/Private/TranslationEditor.cpp
--- a/Engine/Source/Editor/TranslationEditor/Private/TranslationEditor.cpp
+++ b/Engine/Source/Editor/TranslationEditor/Private/TranslationEditor.cpp
@@ -474,7 +474,7 @@ TSharedRef<SDockTab> FTranslationEditor::SpawnTab_Search(const FSpawnTabArgs& Ar
 			[
 				SAssignNew(SearchBox, SSearchBox)
 				.HintText(LOCTEXT("FilterSearch", "Search..."))
-				.ToolTipText(LOCTEXT("FilterSearchHint", "Type here to search").ToString())
+				.ToolTipText(LOCTEXT("FilterSearchHint", "Type here to search"))
 				.OnTextChanged(this, &FTranslationEditor::OnFilterTextChanged)
 				.OnTextCommitted(this, &FTranslationEditor::OnFilterTextCommitted)
 			]
diff --git a/Engine/Source/Editor/UMGEditor/Private/Customizations/DetailCustomizations.cpp b/Engine/Source/Editor/UMGEditor/Private/Customizations/DetailCustomizations.cpp
--- a/Engine/Source/Editor/UMGEditor/Private/Customizations/DetailCustomizations.cpp
+++ b/Engine/Source/Editor/UMGEditor/Private/Customizations/DetailCustomizations.cpp
@@ -37,7 +37,7 @@ public:
 		ChildSlot
 		[
 			SNew(SButton)
-			.ToolTipText(Action->TooltipDescription)
+			.ToolTipText(FText::FromString(Action->TooltipDescription))
 			.OnClicked(this, &SGraphSchemaActionButton::AddOrViewEventBinding)
 			.HAlign(HAlign_Center)
 			[
diff --git a/Engine/Source/Editor/UnrealEd/Private/Dialogs/DlgReferenceTree.cpp b/Engine/Source/Editor/UnrealEd/Private/Dialogs/DlgReferenceTree.cpp
--- a/Engine/Source/Editor/UnrealEd/Private/Dialogs/DlgReferenceTree.cpp
+++ b/Engine/Source/Editor/UnrealEd/Private/Dialogs/DlgReferenceTree.cpp
@@ -553,5 +553,5 @@ TSharedRef< ITableRow > SReferenceTree::OnGenerateRowForReferenceTree( FReferenc
 	return
 		SNew( SReferenceTreeRow, OwnerTable )
 		.Item( Item )
-		.ToolTipText( Item->Object->GetFullName() );
+		.ToolTipText( FText::FromString(Item->Object->GetFullName()) );
 }
diff --git a/Engine/Source/Editor/UnrealEd/Private/Settings/GameMapsSettingsCustomization.h b/Engine/Source/Editor/UnrealEd/Private/Settings/GameMapsSettingsCustomization.h
--- a/Engine/Source/Editor/UnrealEd/Private/Settings/GameMapsSettingsCustomization.h
+++ b/Engine/Source/Editor/UnrealEd/Private/Settings/GameMapsSettingsCustomization.h
@@ -113,7 +113,7 @@ protected:
 							[
 								MakeMapMenu(PropertyHandle)
 							]
-							.ToolTipText(LOCTEXT("AvailableMapsButtonTooltip", "Pick from the list of available maps").ToString())
+							.ToolTipText(LOCTEXT("AvailableMapsButtonTooltip", "Pick from the list of available maps"))
 					]
 			];
 	}
diff --git a/Engine/Source/Editor/WorldBrowser/Private/SWorldHierarchyItem.cpp b/Engine/Source/Editor/WorldBrowser/Private/SWorldHierarchyItem.cpp
--- a/Engine/Source/Editor/WorldBrowser/Private/SWorldHierarchyItem.cpp
+++ b/Engine/Source/Editor/WorldBrowser/Private/SWorldHierarchyItem.cpp
@@ -483,29 +483,29 @@ const FSlateBrush* SWorldHierarchyItem::GetLevelLockBrush() const
 	}
 }
 
-FString SWorldHierarchyItem::GetLevelLockToolTip() const
+FText SWorldHierarchyItem::GetLevelLockToolTip() const
 {
 	//Non-Persistent
 	if (GEngine && GEngine->bLockReadOnlyLevels)
 	{
 		if (LevelModel->IsFileReadOnly())
 		{
-			return LOCTEXT("ReadOnly_LockButtonToolTip", "Read-Only levels are locked!").ToString();
+			return LOCTEXT("ReadOnly_LockButtonToolTip", "Read-Only levels are locked!");
 		}
 	}
 
-	return LOCTEXT("LockButtonToolTip", "Toggle Level Lock").ToString();
+	return LOCTEXT("LockButtonToolTip", "Toggle Level Lock");
 }
 
-FString SWorldHierarchyItem::GetSCCStateTooltip() const
+FText SWorldHierarchyItem::GetSCCStateTooltip() const
 {
 	FSourceControlStatePtr SourceControlState = ISourceControlModule::Get().GetProvider().GetState(LevelModel->GetPackageFileName(), EStateCacheUsage::Use);
 	if(SourceControlState.IsValid())
 	{
-		return SourceControlState->GetDisplayTooltip().ToString();
+		return SourceControlState->GetDisplayTooltip();
 	}
 
-	return FString();
+	return FText::GetEmpty();
 }
 
 const FSlateBrush* SWorldHierarchyItem::GetSCCStateImage() const
diff --git a/Engine/Source/Editor/WorldBrowser/Private/SWorldHierarchyItem.h b/Engine/Source/Editor/WorldBrowser/Private/SWorldHierarchyItem.h
--- a/Engine/Source/Editor/WorldBrowser/Private/SWorldHierarchyItem.h
+++ b/Engine/Source/Editor/WorldBrowser/Private/SWorldHierarchyItem.h
@@ -159,10 +159,10 @@ private:
 
 	
 	/** */
-	FString GetLevelLockToolTip() const;
+	FText GetLevelLockToolTip() const;
 	
 	/** */
-	FString GetSCCStateTooltip() const;
+	FText GetSCCStateTooltip() const;
 	
 	/** */
 	const FSlateBrush* GetSCCStateImage() const;
diff --git a/Engine/Source/Runtime/Slate/Private/Framework/MultiBox/SClippingHorizontalBox.cpp b/Engine/Source/Runtime/Slate/Private/Framework/MultiBox/SClippingHorizontalBox.cpp
--- a/Engine/Source/Runtime/Slate/Private/Framework/MultiBox/SClippingHorizontalBox.cpp
+++ b/Engine/Source/Runtime/Slate/Private/Framework/MultiBox/SClippingHorizontalBox.cpp
@@ -139,7 +139,7 @@ void SClippingHorizontalBox::AddWrapButton()
 		.HasDownArrow( false )
 		.ButtonStyle( StyleSet, ISlateStyle::Join(StyleName, ".Button") )
 		.ContentPadding( 0 )
-		.ToolTipText( NSLOCTEXT("Slate", "ExpandToolbar", "Click to expand toolbar").ToString() )
+		.ToolTipText( NSLOCTEXT("Slate", "ExpandToolbar", "Click to expand toolbar") )
 		.OnGetMenuContent( OnWrapButtonClicked )
 		.Cursor( EMouseCursor::Default )
 		.ButtonContent()
