Removed deferred populate in SPathView::Construct, so that containing widgets (e.g. SPathPicker) are able to set selection in their Construct, without it being cleared the following frame by the populate.

diff --git a/Engine/Source/Editor/ContentBrowser/Private/SPathView.cpp b/Engine/Source/Editor/ContentBrowser/Private/SPathView.cpp
--- a/Engine/Source/Editor/ContentBrowser/Private/SPathView.cpp
+++ b/Engine/Source/Editor/ContentBrowser/Private/SPathView.cpp
@@ -28,8 +28,6 @@ SPathView::~SPathView()
 
 void SPathView::Construct( const FArguments& InArgs )
 {
-	RegisterActiveTimer(0.f, FWidgetActiveTimerDelegate::CreateSP(this, &SPathView::TriggerRepopulate));
-
 	OnPathSelected = InArgs._OnPathSelected;
 	bAllowContextMenu = InArgs._AllowContextMenu;
 	OnGetFolderContextMenu = InArgs._OnGetFolderContextMenu;
