Add Content Dialog - Remove the content browswer button, change the "Create" button to "Add New", and add a menu entry to the "Add New" menu for showing the dialog.

diff --git a/Engine/Source/Editor/ContentBrowser/Private/NewAssetContextMenu.cpp b/Engine/Source/Editor/ContentBrowser/Private/NewAssetContextMenu.cpp
--- a/Engine/Source/Editor/ContentBrowser/Private/NewAssetContextMenu.cpp
+++ b/Engine/Source/Editor/ContentBrowser/Private/NewAssetContextMenu.cpp
@@ -118,8 +118,26 @@ private:
 	TSharedPtr< FAssetThumbnail > Thumbnail;
 };
 
-void FNewAssetContextMenu::MakeContextMenu(FMenuBuilder& MenuBuilder, const FString& InPath, const FOnNewAssetRequested& InOnNewAssetRequested, const FOnNewFolderRequested& InOnNewFolderRequested)
+void FNewAssetContextMenu::MakeContextMenu( FMenuBuilder& MenuBuilder, const FString& InPath, const FOnNewAssetRequested& InOnNewAssetRequested, const FOnNewFolderRequested& InOnNewFolderRequested, const FOnImportAssetRequested& InOnImportAssetRequested, const FOnGetContentRequested& InOnGetContentRequested )
 {
+	// Get Content
+	if ( InOnGetContentRequested.IsBound() )
+	{
+		MenuBuilder.BeginSection( "ContentBrowserGetContent", LOCTEXT( "GetContentMenuHeading", "Content" ) );
+		{
+			MenuBuilder.AddMenuEntry(
+				LOCTEXT( "GetContentText", "Add feature or content pack..." ),
+				LOCTEXT( "GetContentTooltip", "Add features and content packs to the project." ),
+				FSlateIcon( FEditorStyle::GetStyleSetName(), "ContentBrowser.AddContent" ),
+				FUIAction(
+				FExecuteAction::CreateStatic( &FNewAssetContextMenu::ExecuteGetContent, InOnGetContentRequested )
+				)
+				);
+		}
+		MenuBuilder.EndSection();
+	}
+
+	// New Folder
 	if(InOnNewFolderRequested.IsBound() && GetDefault<UContentBrowserSettings>()->DisplayFolders)
 	{
 		MenuBuilder.BeginSection("ContentBrowserNewFolder", LOCTEXT("FolderMenuHeading", "Folder") );
@@ -134,31 +152,34 @@ void FNewAssetContextMenu::MakeContextMenu(FMenuBuilder& MenuBuilder, const FStr
 		MenuBuilder.EndSection(); //ContentBrowserNewFolder
 	}
 
-	// The new style menu
-	MenuBuilder.BeginSection( "ContentBrowserCreateAsset", LOCTEXT( "ImportAssetMenuHeading", "Import Asset" ) );
+	// Import
+	if (InOnImportAssetRequested.IsBound())
 	{
-		// Get the import label
-		FText ImportAssetLabel;
-		if ( !InPath.IsEmpty() )
+		MenuBuilder.BeginSection( "ContentBrowserImportAsset", LOCTEXT( "ImportAssetMenuHeading", "Import Asset" ) );
 		{
-			ImportAssetLabel = FText::Format( LOCTEXT( "ImportAsset", "Import to {0}..." ), FText::FromString( InPath ) );
-		}
-		else
-		{
-			ImportAssetLabel = LOCTEXT( "ImportAsset_NoPath", "Import" );
-		}
+			// Get the import label
+			FText ImportAssetLabel;
+			if ( !InPath.IsEmpty() )
+			{
+				ImportAssetLabel = FText::Format( LOCTEXT( "ImportAsset", "Import to {0}..." ), FText::FromString( InPath ) );
+			}
+			else
+			{
+				ImportAssetLabel = LOCTEXT( "ImportAsset_NoPath", "Import" );
+			}
 
-		MenuBuilder.AddMenuEntry(
-			ImportAssetLabel,
-			LOCTEXT( "ImportAssetTooltip", "Imports an asset from file to this folder." ),
-			FSlateIcon( FEditorStyle::GetStyleSetName(), "ContentBrowser.ImportIcon" ),
-			FUIAction(
-			FExecuteAction::CreateStatic( &FNewAssetContextMenu::ExecuteImportAsset, InPath ),
-			FCanExecuteAction::CreateStatic( &FNewAssetContextMenu::IsAssetPathSelected, InPath )
-			)
-			);
+			MenuBuilder.AddMenuEntry(
+				ImportAssetLabel,
+				LOCTEXT( "ImportAssetTooltip", "Imports an asset from file to this folder." ),
+				FSlateIcon( FEditorStyle::GetStyleSetName(), "ContentBrowser.ImportIcon" ),
+				FUIAction(
+				FExecuteAction::CreateStatic( &FNewAssetContextMenu::ExecuteImportAsset, InOnImportAssetRequested, InPath ),
+				FCanExecuteAction::CreateStatic( &FNewAssetContextMenu::IsAssetPathSelected, InPath )
+				)
+				);
+		}
+		MenuBuilder.EndSection();
 	}
-	MenuBuilder.EndSection();
 
 	MenuBuilder.BeginSection("ContentBrowserNewBasicAsset", LOCTEXT("BasicAssetsMenuHeading", "Create Basic Asset") );
 	{
@@ -303,13 +324,9 @@ bool FNewAssetContextMenu::IsAssetPathSelected(FString InPath)
 	return !InPath.IsEmpty();
 }
 
-void FNewAssetContextMenu::ExecuteImportAsset( FString InPath )
+void FNewAssetContextMenu::ExecuteImportAsset( FOnImportAssetRequested InOnInportAssetRequested, FString InPath )
 {
-	if ( ensure( !InPath.IsEmpty() ) )
-	{
-		FAssetToolsModule& AssetToolsModule = FModuleManager::Get().LoadModuleChecked<FAssetToolsModule>( "AssetTools" );
-		AssetToolsModule.Get().ImportAssets( InPath );
-	}
+	InOnInportAssetRequested.ExecuteIfBound( InPath );
 }
 
 void FNewAssetContextMenu::ExecuteNewAsset(FString InPath, TWeakObjectPtr<UClass> FactoryClass, FOnNewAssetRequested InOnNewAssetRequested)
@@ -328,4 +345,9 @@ void FNewAssetContextMenu::ExecuteNewFolder(FString InPath, FOnNewFolderRequeste
 	}
 }
 
+void FNewAssetContextMenu::ExecuteGetContent( FOnGetContentRequested InOnGetContentRequested )
+{
+	InOnGetContentRequested.ExecuteIfBound();
+}
+
 #undef LOCTEXT_NAMESPACE
diff --git a/Engine/Source/Editor/ContentBrowser/Private/NewAssetContextMenu.h b/Engine/Source/Editor/ContentBrowser/Private/NewAssetContextMenu.h
--- a/Engine/Source/Editor/ContentBrowser/Private/NewAssetContextMenu.h
+++ b/Engine/Source/Editor/ContentBrowser/Private/NewAssetContextMenu.h
@@ -7,9 +7,11 @@ class FNewAssetContextMenu
 public:
 	DECLARE_DELEGATE_TwoParams( FOnNewAssetRequested, const FString& /*SelectedPath*/, TWeakObjectPtr<UClass> /*FactoryClass*/ );
 	DECLARE_DELEGATE_OneParam( FOnNewFolderRequested, const FString& /*SelectedPath*/ );
+	DECLARE_DELEGATE_OneParam( FOnImportAssetRequested, const FString& );
+	DECLARE_DELEGATE( FOnGetContentRequested )
 
 	/** Makes the context menu widget */
-	static void MakeContextMenu(FMenuBuilder& MenuBuilder, const FString& InPath, const FOnNewAssetRequested& InOnNewAssetRequested, const FOnNewFolderRequested& InOnNewFolderRequested);
+	static void MakeContextMenu(FMenuBuilder& MenuBuilder, const FString& InPath, const FOnNewAssetRequested& InOnNewAssetRequested, const FOnNewFolderRequested& InOnNewFolderRequested, const FOnImportAssetRequested& InOnImportAssetRequested, const FOnGetContentRequested& InOnGetContentRequested);
 
 private:
 	/** Handle creating a new asset from an asset category */
@@ -19,11 +21,14 @@ private:
 	static bool IsAssetPathSelected(FString InPath);
 
 	/** Handle when the "Import" button is clicked */
-	static void ExecuteImportAsset( FString InPath );
+	static void ExecuteImportAsset(FOnImportAssetRequested InOnImportAssetRequested, FString InPath);
 
 	/** Create a new asset using the specified factory at the specified path */
 	static void ExecuteNewAsset(FString InPath, TWeakObjectPtr<UClass> FactoryClass, FOnNewAssetRequested InOnNewAssetRequested);
 
 	/** Create a new folder at the specified path */
 	static void ExecuteNewFolder(FString InPath, FOnNewFolderRequested InOnNewFolderRequested);
+
+	/** Handle when the "Get Content" button is clicked */
+	static void ExecuteGetContent( FOnGetContentRequested InOnGetContentRequested );
 };
diff --git a/Engine/Source/Editor/ContentBrowser/Private/PathContextMenu.cpp b/Engine/Source/Editor/ContentBrowser/Private/PathContextMenu.cpp
--- a/Engine/Source/Editor/ContentBrowser/Private/PathContextMenu.cpp
+++ b/Engine/Source/Editor/ContentBrowser/Private/PathContextMenu.cpp
@@ -30,6 +30,11 @@ void FPathContextMenu::SetOnNewAssetRequested(const FNewAssetContextMenu::FOnNew
 	OnNewAssetRequested = InOnNewAssetRequested;
 }
 
+void FPathContextMenu::SetOnImportAssetRequested( const FNewAssetContextMenu::FOnImportAssetRequested& InOnImportAssetRequested )
+{
+	OnImportAssetRequested = InOnImportAssetRequested;
+}
+
 void FPathContextMenu::SetOnRenameFolderRequested(const FOnRenameFolderRequested& InOnRenameFolderRequested)
 {
 	OnRenameFolderRequested = InOnRenameFolderRequested;
@@ -82,7 +87,7 @@ void FPathContextMenu::MakePathViewContextMenu(FMenuBuilder& MenuBuilder)
 		{
 			// New Asset (submenu)
 			MenuBuilder.AddSubMenu(
-				LOCTEXT( "CreateNewAsset", "Create Asset" ),
+				LOCTEXT( "CreateNewAsset", "Add New" ),
 				LOCTEXT( "CreateNewAssetTooltip", "Create a new asset." ),
 				FNewMenuDelegate::CreateRaw( this, &FPathContextMenu::MakeNewAssetSubMenu ),
 				false,
@@ -243,7 +248,7 @@ void FPathContextMenu::MakeNewAssetSubMenu(FMenuBuilder& MenuBuilder)
 	const FString& SourcesPath = GetFirstSelectedPath();
 	if ( ensure(SourcesPath.Len()) )
 	{
-		FNewAssetContextMenu::MakeContextMenu(MenuBuilder, SourcesPath, OnNewAssetRequested, FNewAssetContextMenu::FOnNewFolderRequested());
+		FNewAssetContextMenu::MakeContextMenu(MenuBuilder, SourcesPath, OnNewAssetRequested, FNewAssetContextMenu::FOnNewFolderRequested(), OnImportAssetRequested, FNewAssetContextMenu::FOnGetContentRequested());
 	}
 }
 
diff --git a/Engine/Source/Editor/ContentBrowser/Private/PathContextMenu.h b/Engine/Source/Editor/ContentBrowser/Private/PathContextMenu.h
--- a/Engine/Source/Editor/ContentBrowser/Private/PathContextMenu.h
+++ b/Engine/Source/Editor/ContentBrowser/Private/PathContextMenu.h
@@ -13,6 +13,9 @@ public:
 	/** Sets the handler for when new assets are requested */
 	void SetOnNewAssetRequested(const FNewAssetContextMenu::FOnNewAssetRequested& InOnNewAssetRequested);
 
+	/** Sets the handler for when importing an asset is requested */
+	void SetOnImportAssetRequested( const FNewAssetContextMenu::FOnImportAssetRequested& InOnImportAssetRequested );
+
 	/** Delegate for when the context menu requests a rename of a folder */
 	DECLARE_DELEGATE_OneParam(FOnRenameFolderRequested, const FString& /*FolderToRename*/);
 	void SetOnRenameFolderRequested(const FOnRenameFolderRequested& InOnRenameFolderRequested);
@@ -125,6 +128,7 @@ private:
 	TArray<FString> SelectedPaths;
 	TWeakPtr<SWidget> ParentContent;
 	FNewAssetContextMenu::FOnNewAssetRequested OnNewAssetRequested;
+	FNewAssetContextMenu::FOnImportAssetRequested OnImportAssetRequested;
 	FOnRenameFolderRequested OnRenameFolderRequested;
 	FOnFolderDeleted OnFolderDeleted;
 
diff --git a/Engine/Source/Editor/ContentBrowser/Private/SContentBrowser.cpp b/Engine/Source/Editor/ContentBrowser/Private/SContentBrowser.cpp
--- a/Engine/Source/Editor/ContentBrowser/Private/SContentBrowser.cpp
+++ b/Engine/Source/Editor/ContentBrowser/Private/SContentBrowser.cpp
@@ -57,6 +57,7 @@ void SContentBrowser::Construct( const FArguments& InArgs, const FName& InInstan
 
 	PathContextMenu = MakeShareable(new FPathContextMenu( AsShared() ));
 	PathContextMenu->SetOnNewAssetRequested( FNewAssetContextMenu::FOnNewAssetRequested::CreateSP(this, &SContentBrowser::NewAssetRequested) );
+	PathContextMenu->SetOnImportAssetRequested(FNewAssetContextMenu::FOnImportAssetRequested::CreateSP(this, &SContentBrowser::ImportAsset));
 	PathContextMenu->SetOnRenameFolderRequested(FPathContextMenu::FOnRenameFolderRequested::CreateSP(this, &SContentBrowser::OnRenameFolderRequested));
 	PathContextMenu->SetOnFolderDeleted(FPathContextMenu::FOnFolderDeleted::CreateSP(this, &SContentBrowser::OnOpenedFolderDeleted));
 
@@ -133,7 +134,7 @@ void SContentBrowser::Construct( const FArguments& InArgs, const FName& InInstan
 								[
 									SNew( STextBlock )
 									.TextStyle( FEditorStyle::Get(), "ContentBrowser.TopBar.Font" )
-									.Text( LOCTEXT( "NewButton", "Create" ) )
+									.Text( LOCTEXT( "NewButton", "Add New" ) )
 								]
 							]
 						]
@@ -177,42 +178,6 @@ void SContentBrowser::Construct( const FArguments& InArgs, const FName& InInstan
 							]
 						]
 
-						// Get content
-						+ SHorizontalBox::Slot()
-						.VAlign(VAlign_Center)
-						.HAlign(HAlign_Left)
-						[
-							SNew(SButton)
-							.ButtonStyle(FEditorStyle::Get(), "ToggleButton")
-							.ToolTipText(LOCTEXT("AddContentTooltip", "Get more content."))
-							.ContentPadding(0)
-							.Visibility(UEditorExperimentalSettings::StaticClass()->GetDefaultObject<UEditorExperimentalSettings>()->bGetFeatureContent == false ? EVisibility::Collapsed : EVisibility::Visible)
-							.OnClicked(this, &SContentBrowser::OnAddContentClicked)
-							[
-								SNew(SHorizontalBox)
-
-								// Get Content Icon
-								+ SHorizontalBox::Slot()
-								.AutoWidth()
-								.VAlign(VAlign_Center)
-								[
-									SNew(SImage)
-									.Image(FEditorStyle::GetBrush("ContentBrowser.AddContent"))
-								]
-
-								// Get Content Text
-								+ SHorizontalBox::Slot()
-								.AutoWidth()
-								.VAlign(VAlign_Center)
-								.Padding(0, 0, 2, 0)
-								[
-									SNew(STextBlock)
-									.TextStyle(FEditorStyle::Get(), "ContentBrowser.TopBar.Font")
-									.Text(LOCTEXT("GetContent", "Get Content"))
-								]
-							]
-						]
-
 						// Save
 						+ SHorizontalBox::Slot()
 						.FillWidth(1.0f)
@@ -675,15 +640,17 @@ FText SContentBrowser::GetImportTooltipText() const
 
 FReply SContentBrowser::HandleImportClicked()
 {
-	FString CurrentPath = GetCurrentPath();
+	ImportAsset( GetCurrentPath() );
+	return FReply::Handled();
+}
 
-	if ( ensure( !CurrentPath.IsEmpty() ) )
+void SContentBrowser::ImportAsset( const FString& InPath )
+{
+	if ( ensure( !InPath.IsEmpty() ) )
 	{
 		FAssetToolsModule& AssetToolsModule = FModuleManager::Get().LoadModuleChecked<FAssetToolsModule>( "AssetTools" );
-		AssetToolsModule.Get().ImportAssets( CurrentPath );
+		AssetToolsModule.Get().ImportAssets( InPath );
 	}
-
-	return FReply::Handled();
 }
 
 void SContentBrowser::SyncToAssets( const TArray<FAssetData>& AssetDataList, const bool bAllowImplicitSync )
@@ -1312,7 +1279,9 @@ TSharedRef<SWidget> SContentBrowser::MakeCreateAssetContextMenu()
 		MenuBuilder, 
 		CurrentPath, 
 		FNewAssetContextMenu::FOnNewAssetRequested::CreateSP(this, &SContentBrowser::NewAssetRequested),
-		OnNewFolderRequested);
+		OnNewFolderRequested,
+		FNewAssetContextMenu::FOnImportAssetRequested(),
+		FNewAssetContextMenu::FOnGetContentRequested::CreateSP(this, &SContentBrowser::OnAddContentRequested));
 
 	FDisplayMetrics DisplayMetrics;
 	FSlateApplication::Get().GetDisplayMetrics( DisplayMetrics );
@@ -1360,14 +1329,12 @@ FReply SContentBrowser::OnSaveClicked()
 	return FReply::Handled();
 }
 
-FReply SContentBrowser::OnAddContentClicked()
+void SContentBrowser::OnAddContentRequested()
 {
 	IAddContentDialogModule& AddContentDialogModule = FModuleManager::LoadModuleChecked<IAddContentDialogModule>("AddContentDialog");
 	FWidgetPath WidgetPath;
 	FSlateApplication::Get().GeneratePathToWidgetChecked(AsShared(), WidgetPath);
 	AddContentDialogModule.ShowDialog(WidgetPath.GetWindow());
-
-	return FReply::Handled();
 }
 
 void SContentBrowser::OnAssetSelectionChanged(const FAssetData& SelectedAsset)
diff --git a/Engine/Source/Editor/ContentBrowser/Private/SContentBrowser.h b/Engine/Source/Editor/ContentBrowser/Private/SContentBrowser.h
--- a/Engine/Source/Editor/ContentBrowser/Private/SContentBrowser.h
+++ b/Engine/Source/Editor/ContentBrowser/Private/SContentBrowser.h
@@ -164,11 +164,14 @@ private:
 	/** Imports a new piece of content. */
 	FReply HandleImportClicked();
 
+	/** Imports a new piece of content to a specific path. */
+	void ImportAsset(const FString& InPath);
+
 	/** Saves dirty content. */
 	FReply OnSaveClicked();
 
 	/** Opens the add content dialog. */
-	FReply OnAddContentClicked();
+	void OnAddContentRequested();
 
 	/** Handler for when the selection set in the asset view has changed. */
 	void OnAssetSelectionChanged(const FAssetData& SelectedAsset);
diff --git a/Engine/Source/Editor/UnrealEd/Classes/Settings/EditorExperimentalSettings.h b/Engine/Source/Editor/UnrealEd/Classes/Settings/EditorExperimentalSettings.h
--- a/Engine/Source/Editor/UnrealEd/Classes/Settings/EditorExperimentalSettings.h
+++ b/Engine/Source/Editor/UnrealEd/Classes/Settings/EditorExperimentalSettings.h
@@ -84,10 +84,6 @@ public:
 	/** Enables Environment Queries editor */
 	UPROPERTY(EditAnywhere, config, Category = AI, meta = (DisplayName = "Environment Querying System"))
 	bool bEQSEditor;
-
-	/** Enables GetContent button in asset browser */
-	UPROPERTY(EditAnywhere, config, Category = Tools, meta = (DisplayName = "Get Feature Content (Requires content browser reopening)"))
-	bool bGetFeatureContent;
 	
 	/**
 	 * Returns an event delegate that is executed when a setting has changed.
