Content Browser - Add the import option back to the right click dialog which is shown when not clicking on a folder or asset.

Jira #UE-7679

diff --git a/Engine/Source/Editor/ContentBrowser/Private/SContentBrowser.cpp b/Engine/Source/Editor/ContentBrowser/Private/SContentBrowser.cpp
--- a/Engine/Source/Editor/ContentBrowser/Private/SContentBrowser.cpp
+++ b/Engine/Source/Editor/ContentBrowser/Private/SContentBrowser.cpp
@@ -109,7 +109,7 @@ void SContentBrowser::Construct( const FArguments& InArgs, const FName& InInstan
 							.ComboButtonStyle( FEditorStyle::Get(), "ContentBrowser.NewAsset.Style" )
 							.ForegroundColor(FLinearColor::White)
 							.ContentPadding(0)
-							.OnGetMenuContent( this, &SContentBrowser::MakeCreateAssetContextMenu )
+							.OnGetMenuContent_Lambda( [this]{ return MakeCreateAssetContextMenu( true, false ); } )
 							.ToolTipText( this, &SContentBrowser::GetNewAssetToolTipText )
 							.IsEnabled( this, &SContentBrowser::IsAssetPathSelected )
 							.AddMetaData<FTagMetaData>(FTagMetaData(TEXT("ContentBrowserNewAsset")))
@@ -1248,7 +1248,7 @@ FString SContentBrowser::GetCurrentPath() const
 	return CurrentPath;
 }
 
-TSharedRef<SWidget> SContentBrowser::MakeCreateAssetContextMenu()
+TSharedRef<SWidget> SContentBrowser::MakeCreateAssetContextMenu(bool bShowGetContent, bool bShowImport)
 {
 	FString CurrentPath = GetCurrentPath();
 
@@ -1280,8 +1280,12 @@ TSharedRef<SWidget> SContentBrowser::MakeCreateAssetContextMenu()
 		CurrentPath, 
 		FNewAssetContextMenu::FOnNewAssetRequested::CreateSP(this, &SContentBrowser::NewAssetRequested),
 		OnNewFolderRequested,
-		FNewAssetContextMenu::FOnImportAssetRequested(),
-		FNewAssetContextMenu::FOnGetContentRequested::CreateSP(this, &SContentBrowser::OnAddContentRequested));
+		bShowImport
+			? FNewAssetContextMenu::FOnImportAssetRequested::CreateSP(this, &SContentBrowser::ImportAsset)
+			: FNewAssetContextMenu::FOnImportAssetRequested(),
+		bShowGetContent
+			? FNewAssetContextMenu::FOnGetContentRequested::CreateSP(this, &SContentBrowser::OnAddContentRequested)
+			: FNewAssetContextMenu::FOnGetContentRequested());
 
 	FDisplayMetrics DisplayMetrics;
 	FSlateApplication::Get().GetDisplayMetrics( DisplayMetrics );
@@ -1434,7 +1438,7 @@ TSharedPtr<SWidget> SContentBrowser::OnGetAssetContextMenu(const TArray<FAssetDa
 
 	if ( SelectedAssets.Num() == 0 )
 	{
-		return MakeCreateAssetContextMenu();
+		return MakeCreateAssetContextMenu( false, true );
 	}
 	else
 	{
diff --git a/Engine/Source/Editor/ContentBrowser/Private/SContentBrowser.h b/Engine/Source/Editor/ContentBrowser/Private/SContentBrowser.h
--- a/Engine/Source/Editor/ContentBrowser/Private/SContentBrowser.h
+++ b/Engine/Source/Editor/ContentBrowser/Private/SContentBrowser.h
@@ -150,7 +150,7 @@ private:
 	TSharedRef<SWidget> GetPathPickerContent();
 
 	/** Handle creating a context menu to generate a new asset */
-	TSharedRef<SWidget> MakeCreateAssetContextMenu();
+	TSharedRef<SWidget> MakeCreateAssetContextMenu(bool bShowGetContent, bool bShowImport);
 
 	/** Gets the tool tip for the new asset button */
 	FText GetNewAssetToolTipText() const;
