Added support for showing game and engine classes in the Content Browser

UE-7184 - Show game c++ classes and engine c++ classes in the content browser

The Content Browser now has extra entries for "Game C++ Classes" and "Engine C++ Classes" (with "Game" and "Engine" being renamed to "Game Content" and "Engine Content" respectively). These new folders serve as hosts for the list of available C++ modules, with each module internally mirroring the folder structure on disk.

For example:
- Game C++ Classes
    - ShooterGame
        - Classes
            - Bots
                - ShooterBot
                - ShooterAIController
                - [...]
            - [...]

The Content Browser allows you to navigate and search these classes like you can with assets, and provides convenient access to either edit an existing class, or create a new class (either within a selected folder, or derived from a selected class).

As the Content Browser only shows you known UClass types, any new classes need to be compiled into a loaded module before they will appear. This means that adding a new class will now automatically hot-reload your target module. Should you prefer to handle building and loading your modules manually, you can disable the automatic hot-reload via "Editor Settings" -> "Miscellaneous" -> "Hot Reload" -> "Automatically Hot Reload New Classes" (see UEditorUserSettings::bAutomaticallyHotReloadNewClasses).

diff --git a/Engine/Config/BaseEditorUserSettings.ini b/Engine/Config/BaseEditorUserSettings.ini
--- a/Engine/Config/BaseEditorUserSettings.ini
+++ b/Engine/Config/BaseEditorUserSettings.ini
@@ -58,6 +58,7 @@ HoverHighlightIntensity=0.000000
 bDisplayUIExtensionPoints=False
 bUseCurvesForDistributions=False
 bAutoloadCheckedOutPackages=False
+bAutomaticallyHotReloadNewClasses=True
 
 [/Script/EditorStyle.EditorStyleSettings]
 ; Applies a color vision deficiency filter to the entire editor
diff --git a/Engine/Source/Developer/AssetTools/AssetTools.Build.cs b/Engine/Source/Developer/AssetTools/AssetTools.Build.cs
--- a/Engine/Source/Developer/AssetTools/AssetTools.Build.cs
+++ b/Engine/Source/Developer/AssetTools/AssetTools.Build.cs
@@ -35,6 +35,7 @@ public class AssetTools : ModuleRules
                 "CurveAssetEditor",
 				"DesktopPlatform",
 				"EditorWidgets",
+				"GameProjectGeneration",
 				"Kismet",
 				"MainFrame",
 				"MaterialEditor",
@@ -57,6 +58,7 @@ public class AssetTools : ModuleRules
 				"DesktopPlatform",
 				"DestructibleMeshEditor",
 				"EditorWidgets",
+				"GameProjectGeneration",
 				"PropertyEditor",
 				"MainFrame",
 				"MaterialEditor",
diff --git a/Engine/Source/Developer/AssetTools/Private/AssetTools.cpp b/Engine/Source/Developer/AssetTools/Private/AssetTools.cpp
--- a/Engine/Source/Developer/AssetTools/Private/AssetTools.cpp
+++ b/Engine/Source/Developer/AssetTools/Private/AssetTools.cpp
@@ -45,6 +45,7 @@ FAssetTools::FAssetTools()
 	RegisterAssetTypeActions( MakeShareable(new FAssetTypeActions_DialogueVoice) );
 	RegisterAssetTypeActions( MakeShareable(new FAssetTypeActions_DialogueWave) );
 	RegisterAssetTypeActions( MakeShareable(new FAssetTypeActions_Enum) );
+	RegisterAssetTypeActions( MakeShareable(new FAssetTypeActions_Class) );
 	RegisterAssetTypeActions( MakeShareable(new FAssetTypeActions_Struct) );
 	RegisterAssetTypeActions( MakeShareable(new FAssetTypeActions_Font) );
 	RegisterAssetTypeActions( MakeShareable(new FAssetTypeActions_ForceFeedbackEffect) );
diff --git a/Engine/Source/Developer/AssetTools/Private/AssetToolsPrivatePCH.h b/Engine/Source/Developer/AssetTools/Private/AssetToolsPrivatePCH.h
--- a/Engine/Source/Developer/AssetTools/Private/AssetToolsPrivatePCH.h
+++ b/Engine/Source/Developer/AssetTools/Private/AssetToolsPrivatePCH.h
@@ -43,6 +43,7 @@ DECLARE_LOG_CATEGORY_EXTERN(LogAssetTools, Log, All);
 #include "AssetTypeActions/AssetTypeActions_DialogueVoice.h"
 #include "AssetTypeActions/AssetTypeActions_DialogueWave.h"
 #include "AssetTypeActions/AssetTypeActions_Enum.h"
+#include "AssetTypeActions/AssetTypeActions_Class.h"
 #include "AssetTypeActions/AssetTypeActions_Struct.h"
 #include "AssetTypeActions/AssetTypeActions_Font.h"
 #include "AssetTypeActions/AssetTypeActions_ForceFeedbackEffect.h"
diff --git a/Engine/Source/Developer/AssetTools/Private/AssetTypeActions/AssetTypeActions_Class.cpp b/Engine/Source/Developer/AssetTools/Private/AssetTypeActions/AssetTypeActions_Class.cpp
new file mode 100644
--- /dev/null
+++ b/Engine/Source/Developer/AssetTools/Private/AssetTypeActions/AssetTypeActions_Class.cpp
@@ -0,0 +1,90 @@
+// Copyright 1998-2015 Epic Games, Inc. All Rights Reserved.
+
+#include "AssetToolsPrivatePCH.h"
+#include "SourceCodeNavigation.h"
+#include "GameProjectGenerationModule.h"
+
+#define LOCTEXT_NAMESPACE "AssetTypeActions"
+
+bool FAssetTypeActions_Class::HasActions(const TArray<UObject*>& InObjects) const
+{
+	return true;
+}
+
+void FAssetTypeActions_Class::GetActions(const TArray<UObject*>& InObjects, FMenuBuilder& MenuBuilder)
+{
+	UClass *const BaseClass = (InObjects.Num() == 1) ? Cast<UClass>(InObjects[0]) : nullptr;
+
+	// Only allow the New class option if we have a base class that we can actually derive from in one of our project modules
+	FGameProjectGenerationModule& GameProjectGenerationModule = FGameProjectGenerationModule::Get();
+	TArray<FModuleContextInfo> ProjectModules = GameProjectGenerationModule.GetCurrentProjectModules();
+	const bool bIsValidBaseClass = BaseClass && GameProjectGenerationModule.IsValidBaseClassForCreation(BaseClass, ProjectModules);
+
+	auto CreateCreateDerivedClass = [BaseClass]()
+	{
+		// Work out where the header file for the current class is, as we'll use that path as the default for the new class
+		FString BaseClassPath;
+		if(FSourceCodeNavigation::FindClassHeaderPath(BaseClass, BaseClassPath))
+		{
+			// Strip off the actual filename as we only need the path
+			BaseClassPath = FPaths::GetPath(BaseClassPath);
+		}
+
+		FGameProjectGenerationModule::Get().OpenAddCodeToProjectDialog(BaseClass, BaseClassPath, FGlobalTabmanager::Get()->GetRootWindow());
+	};
+
+	auto CanCreateDerivedClass = [bIsValidBaseClass]() -> bool
+	{
+		return bIsValidBaseClass;
+	};
+
+	FText NewDerivedClassClassLabel;
+	FText NewDerivedClassClassToolTip;
+	if(InObjects.Num() == 1)
+	{
+		const FText BaseClassName = FText::FromName(BaseClass->GetFName());
+		NewDerivedClassClassLabel = FText::Format(LOCTEXT("Class_NewDerivedClassLabel_CreateFrom", "New Class from {0}"), BaseClassName);
+		if(bIsValidBaseClass)
+		{
+			NewDerivedClassClassToolTip = FText::Format(LOCTEXT("Class_NewDerivedClassTooltip_CreateFrom", "Create a new class deriving from {0}."), BaseClassName);
+		}
+		else
+		{
+			NewDerivedClassClassToolTip = FText::Format(LOCTEXT("Class_NewDerivedClassTooltip_InvalidClass", "Cannot create a new class deriving from {0}."), BaseClassName);
+		}
+	}
+	else
+	{
+		NewDerivedClassClassLabel = LOCTEXT("Class_NewDerivedClassLabel_InvalidNumberOfBases", "New Class from...");
+		NewDerivedClassClassToolTip = LOCTEXT("Class_NewDerivedClassTooltip_InvalidNumberOfBases", "Can only create a derived class when there is a single base class selected.");
+	}
+
+	MenuBuilder.AddMenuEntry(
+		NewDerivedClassClassLabel,
+		NewDerivedClassClassToolTip,
+		FSlateIcon(FEditorStyle::GetStyleSetName(), "MainFrame.AddCodeToProject"),
+		FUIAction(
+			FExecuteAction::CreateLambda(CreateCreateDerivedClass),
+			FCanExecuteAction::CreateLambda(CanCreateDerivedClass)
+			)
+		);
+}
+
+void FAssetTypeActions_Class::OpenAssetEditor( const TArray<UObject*>& InObjects, TSharedPtr<class IToolkitHost> EditWithinLevelEditor )
+{
+	for(UObject* Object : InObjects)
+	{
+		UClass* const Class = Cast<UClass>(Object);
+		if(Class)
+		{
+			FString ClassHeaderPath;
+			if(FSourceCodeNavigation::FindClassHeaderPath(Class, ClassHeaderPath))
+			{
+				const FString AbsoluteHeaderPath = IFileManager::Get().ConvertToAbsolutePathForExternalAppForRead(*ClassHeaderPath);
+				FSourceCodeNavigation::OpenSourceFile(AbsoluteHeaderPath);
+			}
+		}
+	}
+}
+
+#undef LOCTEXT_NAMESPACE
diff --git a/Engine/Source/Developer/AssetTools/Private/AssetTypeActions/AssetTypeActions_Class.h b/Engine/Source/Developer/AssetTools/Private/AssetTypeActions/AssetTypeActions_Class.h
new file mode 100644
--- /dev/null
+++ b/Engine/Source/Developer/AssetTools/Private/AssetTypeActions/AssetTypeActions_Class.h
@@ -0,0 +1,20 @@
+// Copyright 1998-2015 Epic Games, Inc. All Rights Reserved.
+
+#pragma once
+
+#include "AssetTypeActions_Base.h"
+
+class FAssetTypeActions_Class : public FAssetTypeActions_Base
+{
+public:
+	// IAssetTypeActions Implementation
+	virtual FText GetName() const override { return NSLOCTEXT("AssetTypeActions", "AssetTypeActions_Class", "Class"); }
+	virtual FColor GetTypeColor() const override { return FColor::White; }
+	virtual UClass* GetSupportedClass() const override { return UClass::StaticClass(); }
+	virtual uint32 GetCategories() override { return EAssetTypeCategories::None; }
+	
+	virtual bool HasActions(const TArray<UObject*>& InObjects) const override;
+	virtual void GetActions(const TArray<UObject*>& InObjects, FMenuBuilder& MenuBuilder) override;
+
+	virtual void OpenAssetEditor( const TArray<UObject*>& InObjects, TSharedPtr<class IToolkitHost> EditWithinLevelEditor = TSharedPtr<IToolkitHost>() ) override;
+};
diff --git a/Engine/Source/Editor/ClassViewer/Private/SClassViewer.cpp b/Engine/Source/Editor/ClassViewer/Private/SClassViewer.cpp
--- a/Engine/Source/Editor/ClassViewer/Private/SClassViewer.cpp
+++ b/Engine/Source/Editor/ClassViewer/Private/SClassViewer.cpp
@@ -1132,25 +1132,7 @@ namespace ClassViewer
 		 */
 		static void OpenCreateCPlusPlusClassWizard(UClass* InCreationClass)
 		{
-			TSharedRef<SWindow> AddCodeWindow =
-				SNew(SWindow)
-				.Title(LOCTEXT( "AddCodeWindowHeader", "Add Code"))
-				.ClientSize( FVector2D(1280, 720) )
-				.SizingRule( ESizingRule::FixedSize )
-				.SupportsMinimize(false) 
-				.SupportsMaximize(false);
-
-			AddCodeWindow->SetContent( FGameProjectGenerationModule::Get().CreateNewClassDialog(InCreationClass) );
-
-			TSharedPtr<SWindow> ParentWindow = FGlobalTabmanager::Get()->GetRootWindow();
-			if (ParentWindow.IsValid())
-			{
-				FSlateApplication::Get().AddWindowAsNativeChild(AddCodeWindow, ParentWindow.ToSharedRef());
-			}
-			else
-			{
-				FSlateApplication::Get().AddWindow(AddCodeWindow);
-			}
+			FGameProjectGenerationModule::Get().OpenAddCodeToProjectDialog(InCreationClass, FString(), FGlobalTabmanager::Get()->GetRootWindow());
 		}
 
 		/**
diff --git a/Engine/Source/Editor/ContentBrowser/ContentBrowser.Build.cs b/Engine/Source/Editor/ContentBrowser/ContentBrowser.Build.cs
--- a/Engine/Source/Editor/ContentBrowser/ContentBrowser.Build.cs
+++ b/Engine/Source/Editor/ContentBrowser/ContentBrowser.Build.cs
@@ -26,6 +26,7 @@ public class ContentBrowser : ModuleRules
 				"CoreUObject",
                 "InputCore",
 				"Engine",
+				"GameProjectGeneration",
 				"Slate",
 				"SlateCore",
                 "EditorStyle",
diff --git a/Engine/Source/Editor/ContentBrowser/Private/AssetContextMenu.cpp b/Engine/Source/Editor/ContentBrowser/Private/AssetContextMenu.cpp
--- a/Engine/Source/Editor/ContentBrowser/Private/AssetContextMenu.cpp
+++ b/Engine/Source/Editor/ContentBrowser/Private/AssetContextMenu.cpp
@@ -40,6 +40,7 @@
 FAssetContextMenu::FAssetContextMenu(const TWeakPtr<SAssetView>& InAssetView)
 	: AssetView(InAssetView)
 	, bAtLeastOneNonRedirectorSelected(false)
+	, bAtLeastOneClassSelected(false)
 	, bCanExecuteSCCMerge(false)
 	, bCanExecuteSCCCheckOut(false)
 	, bCanExecuteSCCOpenForAdd(false)
@@ -243,7 +244,15 @@ bool FAssetContextMenu::AddCommonMenuOptions(FMenuBuilder& MenuBuilder)
 	MenuBuilder.AddSubMenu(
 		LOCTEXT("AssetActionsSubMenuLabel", "Asset Actions"),
 		LOCTEXT("AssetActionsSubMenuToolTip", "Other asset actions"),
-		FNewMenuDelegate::CreateSP(this, &FAssetContextMenu::MakeAssetActionsSubMenu), false, FSlateIcon(FEditorStyle::GetStyleSetName(), "ContentBrowser.AssetActions")
+		FNewMenuDelegate::CreateSP(this, &FAssetContextMenu::MakeAssetActionsSubMenu),
+		FUIAction(
+			FExecuteAction(),
+			FCanExecuteAction::CreateSP( this, &FAssetContextMenu::CanExecuteAssetActions )
+			),
+		NAME_None,
+		EUserInterfaceActionType::Button,
+		false, 
+		FSlateIcon(FEditorStyle::GetStyleSetName(), "ContentBrowser.AssetActions")
 		);
 	}
 	MenuBuilder.EndSection();
@@ -414,6 +423,11 @@ void FAssetContextMenu::MakeAssetActionsSubMenu(FMenuBuilder& MenuBuilder)
 	MenuBuilder.EndSection();
 }
 
+bool FAssetContextMenu::CanExecuteAssetActions() const
+{
+	return !bAtLeastOneClassSelected;
+}
+
 bool FAssetContextMenu::AddReferenceMenuOptions(FMenuBuilder& MenuBuilder)
 {
 	MenuBuilder.BeginSection("AssetContextReferences"/*, LOCTEXT("ReferencesMenuHeading", "References")*/);
@@ -598,6 +612,12 @@ bool FAssetContextMenu::AddSourceControlMenuOptions(FMenuBuilder& MenuBuilder)
 			LOCTEXT("SourceControlSubMenuLabel", "Source Control"),
 			LOCTEXT("SourceControlSubMenuToolTip", "Source control actions."),
 			FNewMenuDelegate::CreateSP(this, &FAssetContextMenu::FillSourceControlSubMenu),
+			FUIAction(
+				FExecuteAction(),
+				FCanExecuteAction::CreateSP( this, &FAssetContextMenu::CanExecuteSourceControlActions )
+				),
+			NAME_None,
+			EUserInterfaceActionType::Button,
 			false,
 			FSlateIcon(FEditorStyle::GetStyleSetName(), "SourceControl.StatusIcon.On")
 			);
@@ -610,9 +630,9 @@ bool FAssetContextMenu::AddSourceControlMenuOptions(FMenuBuilder& MenuBuilder)
 			FSlateIcon(FEditorStyle::GetStyleSetName(), "SourceControl.StatusIcon.Unknown"),
 			FUIAction(
 				FExecuteAction::CreateSP( this, &FAssetContextMenu::ExecuteEnableSourceControl ),
-				FCanExecuteAction()
+				FCanExecuteAction::CreateSP( this, &FAssetContextMenu::CanExecuteSourceControlActions )
 				)
-			);		
+			);
 	}
 
 	// Diff selected
@@ -745,6 +765,11 @@ void FAssetContextMenu::FillSourceControlSubMenu(FMenuBuilder& MenuBuilder)
 	}
 }
 
+bool FAssetContextMenu::CanExecuteSourceControlActions() const
+{
+	return !bAtLeastOneClassSelected;
+}
+
 bool FAssetContextMenu::AddCollectionMenuOptions(FMenuBuilder& MenuBuilder)
 {
 	// "Remove from collection" (only display option if exactly one collection is selected)
@@ -1617,15 +1642,21 @@ bool FAssetContextMenu::CanExecutePropertyMatrix() const
 
 bool FAssetContextMenu::CanExecuteDuplicate() const
 {
-	return bAtLeastOneNonRedirectorSelected;
+	return bAtLeastOneNonRedirectorSelected && !bAtLeastOneClassSelected;
 }
 
 bool FAssetContextMenu::CanExecuteRename() const
 {
 	TArray< FAssetData > AssetViewSelectedAssets = AssetView.Pin()->GetSelectedAssets();
 	TArray< FString > SelectedFolders = AssetView.Pin()->GetSelectedFolders();
-	const bool bOneAssetSelected = AssetViewSelectedAssets.Num() == 1 && SelectedFolders.Num() == 0 && AssetViewSelectedAssets[0].AssetClass != UObjectRedirector::StaticClass()->GetFName();
-	const bool bOneFolderSelected = AssetViewSelectedAssets.Num() == 0 && SelectedFolders.Num() == 1;
+
+	const bool bOneAssetSelected = AssetViewSelectedAssets.Num() == 1 && SelectedFolders.Num() == 0		// A single asset
+		&& AssetViewSelectedAssets[0].AssetClass != UObjectRedirector::StaticClass()->GetFName()		// Which isn't a redirector
+		&& AssetViewSelectedAssets[0].AssetClass != NAME_Class;											// And isn't a class
+
+	const bool bOneFolderSelected = AssetViewSelectedAssets.Num() == 0 && SelectedFolders.Num() == 1	// A single folder
+		&& !ContentBrowserUtils::IsClassPath(SelectedFolders[0]);										// Which doesn't belong to a class path
+	
 	return (bOneAssetSelected || bOneFolderSelected) && !AssetView.Pin()->IsThumbnailEditMode();
 }
 
@@ -1633,7 +1664,15 @@ bool FAssetContextMenu::CanExecuteDelete() const
 {
 	TArray< FAssetData > AssetViewSelectedAssets = AssetView.Pin()->GetSelectedAssets();
 	TArray< FString > SelectedFolders = AssetView.Pin()->GetSelectedFolders();
-	return AssetViewSelectedAssets.Num() > 0 || SelectedFolders.Num() > 0;
+
+	int32 NumAssetItems, NumClassItems;
+	ContentBrowserUtils::CountItemTypes(AssetViewSelectedAssets, NumAssetItems, NumClassItems);
+
+	int32 NumAssetPaths, NumClassPaths;
+	ContentBrowserUtils::CountPathTypes(SelectedFolders, NumAssetPaths, NumClassPaths);
+
+	// We can't delete classes, or folders containing classes
+	return (NumAssetItems > 0 && NumClassItems == 0) || (NumAssetPaths > 0 && NumClassPaths == 0);
 }
 
 bool FAssetContextMenu::CanExecuteRemoveFromCollection() const 
@@ -1736,6 +1775,11 @@ bool FAssetContextMenu::CanExecuteConsolidate() const
 
 bool FAssetContextMenu::CanExecuteSaveAsset() const
 {
+	if ( bAtLeastOneClassSelected )
+	{
+		return false;
+	}
+
 	TArray<UPackage*> Packages;
 	GetSelectedPackages(Packages);
 
@@ -1754,7 +1798,7 @@ bool FAssetContextMenu::CanExecuteSaveAsset() const
 bool FAssetContextMenu::CanExecuteDiffSelected() const
 {
 	bool bCanDiffSelected = false;
-	if (SelectedAssets.Num() == 2)
+	if (SelectedAssets.Num() == 2 && !bAtLeastOneClassSelected)
 	{
 		FAssetData const& FirstSelection = SelectedAssets[0];
 		FAssetData const& SecondSelection = SelectedAssets[1];
@@ -1793,6 +1837,7 @@ bool FAssetContextMenu::CanClearCustomThumbnails() const
 void FAssetContextMenu::CacheCanExecuteVars()
 {
 	bAtLeastOneNonRedirectorSelected = false;
+	bAtLeastOneClassSelected = false;
 	bCanExecuteSCCMerge = false;
 	bCanExecuteSCCCheckOut = false;
 	bCanExecuteSCCOpenForAdd = false;
@@ -1814,6 +1859,8 @@ void FAssetContextMenu::CacheCanExecuteVars()
 			bAtLeastOneNonRedirectorSelected = true;
 		}
 
+		bAtLeastOneClassSelected |= AssetData.AssetClass == NAME_Class;
+
 		ISourceControlProvider& SourceControlProvider = ISourceControlModule::Get().GetProvider();
 		if ( ISourceControlModule::Get().IsEnabled() )
 		{
@@ -1854,6 +1901,7 @@ void FAssetContextMenu::CacheCanExecuteVars()
 		}
 
 		if ( bAtLeastOneNonRedirectorSelected
+			&& bAtLeastOneClassSelected
 			&& bCanExecuteSCCMerge
 			&& bCanExecuteSCCCheckOut
 			&& bCanExecuteSCCOpenForAdd
diff --git a/Engine/Source/Editor/ContentBrowser/Private/AssetContextMenu.h b/Engine/Source/Editor/ContentBrowser/Private/AssetContextMenu.h
--- a/Engine/Source/Editor/ContentBrowser/Private/AssetContextMenu.h
+++ b/Engine/Source/Editor/ContentBrowser/Private/AssetContextMenu.h
@@ -89,6 +89,9 @@ private:
 	/** Adds Asset Actions sub-menu to a menu builder. */
 	void MakeAssetActionsSubMenu(FMenuBuilder& MenuBuilder);
 
+	/** Handler to check to see if "Asset Actions" are allowed */
+	bool CanExecuteAssetActions() const;
+
 	/** Adds asset reference menu options to a menu builder. Returns true if any options were added. */
 	bool AddReferenceMenuOptions(FMenuBuilder& MenuBuilder);
 
@@ -101,6 +104,9 @@ private:
 	/** Fills the source control sub-menu */
 	void FillSourceControlSubMenu(FMenuBuilder& MenuBuilder);
 
+	/** Handler to check to see if SCC actions are allowed */
+	bool CanExecuteSourceControlActions() const;
+
 	/** Adds menu options related to working with collections */
 	bool AddCollectionMenuOptions(FMenuBuilder& MenuBuilder);
 
@@ -317,6 +323,7 @@ private:
 
 	/** Cached CanExecute vars */
 	bool bAtLeastOneNonRedirectorSelected;
+	bool bAtLeastOneClassSelected;
 	bool bCanExecuteSCCMerge;
 	bool bCanExecuteSCCCheckOut;
 	bool bCanExecuteSCCOpenForAdd;
diff --git a/Engine/Source/Editor/ContentBrowser/Private/ContentBrowserSingleton.cpp b/Engine/Source/Editor/ContentBrowser/Private/ContentBrowserSingleton.cpp
--- a/Engine/Source/Editor/ContentBrowser/Private/ContentBrowserSingleton.cpp
+++ b/Engine/Source/Editor/ContentBrowser/Private/ContentBrowserSingleton.cpp
@@ -10,6 +10,7 @@
 #include "SAssetDialog.h"
 #include "TutorialMetaData.h"
 #include "SDockTab.h"
+#include "NativeClassHierarchy.h"
 
 #define LOCTEXT_NAMESPACE "ContentBrowser"
 
@@ -295,7 +296,8 @@ void FContentBrowserSingleton::OnEditorLoadSelectedAssetsIfNeeded()
 
 FContentBrowserSingleton& FContentBrowserSingleton::Get()
 {
-	FContentBrowserModule& Module = FModuleManager::GetModuleChecked<FContentBrowserModule>("ContentBrowser");
+	static const FName ModuleName = "ContentBrowser";
+	FContentBrowserModule& Module = FModuleManager::GetModuleChecked<FContentBrowserModule>(ModuleName);
 	return static_cast<FContentBrowserSingleton&>(Module.Get());
 }
 
@@ -335,6 +337,15 @@ void FContentBrowserSingleton::ContentBrowserClosed(const TSharedRef<SContentBro
 	BrowserToLastKnownTabManagerMap.Add(ClosedBrowser->GetInstanceName(), ClosedBrowser->GetTabManager());
 }
 
+TSharedRef<FNativeClassHierarchy> FContentBrowserSingleton::GetNativeClassHierarchy()
+{
+	if(!NativeClassHierarchy.IsValid())
+	{
+		NativeClassHierarchy = MakeShareable(new FNativeClassHierarchy());
+	}
+	return NativeClassHierarchy.ToSharedRef();
+}
+
 void FContentBrowserSingleton::SharedCreateAssetDialogWindow(const TSharedRef<SAssetDialog>& AssetDialog, const FSharedAssetDialogConfig& InConfig, bool bModal) const
 {
 	const FVector2D DefaultWindowSize(1152.0f, 648.0f);
diff --git a/Engine/Source/Editor/ContentBrowser/Private/ContentBrowserSingleton.h b/Engine/Source/Editor/ContentBrowser/Private/ContentBrowserSingleton.h
--- a/Engine/Source/Editor/ContentBrowser/Private/ContentBrowserSingleton.h
+++ b/Engine/Source/Editor/ContentBrowser/Private/ContentBrowserSingleton.h
@@ -4,6 +4,7 @@
 #pragma once
 
 class SContentBrowser;
+class FNativeClassHierarchy;
 
 #define MAX_CONTENT_BROWSERS 4
 
@@ -41,6 +42,8 @@ public:
 	/** Notifies the singleton that a browser was closed */
 	void ContentBrowserClosed(const TSharedRef<SContentBrowser>& ClosedBrowser);
 
+	TSharedRef<FNativeClassHierarchy> GetNativeClassHierarchy();
+
 private:
 
 	/** Shared code to open an asset dialog window with a config */
@@ -83,6 +86,8 @@ private:
 
 	TWeakPtr<SContentBrowser> PrimaryContentBrowser;
 
+	TSharedPtr<FNativeClassHierarchy> NativeClassHierarchy;
+
 	/** An incrementing int32 which is used when making unique settings strings */
 	int32 SettingsStringID;
 };
\ No newline at end of file
diff --git a/Engine/Source/Editor/ContentBrowser/Private/ContentBrowserUtils.cpp b/Engine/Source/Editor/ContentBrowser/Private/ContentBrowserUtils.cpp
--- a/Engine/Source/Editor/ContentBrowser/Private/ContentBrowserUtils.cpp
+++ b/Engine/Source/Editor/ContentBrowser/Private/ContentBrowserUtils.cpp
@@ -1109,6 +1109,8 @@ bool ContentBrowserUtils::IsValidFolderName(const FString& FolderName, FText& Re
 
 bool ContentBrowserUtils::DoesFolderExist(const FString& FolderPath)
 {
+	// todo: jdale - CLASS - Will need updating to handle class folders
+
 	TArray<FString> SubPaths;
 	FAssetRegistryModule& AssetRegistryModule = FModuleManager::LoadModuleChecked<FAssetRegistryModule>(TEXT("AssetRegistry"));
 	AssetRegistryModule.Get().GetSubPaths(FPaths::GetPath(FolderPath), SubPaths, false);
@@ -1124,80 +1126,238 @@ bool ContentBrowserUtils::DoesFolderExist(const FString& FolderPath)
 	return false;
 }
 
+bool ContentBrowserUtils::IsRootDir(const FString& FolderPath)
+{
+	return IsAssetRootDir(FolderPath) || IsClassRootDir(FolderPath);
+}
+
 bool ContentBrowserUtils::IsAssetRootDir(const FString& FolderPath)
 {
-	return FolderPath == TEXT("/Game") || FolderPath == TEXT("/Engine") || FolderPath == TEXT("/Classes");
+	// All root asset folders start with "/" (not "/Classes_") and contain only a single / (at the beginning)
+	int32 LastSlashIndex = INDEX_NONE;
+	return FolderPath.Len() > 1 && !IsClassPath(FolderPath) && FolderPath.FindLastChar(TEXT('/'), LastSlashIndex) && LastSlashIndex == 0;
 }
 
-const TSharedPtr<FLinearColor> ContentBrowserUtils::LoadColor(const FString& FolderPath)
+bool ContentBrowserUtils::IsClassRootDir(const FString& FolderPath)
 {
-	// Ignore classes folder and newly created folders
-	if ( FolderPath != TEXT("/Classes") )
+	// All root class folders start with "/Classes_" and contain only a single / (at the beginning)
+	int32 LastSlashIndex = INDEX_NONE;
+	return IsClassPath(FolderPath) && FolderPath.FindLastChar(TEXT('/'), LastSlashIndex) && LastSlashIndex == 0;
+}
+
+FText ContentBrowserUtils::GetRootDirDisplayName(const FString& FolderPath)
+{
+	// Strip off any leading or trailing forward slashes
+	// We just want the name without any path separators
+	FString CleanFolderPath = FolderPath;
+	while(CleanFolderPath.StartsWith(TEXT("/")))
 	{
-		// Load the color from the config at this path
-		const FString RelativePath = FPackageName::LongPackageNameToFilename(FolderPath + TEXT("/"));
+		CleanFolderPath = CleanFolderPath.Mid(1);
+	}
+	while(CleanFolderPath.EndsWith(TEXT("/")))
+	{
+		CleanFolderPath = CleanFolderPath.Mid(0, CleanFolderPath.Len() - 1);
+	}
+
+	static const FString ClassesPrefix = TEXT("Classes_");
+	const bool bIsClassDir = CleanFolderPath.StartsWith(ClassesPrefix);
+
+	// Strip off the "Classes_" prefix
+	if(bIsClassDir)
+	{
+		CleanFolderPath = CleanFolderPath.Mid(ClassesPrefix.Len());
+	}
+
+	// Also localize well known folder names, like "Engine" and "Game"
+	static const FString EngineFolderName = TEXT("Engine");
+	static const FString GameFolderName = TEXT("Game");
+	FText LocalizedFolderName;
+	if(CleanFolderPath == EngineFolderName)
+	{
+		LocalizedFolderName = LOCTEXT("EngineFolderName", "Engine");
+	}
+	else if(CleanFolderPath == GameFolderName)
+	{
+		LocalizedFolderName = LOCTEXT("GameFolderName", "Game");
+	}
+	else
+	{
+		LocalizedFolderName = FText::FromString(CleanFolderPath);
+	}
+
+	return FText::Format((bIsClassDir) ? LOCTEXT("ClassesFolderFmt", "{0} C++ Classes") : LOCTEXT("ContentFolderFmt", "{0} Content"), LocalizedFolderName);
+}
+
+bool ContentBrowserUtils::IsClassPath(const FString& InPath)
+{
+	static const FString ClassesRootPrefix = TEXT("/Classes_");
+	return InPath.StartsWith(ClassesRootPrefix);
+}
+
+void ContentBrowserUtils::CountPathTypes(const TArray<FString>& InPaths, int32& OutNumAssetPaths, int32& OutNumClassPaths)
+{
+	OutNumAssetPaths = 0;
+	OutNumClassPaths = 0;
+
+	for(const FString& Path : InPaths)
+	{
+		if(IsClassPath(Path))
+		{
+			++OutNumClassPaths;
+		}
+		else
+		{
+			++OutNumAssetPaths;
+		}
+	}
+}
+
+void ContentBrowserUtils::CountPathTypes(const TArray<FName>& InPaths, int32& OutNumAssetPaths, int32& OutNumClassPaths)
+{
+	OutNumAssetPaths = 0;
+	OutNumClassPaths = 0;
+
+	for(const FName& Path : InPaths)
+	{
+		if(IsClassPath(Path.ToString()))
+		{
+			++OutNumClassPaths;
+		}
+		else
+		{
+			++OutNumAssetPaths;
+		}
+	}
+}
+
+void ContentBrowserUtils::CountItemTypes(const TArray<FAssetData>& InItems, int32& OutNumAssetItems, int32& OutNumClassItems)
+{
+	OutNumAssetItems = 0;
+	OutNumClassItems = 0;
 
+	for(const FAssetData& Item : InItems)
+	{
+		if(Item.AssetClass == NAME_Class)
+		{
+			++OutNumClassItems;
+		}
+		else
+		{
+			++OutNumAssetItems;
+		}
+	}
+}
+
+bool ContentBrowserUtils::IsValidPathToCreateNewClass(const FString& InPath)
+{
+	// Classes can currently only be added to game modules - if this is restricted, we can use IsClassPath here instead
+	// Classes can only be created in modules, so that will be at least two folders deep (two /)
+	static const FString GameClassesRootPrefix = TEXT("/Classes_Game");
+
+	int32 LastSlashIndex = INDEX_NONE;
+	return InPath.StartsWith(GameClassesRootPrefix) && InPath.FindLastChar(TEXT('/'), LastSlashIndex) && LastSlashIndex != 0;
+}
+
+bool ContentBrowserUtils::IsValidPathToCreateNewFolder(const FString& InPath)
+{
+	// We can't currently make folders in class paths
+	// If we do later allow folders in class paths, they must only be created within modules (see IsValidPathToCreateNewClass above)
+	return !IsClassPath(InPath);
+}
+
+const TSharedPtr<FLinearColor> ContentBrowserUtils::LoadColor(const FString& FolderPath)
+{
+	auto LoadColorInternal = [](const FString& InPath) -> TSharedPtr<FLinearColor>
+	{
 		// See if we have a value cached first
-		TSharedPtr< FLinearColor >* CachedColor = PathColors.Find( RelativePath );
-		if ( CachedColor )
+		TSharedPtr<FLinearColor> CachedColor = PathColors.FindRef(InPath);
+		if(CachedColor.IsValid())
 		{
-			return *CachedColor;
+			return CachedColor;
 		}
 		
 		// Loads the color of folder at the given path from the config
-		if (FPaths::FileExists(GEditorUserSettingsIni))
+		if(FPaths::FileExists(GEditorUserSettingsIni))
 		{
 			// Create a new entry from the config, skip if it's default
 			FString ColorStr;
-			if ( GConfig->GetString(TEXT("PathColor"), *RelativePath, ColorStr, GEditorUserSettingsIni) )
+			if(GConfig->GetString(TEXT("PathColor"), *InPath, ColorStr, GEditorUserSettingsIni))
 			{
 				FLinearColor Color;
-				if( Color.InitFromString( ColorStr ) && !Color.Equals( ContentBrowserUtils::GetDefaultColor() ) )
+				if(Color.InitFromString(ColorStr) && !Color.Equals(ContentBrowserUtils::GetDefaultColor()))
 				{
-					return PathColors.Add( RelativePath, MakeShareable( new FLinearColor( Color ) ) );
+					return PathColors.Add(InPath, MakeShareable(new FLinearColor(Color)));
 				}
 			}
 			else
 			{
-				return PathColors.Add( RelativePath, MakeShareable( new FLinearColor( ContentBrowserUtils::GetDefaultColor() ) ) );
+				return PathColors.Add(InPath, MakeShareable(new FLinearColor(ContentBrowserUtils::GetDefaultColor())));
 			}
 		}
-	}
 
-	return NULL;
-}
+		return nullptr;
+	};
 
-void ContentBrowserUtils::SaveColor(const FString& FolderPath, const TSharedPtr<FLinearColor> FolderColor, bool bForceAdd)
-{
-	check( FolderPath != TEXT("/Classes") );
+	// First try and find the color using the given path, as this works correctly for both assets and classes
+	TSharedPtr<FLinearColor> FoundColor = LoadColorInternal(FolderPath);
+	if(FoundColor.IsValid())
+	{
+		return FoundColor;
+	}
 
-	const FString RelativePath = FPackageName::LongPackageNameToFilename(FolderPath + TEXT("/"));
+	// If that failed, try and use the filename (assets used to use this as their color key, but it doesn't work with classes)
+	if(!IsClassPath(FolderPath))
+	{
+		const FString RelativePath = FPackageName::LongPackageNameToFilename(FolderPath + TEXT("/"));
+		return LoadColorInternal(RelativePath);
+	}
 
-	// Remove the color if it's invalid or default
-	const bool bRemove = !FolderColor.IsValid() || ( !bForceAdd && FolderColor->Equals( ContentBrowserUtils::GetDefaultColor() ) );
+	return nullptr;
+}
 
-	// Saves the color of the folder to the config
-	if (FPaths::FileExists(GEditorUserSettingsIni))
+void ContentBrowserUtils::SaveColor(const FString& FolderPath, const TSharedPtr<FLinearColor>& FolderColor, bool bForceAdd)
+{
+	auto SaveColorInternal = [](const FString& InPath, const TSharedPtr<FLinearColor>& FolderColor)
 	{
-		// If this is no longer custom, remove it
-		if ( bRemove )
+		// Saves the color of the folder to the config
+		if(FPaths::FileExists(GEditorUserSettingsIni))
 		{
-			GConfig->RemoveKey(TEXT("PathColor"), *RelativePath, GEditorUserSettingsIni);
+			GConfig->SetString(TEXT("PathColor"), *InPath, *FolderColor->ToString(), GEditorUserSettingsIni);
 		}
-		else
+
+		// Update the map too
+		PathColors.Add(InPath, FolderColor);
+	};
+
+	auto RemoveColorInternal = [](const FString& InPath)
+	{
+		// Remove the color of the folder from the config
+		if(FPaths::FileExists(GEditorUserSettingsIni))
 		{
-			GConfig->SetString(TEXT("PathColor"), *RelativePath, *FolderColor->ToString(), GEditorUserSettingsIni);
+			GConfig->RemoveKey(TEXT("PathColor"), *InPath, GEditorUserSettingsIni);
 		}
-	}
 
-	// Update the map too
-	if ( bRemove )
+		// Update the map too
+		PathColors.Remove(InPath);
+	};
+
+	// Remove the color if it's invalid or default
+	const bool bRemove = !FolderColor.IsValid() || (!bForceAdd && FolderColor->Equals(ContentBrowserUtils::GetDefaultColor()));
+
+	if(bRemove)
 	{
-		PathColors.Remove( RelativePath );
+		RemoveColorInternal(FolderPath);
 	}
 	else
 	{
-		PathColors.Add( RelativePath, FolderColor );
+		SaveColorInternal(FolderPath, FolderColor);
+	}
+
+	// Make sure and remove any colors using the legacy path format
+	if(!IsClassPath(FolderPath))
+	{
+		const FString RelativePath = FPackageName::LongPackageNameToFilename(FolderPath + TEXT("/"));
+		return RemoveColorInternal(RelativePath);
 	}
 }
 
@@ -1219,39 +1379,35 @@ bool ContentBrowserUtils::HasCustomColors( TArray< FLinearColor >* OutColors )
 
 			FString PathStr;
 			FString ColorStr;
-			if ( EntryStr.Split( TEXT( "/=" ), &PathStr, &ColorStr ) )	// DoesFolderExist doesn't like ending in a '/' so trim it here
+			if ( EntryStr.Split( TEXT( "=" ), &PathStr, &ColorStr ) )
 			{
-				// Ignore any that reference old folders
-				if ( FPackageName::TryConvertFilenameToLongPackageName(PathStr, PathStr) && DoesFolderExist( PathStr ) )
+				// Ignore any that have invalid or default colors
+				FLinearColor CurrentColor;
+				if( CurrentColor.InitFromString( ColorStr ) && !CurrentColor.Equals( ContentBrowserUtils::GetDefaultColor() ) )
 				{
-					// Ignore any that have invalid or default colors
-					FLinearColor CurrentColor;
-					if( CurrentColor.InitFromString( ColorStr ) && !CurrentColor.Equals( ContentBrowserUtils::GetDefaultColor() ) )
+					bHasCustom = true;
+					if ( OutColors )
 					{
-						bHasCustom = true;
-						if ( OutColors )
+						// Only add if not already present (ignores near matches too)
+						bool bAdded = false;
+						for( int32 ColorIndex = 0; ColorIndex < OutColors->Num(); ColorIndex++ )
 						{
-							// Only add if not already present (ignores near matches too)
-							bool bAdded = false;
-							for( int32 ColorIndex = 0; ColorIndex < OutColors->Num(); ColorIndex++ )
-							{
-								const FLinearColor& Color = (*OutColors)[ ColorIndex ];
-								if( CurrentColor.Equals( Color ) )
-								{
-									bAdded = true;
-									break;
-								}
-							}
-							if ( !bAdded )
+							const FLinearColor& Color = (*OutColors)[ ColorIndex ];
+							if( CurrentColor.Equals( Color ) )
 							{
-								OutColors->Add( CurrentColor );
+								bAdded = true;
+								break;
 							}
 						}
-						else
+						if ( !bAdded )
 						{
-							break;
+							OutColors->Add( CurrentColor );
 						}
 					}
+					else
+					{
+						break;
+					}
 				}
 			}
 		}
diff --git a/Engine/Source/Editor/ContentBrowser/Private/ContentBrowserUtils.h b/Engine/Source/Editor/ContentBrowser/Private/ContentBrowserUtils.h
--- a/Engine/Source/Editor/ContentBrowser/Private/ContentBrowserUtils.h
+++ b/Engine/Source/Editor/ContentBrowser/Private/ContentBrowserUtils.h
@@ -138,9 +138,36 @@ namespace ContentBrowserUtils
 	/** Returns true if the path specified exists as a folder in the asset registry */
 	bool DoesFolderExist(const FString& FolderPath);
 
-	/** Returns true if the pass-ed-in path is one of the asset root directories */
+	/** Check to see whether the given path is a root directory (either for asset or classes) */
+	bool IsRootDir(const FString& FolderPath);
+
+	/** Check to see whether the given path is a root asset directory */
 	bool IsAssetRootDir(const FString& FolderPath);
 
+	/** Check to see whether the given path is a root class directory */
+	bool IsClassRootDir(const FString& FolderPath);
+
+	/** Get the localized display name to use for the given root directory */
+	FText GetRootDirDisplayName(const FString& FolderPath);
+
+	/** Check to see whether the given path is rooted against a class directory */
+	bool IsClassPath(const FString& InPath);
+
+	/** Given an array of paths, work out how many are rooted against class roots, and how many are rooted against asset roots */
+	void CountPathTypes(const TArray<FString>& InPaths, int32& OutNumAssetPaths, int32& OutNumClassPaths);
+
+	/** Given an array of paths, work out how many are rooted against class roots, and how many are rooted against asset roots */
+	void CountPathTypes(const TArray<FName>& InPaths, int32& OutNumAssetPaths, int32& OutNumClassPaths);
+
+	/** Given an array of "asset" data, work out how many are assets, and how many are classes */
+	void CountItemTypes(const TArray<FAssetData>& InItems, int32& OutNumAssetItems, int32& OutNumClassItems);
+
+	/** Check to see whether the given path is a valid location in which to create new classes */
+	bool IsValidPathToCreateNewClass(const FString& InPath);
+
+	/** Check to see whether the given path is a valid location in which to create a new folder */
+	bool IsValidPathToCreateNewFolder(const FString& InPath);
+
 	/**
 	 * Loads the color of this path from the config
 	 *
@@ -156,7 +183,7 @@ namespace ContentBrowserUtils
 	 * @param FolderColor - The color the folder should appear as
 	 * @param bForceAdd - If true, force the color to be added for the path
 	 */
-	void SaveColor(const FString& FolderPath, const TSharedPtr<FLinearColor> FolderColor, bool bForceAdd = false);
+	void SaveColor(const FString& FolderPath, const TSharedPtr<FLinearColor>& FolderColor, bool bForceAdd = false);
 
 	/**
 	 * Checks to see if any folder has a custom color, optionally outputs them to a list
diff --git a/Engine/Source/Editor/ContentBrowser/Private/NativeClassHierarchy.cpp b/Engine/Source/Editor/ContentBrowser/Private/NativeClassHierarchy.cpp
new file mode 100644
--- /dev/null
+++ b/Engine/Source/Editor/ContentBrowser/Private/NativeClassHierarchy.cpp
@@ -0,0 +1,496 @@
+// Copyright 1998-2015 Epic Games, Inc. All Rights Reserved.
+
+#include "ContentBrowserPCH.h"
+#include "NativeClassHierarchy.h"
+#include "KismetEditorUtilities.h"
+#include "GameProjectGenerationModule.h"
+#include "IProjectManager.h"
+#include "HotReloadInterface.h"
+#include "SourceCodeNavigation.h"
+
+TSharedRef<FNativeClassHierarchyNode> FNativeClassHierarchyNode::MakeFolderEntry(FName InEntryName, FString InEntryPath)
+{
+	TSharedRef<FNativeClassHierarchyNode> NewEntry = MakeShareable(new FNativeClassHierarchyNode());
+	NewEntry->Type = ENativeClassHierarchyNodeType::Folder;
+	NewEntry->Class = nullptr;
+	NewEntry->EntryName = InEntryName;
+	NewEntry->EntryPath = MoveTemp(InEntryPath);
+	return NewEntry;
+}
+
+TSharedRef<FNativeClassHierarchyNode> FNativeClassHierarchyNode::MakeClassEntry(UClass* InClass, FName InClassModuleName, FString InClassModuleRelativePath, FString InEntryPath)
+{
+	TSharedRef<FNativeClassHierarchyNode> NewEntry = MakeShareable(new FNativeClassHierarchyNode());
+	NewEntry->Type = ENativeClassHierarchyNodeType::Class;
+	NewEntry->Class = InClass;
+	NewEntry->ClassModuleName = MoveTemp(InClassModuleName);
+	NewEntry->ClassModuleRelativePath = MoveTemp(InClassModuleRelativePath);
+	NewEntry->EntryName = InClass->GetFName();
+	NewEntry->EntryPath = MoveTemp(InEntryPath);
+	return NewEntry;
+}
+
+void FNativeClassHierarchyNode::AddChild(TSharedRef<FNativeClassHierarchyNode> ChildEntry)
+{
+	check(Type == ENativeClassHierarchyNodeType::Folder);
+	Children.Add(ChildEntry->EntryName, MoveTemp(ChildEntry));
+}
+
+FNativeClassHierarchy::FNativeClassHierarchy()
+{
+	PopulateHierarchy();
+
+	// Register to be notified of module changes
+	FModuleManager::Get().OnModulesChanged().AddRaw(this, &FNativeClassHierarchy::OnModulesChanged);
+
+	// Register to be notified of hot reloads
+	IHotReloadInterface& HotReloadSupport = FModuleManager::LoadModuleChecked<IHotReloadInterface>("HotReload");
+	HotReloadSupport.OnHotReload().AddRaw(this, &FNativeClassHierarchy::OnHotReload);
+}
+
+FNativeClassHierarchy::~FNativeClassHierarchy()
+{
+	FModuleManager::Get().OnModulesChanged().RemoveAll(this);
+
+	if(FModuleManager::Get().IsModuleLoaded("HotReload"))
+	{
+		IHotReloadInterface& HotReloadSupport = FModuleManager::GetModuleChecked<IHotReloadInterface>("HotReload");
+		HotReloadSupport.OnHotReload().RemoveAll(this);
+	}
+}
+
+void FNativeClassHierarchy::GetMatchingClasses(const FNativeClassHierarchyFilter& Filter, TArray<UClass*>& OutClasses) const
+{
+	TArray<TSharedRef<FNativeClassHierarchyNode>> NodesToSearch;
+	GatherMatchingNodesForPaths(Filter.ClassPaths, NodesToSearch);
+
+	for(const auto& NodeToSearch : NodesToSearch)
+	{
+		GetClassesRecursive(NodeToSearch, OutClasses, Filter.bRecursivePaths);
+	}
+}
+
+void FNativeClassHierarchy::GetMatchingFolders(const FNativeClassHierarchyFilter& Filter, TArray<FString>& OutFolders) const
+{
+	TArray<TSharedRef<FNativeClassHierarchyNode>> NodesToSearch;
+	GatherMatchingNodesForPaths(Filter.ClassPaths, NodesToSearch);
+
+	for(const auto& NodeToSearch : NodesToSearch)
+	{
+		GetFoldersRecursive(NodeToSearch, OutFolders, Filter.bRecursivePaths);
+	}
+}
+
+void FNativeClassHierarchy::GetClassFolders(TArray<FName>& OutClassRoots, TArray<FString>& OutClassFolders, const bool bIncludeEngineClasses, const bool bIncludePluginClasses) const
+{
+	static const FName EngineRootNodeName = "Classes_Engine";
+	static const FName GameRootNodeName = "Classes_Game";
+
+	for(const auto& RootNode : RootNodes)
+	{
+		const bool bRootNodePassesFilter = 
+			(RootNode.Key == GameRootNodeName) ||								// Always include game classes
+			(RootNode.Key == EngineRootNodeName && bIncludeEngineClasses) ||	// Only include engine classes if we were asked for them
+			(RootNode.Key != EngineRootNodeName && bIncludePluginClasses);		// Only include plugin classes if we were asked for them
+
+		if(bRootNodePassesFilter)
+		{
+			OutClassRoots.Add(RootNode.Key);
+			GetFoldersRecursive(RootNode.Value.ToSharedRef(), OutClassFolders);
+		}
+	}
+}
+
+void FNativeClassHierarchy::GetClassesRecursive(const TSharedRef<FNativeClassHierarchyNode>& HierarchyNode, TArray<UClass*>& OutClasses, const bool bRecurse)
+{
+	for(const auto& ChildNode : HierarchyNode->Children)
+	{
+		if(ChildNode.Value->Type == ENativeClassHierarchyNodeType::Class)
+		{
+			OutClasses.Add(ChildNode.Value->Class);
+		}
+
+		if(bRecurse)
+		{
+			GetClassesRecursive(ChildNode.Value.ToSharedRef(), OutClasses);
+		}
+	}
+}
+
+void FNativeClassHierarchy::GetFoldersRecursive(const TSharedRef<FNativeClassHierarchyNode>& HierarchyNode, TArray<FString>& OutFolders, const bool bRecurse)
+{
+	for(const auto& ChildNode : HierarchyNode->Children)
+	{
+		if(ChildNode.Value->Type == ENativeClassHierarchyNodeType::Folder)
+		{
+			OutFolders.Add(ChildNode.Value->EntryPath);
+		}
+
+		if(bRecurse)
+		{
+			GetFoldersRecursive(ChildNode.Value.ToSharedRef(), OutFolders);
+		}
+	}
+}
+
+void FNativeClassHierarchy::GatherMatchingNodesForPaths(const TArray<FName>& InClassPaths, TArray<TSharedRef<FNativeClassHierarchyNode>>& OutMatchingNodes) const
+{
+	if(InClassPaths.Num() == 0)
+	{
+		// No paths means search all roots
+		OutMatchingNodes.Reserve(RootNodes.Num());
+		for(const auto& RootNode : RootNodes)
+		{
+			OutMatchingNodes.Add(RootNode.Value.ToSharedRef());
+		}
+	}
+	else
+	{
+		for(const FName& ClassPath : InClassPaths)
+		{
+			TSharedPtr<FNativeClassHierarchyNode> CurrentNode;
+
+			TArray<FString> ClassPathParts;
+			ClassPath.ToString().ParseIntoArray(&ClassPathParts, TEXT("/"), true);
+			for(const FString& ClassPathPart : ClassPathParts)
+			{
+				// Try and find the node associated with this part of the path...
+				const FName ClassPathPartName = *ClassPathPart;
+				CurrentNode = (CurrentNode.IsValid()) ? CurrentNode->Children.FindRef(ClassPathPartName) : RootNodes.FindRef(ClassPathPartName);
+
+				// ... bail out if we didn't find a valid node
+				if(!CurrentNode.IsValid())
+				{
+					break;
+				}
+			}
+
+			if(CurrentNode.IsValid())
+			{
+				OutMatchingNodes.Add(CurrentNode.ToSharedRef());
+			}
+		}
+	}
+}
+
+void FNativeClassHierarchy::PopulateHierarchy()
+{
+	FAddClassMetrics AddClassMetrics;
+
+	RootNodes.Empty();
+
+	TSet<FName> GameModules = GetGameModules();
+	TMap<FName, FName> PluginModules = GetPluginModules();
+
+	for(TObjectIterator<UClass> ClassIt; ClassIt; ++ClassIt)
+	{
+		UClass* const CurrentClass = *ClassIt;
+		AddClass(CurrentClass, GameModules, PluginModules, AddClassMetrics);
+	}
+
+	UE_LOG(LogContentBrowser, Log, TEXT("Native class hierarchy populated in %0.4f seconds. Added %d classes and %d folders."), FPlatformTime::Seconds() - AddClassMetrics.StartTime, AddClassMetrics.NumClassesAdded, AddClassMetrics.NumFoldersAdded);
+
+	ClassHierarchyUpdatedDelegate.Broadcast();
+}
+
+void FNativeClassHierarchy::AddClassesForModule(const FName& InModuleName)
+{
+	FAddClassMetrics AddClassMetrics;
+
+	// Find the class package for this module
+	UPackage* const ClassPackage = FindPackage(nullptr, *(FString("/Script/") + InModuleName.ToString()));
+	if(!ClassPackage)
+	{
+		return;
+	}
+
+	TSet<FName> GameModules = GetGameModules();
+	TMap<FName, FName> PluginModules = GetPluginModules();
+
+	TArray<UObject*> PackageObjects;
+	GetObjectsWithOuter(ClassPackage, PackageObjects, false);
+	for(UObject* Object : PackageObjects)
+	{
+		UClass* const CurrentClass = Cast<UClass>(Object);
+		if(CurrentClass)
+		{
+			AddClass(CurrentClass, GameModules, PluginModules, AddClassMetrics);
+		}
+	}
+
+	UE_LOG(LogContentBrowser, Log, TEXT("Native class hierarchy updated for '%s' in %0.4f seconds. Added %d classes and %d folders."), *InModuleName.ToString(), FPlatformTime::Seconds() - AddClassMetrics.StartTime, AddClassMetrics.NumClassesAdded, AddClassMetrics.NumFoldersAdded);
+
+	ClassHierarchyUpdatedDelegate.Broadcast();
+}
+
+void FNativeClassHierarchy::RemoveClassesForModule(const FName& InModuleName)
+{
+	// Modules always exist directly under a root
+	for(const auto& RootNode : RootNodes)
+	{
+		TSharedPtr<FNativeClassHierarchyNode> ModuleNode = RootNode.Value->Children.FindRef(InModuleName);
+		if(ModuleNode.IsValid())
+		{
+			// Remove this module from its root
+			RootNode.Value->Children.Remove(InModuleName);
+
+			// If this module was the only child of this root, then we need to remove the root as well
+			if(RootNode.Value->Children.Num() == 0)
+			{
+				RootNodes.Remove(RootNode.Key);
+			}
+
+			ClassHierarchyUpdatedDelegate.Broadcast();
+
+			// We've found the module - break
+			break;
+		}
+	}
+}
+
+void FNativeClassHierarchy::AddClass(UClass* InClass, const TSet<FName>& InGameModules, const TMap<FName, FName>& InPluginModules, FAddClassMetrics& AddClassMetrics)
+{
+	// Ignore deprecated and temporary classes
+	if(InClass->HasAnyClassFlags(CLASS_Deprecated | CLASS_NewerVersionExists) || FKismetEditorUtilities::IsClassABlueprintSkeleton(InClass))
+	{
+		return;
+	}
+
+	const FName ClassModuleName = GetClassModuleName(InClass);
+	if(ClassModuleName.IsNone())
+	{
+		return;
+	}
+
+	static const FName ModuleRelativePathMetaDataKey = "ModuleRelativePath";
+	const FString& ClassModuleRelativeIncludePath = InClass->GetMetaData(ModuleRelativePathMetaDataKey);
+	if(ClassModuleRelativeIncludePath.IsEmpty())
+	{
+		return;
+	}
+
+	// Work out which root this class should go under
+	const FName RootNodeName = GetClassPathRootForModule(ClassModuleName, InGameModules, InPluginModules);
+
+	// Work out the final path to this class within the hierarchy (which isn't the same as the path on disk)
+	const FString ClassModuleRelativePath = ClassModuleRelativeIncludePath.Left(ClassModuleRelativeIncludePath.Find(TEXT("/"), ESearchCase::CaseSensitive, ESearchDir::FromEnd));
+	const FString ClassHierarchyPath = ClassModuleName.ToString() + TEXT("/") + ClassModuleRelativePath;
+
+	// Ensure we've added a valid root node
+	TSharedPtr<FNativeClassHierarchyNode>& RootNode = RootNodes.FindOrAdd(RootNodeName);
+	if(!RootNode.IsValid())
+	{
+		RootNode = FNativeClassHierarchyNode::MakeFolderEntry(RootNodeName, TEXT("/") + RootNodeName.ToString());
+		++AddClassMetrics.NumFoldersAdded;
+	}
+
+	// Split the class path and ensure we have nodes for each part
+	TArray<FString> HierarchyPathParts;
+	ClassHierarchyPath.ParseIntoArray(&HierarchyPathParts, TEXT("/"), true);
+	TSharedPtr<FNativeClassHierarchyNode> CurrentNode = RootNode;
+	for(const FString& HierarchyPathPart : HierarchyPathParts)
+	{
+		const FName HierarchyPathPartName = *HierarchyPathPart;
+		TSharedPtr<FNativeClassHierarchyNode>& ChildNode = CurrentNode->Children.FindOrAdd(HierarchyPathPartName);
+		if(!ChildNode.IsValid())
+		{
+			ChildNode = FNativeClassHierarchyNode::MakeFolderEntry(HierarchyPathPartName, CurrentNode->EntryPath + TEXT("/") + HierarchyPathPart);
+			++AddClassMetrics.NumFoldersAdded;
+		}
+		CurrentNode = ChildNode;
+	}
+
+	// Now add the final entry for the class
+	CurrentNode->AddChild(FNativeClassHierarchyNode::MakeClassEntry(InClass, ClassModuleName, ClassModuleRelativePath, CurrentNode->EntryPath + TEXT("/") + InClass->GetName()));
+	++AddClassMetrics.NumClassesAdded;
+}
+
+void FNativeClassHierarchy::AddFolder(const FString& InClassPath)
+{
+	bool bHasAddedFolder = false;
+
+	// Split the class path and ensure we have nodes for each part
+	TArray<FString> ClassPathParts;
+	InClassPath.ParseIntoArray(&ClassPathParts, TEXT("/"), true);
+	TSharedPtr<FNativeClassHierarchyNode> CurrentNode;
+	for(const FString& ClassPathPart : ClassPathParts)
+	{
+		const FName ClassPathPartName = *ClassPathPart;
+		TSharedPtr<FNativeClassHierarchyNode>& ChildNode = (CurrentNode.IsValid()) ? CurrentNode->Children.FindOrAdd(ClassPathPartName) : RootNodes.FindOrAdd(ClassPathPartName);
+		if(!ChildNode.IsValid())
+		{
+			ChildNode = FNativeClassHierarchyNode::MakeFolderEntry(ClassPathPartName, CurrentNode->EntryPath + TEXT("/") + ClassPathPart);
+			bHasAddedFolder = true;
+		}
+		CurrentNode = ChildNode;
+	}
+
+	if(bHasAddedFolder)
+	{
+		ClassHierarchyUpdatedDelegate.Broadcast();
+	}
+}
+
+bool FNativeClassHierarchy::GetFileSystemPath(const FString& InClassPath, FString& OutFileSystemPath) const
+{
+	// Split the class path into its component parts
+	TArray<FString> ClassPathParts;
+	InClassPath.ParseIntoArray(&ClassPathParts, TEXT("/"), true);
+	
+	// We need to have at least two sections (a root, and a module name) to be able to resolve a file system path
+	if(ClassPathParts.Num() < 2)
+	{
+		return false;
+	}
+
+	// Is this path using a known root?
+	TSharedPtr<FNativeClassHierarchyNode> RootNode = RootNodes.FindRef(*ClassPathParts[0]);
+	if(!RootNode.IsValid())
+	{
+		return false;
+	}
+
+	// Is this path using a known module within that root?
+	TSharedPtr<FNativeClassHierarchyNode> ModuleNode = RootNode->Children.FindRef(*ClassPathParts[1]);
+	if(!ModuleNode.IsValid())
+	{
+		return false;
+	}
+
+	// Get the base file path to the module, and then append any remaining parts of the class path (as the remaining parts mirror the file system)
+	if(FSourceCodeNavigation::FindModulePath(ClassPathParts[1], OutFileSystemPath))
+	{
+		for(int32 PathPartIndex = 2; PathPartIndex < ClassPathParts.Num(); ++PathPartIndex)
+		{
+			OutFileSystemPath /= ClassPathParts[PathPartIndex];
+		}
+		return true;
+	}
+
+	return false;
+}
+
+bool FNativeClassHierarchy::GetClassPath(UClass* InClass, FString& OutClassPath, const bool bIncludeClassName) const
+{
+	const FName ClassModuleName = GetClassModuleName(InClass);
+	if(ClassModuleName.IsNone())
+	{
+		return false;
+	}
+
+	static const FName ModuleRelativePathMetaDataKey = "ModuleRelativePath";
+	const FString& ClassModuleRelativeIncludePath = InClass->GetMetaData(ModuleRelativePathMetaDataKey);
+	if(ClassModuleRelativeIncludePath.IsEmpty())
+	{
+		return false;
+	}
+
+	TSet<FName> GameModules = GetGameModules();
+	TMap<FName, FName> PluginModules = GetPluginModules();
+
+	// Work out which root this class should go under
+	const FName RootNodeName = GetClassPathRootForModule(ClassModuleName, GameModules, PluginModules);
+
+	// Work out the final path to this class within the hierarchy (which isn't the same as the path on disk)
+	const FString ClassModuleRelativePath = ClassModuleRelativeIncludePath.Left(ClassModuleRelativeIncludePath.Find(TEXT("/"), ESearchCase::CaseSensitive, ESearchDir::FromEnd));
+	OutClassPath = FString(TEXT("/")) + RootNodeName.ToString() + TEXT("/") + ClassModuleName.ToString() + TEXT("/") + ClassModuleRelativePath;
+
+	if(bIncludeClassName)
+	{
+		OutClassPath += TEXT("/") + InClass->GetName();
+	}
+
+	return true;
+}
+
+void FNativeClassHierarchy::OnModulesChanged(FName InModuleName, EModuleChangeReason InModuleChangeReason)
+{
+	switch(InModuleChangeReason)
+	{
+	case EModuleChangeReason::ModuleLoaded:
+		AddClassesForModule(InModuleName);
+		break;
+
+	case EModuleChangeReason::ModuleUnloaded:
+		RemoveClassesForModule(InModuleName);
+		break;
+
+	default:
+		break;
+	}
+}
+
+void FNativeClassHierarchy::OnHotReload(bool bWasTriggeredAutomatically)
+{
+	PopulateHierarchy();
+}
+
+FName FNativeClassHierarchy::GetClassModuleName(UClass* InClass)
+{
+	UPackage* const ClassPackage = InClass->GetOuterUPackage();
+
+	if(ClassPackage)
+	{
+		return FPackageName::GetShortFName(ClassPackage->GetFName());
+	}
+
+	return NAME_None;
+}
+
+FName FNativeClassHierarchy::GetClassPathRootForModule(const FName& InModuleName, const TSet<FName>& InGameModules, const TMap<FName, FName>& InPluginModules)
+{
+	static const FName EngineRootNodeName = "Classes_Engine";
+	static const FName GameRootNodeName = "Classes_Game";
+
+	// Work out which root this class should go under (anything that isn't a game or plugin module goes under engine)
+	FName RootNodeName = EngineRootNodeName;
+	if(InGameModules.Contains(InModuleName))
+	{
+		RootNodeName = GameRootNodeName;
+	}
+	else if(InPluginModules.Contains(InModuleName))
+	{
+		const FName PluginName = InPluginModules.FindRef(InModuleName);
+		RootNodeName = FName(*(FString(TEXT("Classes_")) + PluginName.ToString()));
+	}
+
+	return RootNodeName;
+}
+
+TSet<FName> FNativeClassHierarchy::GetGameModules()
+{
+	FGameProjectGenerationModule& GameProjectModule = FGameProjectGenerationModule::Get();
+
+	// Build up a set of known game modules - used to work out which modules populate Classes_Game
+	TSet<FName> GameModules;
+	if(GameProjectModule.ProjectHasCodeFiles())
+	{
+		TArray<FModuleContextInfo> GameModulesInfo = GameProjectModule.GetCurrentProjectModules();
+		for(const auto& GameModuleInfo : GameModulesInfo)
+		{
+			GameModules.Add(FName(*GameModuleInfo.ModuleName));
+		}
+	}
+
+	return GameModules;
+}
+
+TMap<FName, FName> FNativeClassHierarchy::GetPluginModules()
+{
+	IPluginManager& PluginManager = IPluginManager::Get();
+
+	// Build up a map of plugin modules -> plugin names - used to work out which modules populate a given Classes_PluginName
+	TMap<FName, FName> PluginModules;
+	{
+		TArray<FPluginStatus> Plugins = PluginManager.QueryStatusForAllPlugins();
+		for(const auto& Plugin : Plugins)
+		{
+			for(const auto& PluginModule : Plugin.Modules)
+			{
+				PluginModules.Add(PluginModule.Name, FName(*Plugin.Name));
+			}
+		}
+	}
+
+	return PluginModules;
+}
diff --git a/Engine/Source/Editor/ContentBrowser/Private/NativeClassHierarchy.h b/Engine/Source/Editor/ContentBrowser/Private/NativeClassHierarchy.h
new file mode 100644
--- /dev/null
+++ b/Engine/Source/Editor/ContentBrowser/Private/NativeClassHierarchy.h
@@ -0,0 +1,283 @@
+// Copyright 1998-2015 Epic Games, Inc. All Rights Reserved.
+
+#pragma once
+
+/**
+ * Type of hierarchy node
+ */
+enum class ENativeClassHierarchyNodeType : uint8
+{
+	Folder,
+	Class,
+};
+
+/**
+ * Single node in the class hierarchy
+ */
+struct FNativeClassHierarchyNode
+{
+	/** Helper function to make a folder node entry */
+	static TSharedRef<FNativeClassHierarchyNode> MakeFolderEntry(FName InEntryName, FString InEntryPath);
+
+	/** Helper function to make a class node entry */
+	static TSharedRef<FNativeClassHierarchyNode> MakeClassEntry(UClass* InClass, FName InClassModuleName, FString InClassModuleRelativePath, FString InEntryPath);
+
+	void AddChild(TSharedRef<FNativeClassHierarchyNode> ChildEntry);
+
+	/** Type of node, folder or class */
+	ENativeClassHierarchyNodeType Type;
+
+	/** The class this node is for (Type == Class) */
+	UClass* Class;
+
+	/** The name of the module the class is in (Type == Class) */
+	FName ClassModuleName;
+
+	/** Folder this class is in, relative to the class module (Type == Class) */
+	FString ClassModuleRelativePath;
+
+	/** Name used when showing this entry in the UI */
+	FName EntryName;
+
+	/** Path to this entry in the class hierarchy (not the same as the location on disk) */
+	FString EntryPath;
+
+	/** Child entries (Type == Folder) */
+	TMap<FName, TSharedPtr<FNativeClassHierarchyNode>> Children;
+};
+
+/**
+ * A filter used when querying the native class hierarchy.
+ * Each component element is processed as an 'OR' operation while all the components are processed together as an 'AND' operation.
+ */
+struct FNativeClassHierarchyFilter
+{
+	/** The filter component for class paths */
+	TArray<FName> ClassPaths;
+	/** If true, ClassPaths components will be recursive */
+	bool bRecursivePaths;
+
+	FNativeClassHierarchyFilter()
+	{
+		bRecursivePaths = false;
+	}
+
+	/** Appends the other filter to this one */
+	void Append(const FNativeClassHierarchyFilter& Other)
+	{
+		ClassPaths.Append(Other.ClassPaths);
+
+		bRecursivePaths |= Other.bRecursivePaths;
+	}
+
+	/** Returns true if this filter has no entries */
+	bool IsEmpty() const
+	{
+		return ClassPaths.Num() == 0;
+	}
+
+	/** Clears this filter of all entries */
+	void Clear()
+	{
+		ClassPaths.Empty();
+
+		bRecursivePaths = false;
+
+		ensure(IsEmpty());
+	}
+};
+
+/** 
+ * Generates a hierarchical tree of native UObject classes based on their location in the file system as used by the asset view when showing C++ classes
+ * This keeps its class hierarchy up-to-date as modules are loaded/unloaded
+ */
+class FNativeClassHierarchy
+{
+public:
+	/** Constructor and destructor */
+	FNativeClassHierarchy();
+	~FNativeClassHierarchy();
+
+	/** Get the delegate called when classes are added or removed from this class hierarchy */
+	FSimpleMulticastDelegate& OnClassHierarchyUpdated()
+	{
+		return ClassHierarchyUpdatedDelegate;
+	}
+
+	/**
+	 * Work out which classes known to the class hierarchy match the given filter
+	 *
+	 * @param Filter - The filter to apply when working out which classes match
+	 * @param OutClasses - Array to be populated with matching classes
+	 */
+	void GetMatchingClasses(const FNativeClassHierarchyFilter& Filter, TArray<UClass*>& OutClasses) const;
+
+	/**
+	 * Work out which folders known to the class hierarchy match the given filter
+	 *
+	 * @param Filter - The filter to apply when working out which folders match
+	 * @param OutFolders - Array to be populated with matching folder paths (these are in the form "/Classes_Name/ModuleName/SubFolder")
+	 */
+	void GetMatchingFolders(const FNativeClassHierarchyFilter& Filter, TArray<FString>& OutFolders) const;
+
+	/**
+	 * Get all folders known to the class hierarchy
+	 *
+	 * @param OutClassRoots - Array to be populated with the known class path root folders (these are in the form "/Classes_Name")
+	 * @param OutClassFolders - Array to be populated with known class folder paths (these are in the form "/Classes_Name/ModuleName/SubFolder")
+	 * @param bIncludeEngineClasses - Whether we should include the "/Classes_Engine" root (and its associated folders) in the returned results
+	 * @param bIncludePluginClasses - Whether we should include plugin class path roots (and their associated folders) in the returned results (these are all roots except for "/Classes_Game" and "/Classes_Engine")
+	 */
+	void GetClassFolders(TArray<FName>& OutClassRoots, TArray<FString>& OutClassFolders, const bool bIncludeEngineClasses = true, const bool bIncludePluginClasses = true) const;
+
+	/**
+	 * Given a class path, work out the corresponding filesystem path on disk
+	 * eg) Given "/Classes_Game/MyGame/MyAwesomeCode", we might resolve that to "../../../MyGame/Source/MyGame/MyAwesomeCode"
+	 *
+	 * @param InClassPath - The class path to try and resolve, eg) "/Classes_Game/MyGame/MyAwesomeCode"
+	 * @param OutFileSystemPath - The string to fill with the resolved filesystem path, eg) "../../../MyGame/Source/MyGame/MyAwesomeCode"
+	 * 
+	 * @return true if the file system path could be resolved and OutFileSystemPath was filled in, false otherwise
+	 */
+	bool GetFileSystemPath(const FString& InClassPath, FString& OutFileSystemPath) const;
+
+	/**
+	 * Work out the class path that should be used for the given class
+	 *
+	 * @param InClass - The class path to try and get the class path for
+	 * @param OutClassPath - The string to fill with the resolved class path, eg) "/Classes_Game/MyGame/MyAwesomeCode/MyAwesomeClass"
+	 * @param bIncludeClassName - true to include the class name on the returned class path, false to get the class path of the containing folder
+	 * 
+	 * @return true if the class path could be resolved and OutClassPath was filled in, false otherwise
+	 */
+	bool GetClassPath(UClass* InClass, FString& OutClassPath, const bool bIncludeClassName = true) const;
+
+	/**
+	 * This will add a transient folder into the hierarchy
+	 * The folder will be lost unless a class is added to it before the hierarchy is next re-populated
+	 *
+	 * @param InClassPath - The location of the new folder (in class path form - eg) "/Classes_Game/MyGame/MyAwesomeCode")
+	 */
+	void AddFolder(const FString& InClassPath);
+
+private:
+	struct FAddClassMetrics
+	{
+		FAddClassMetrics()
+			: StartTime(FPlatformTime::Seconds())
+			, NumClassesAdded(0)
+			, NumFoldersAdded(0)
+		{
+		}
+
+		double StartTime;
+		int32 NumClassesAdded;
+		int32 NumFoldersAdded;
+	};
+
+	/**
+	 * Update OutClasses with any classes that are children of the given node
+	 *
+	 * @param HierarchyNode - The node to add the children of
+	 * @param OutClasses - Array to be populated with the child classes
+	 * @param bRecurse - True to recurse into sub-folders, false to only check the given node
+	 */
+	static void GetClassesRecursive(const TSharedRef<FNativeClassHierarchyNode>& HierarchyNode, TArray<UClass*>& OutClasses, const bool bRecurse = true);
+
+	/**
+	 * Update OutFolders with any folders that are children of the given node
+	 *
+	 * @param HierarchyNode - The node to add the children of
+	 * @param OutFolders - Array to be populated with the child folders
+	 * @param bRecurse - True to recurse into sub-folders, false to only check the given node
+	 */
+	static void GetFoldersRecursive(const TSharedRef<FNativeClassHierarchyNode>& HierarchyNode, TArray<FString>& OutFolders, const bool bRecurse = true);
+
+	/**
+	 * Populate OutMatchingNodes with the nodes that correspond to the given class paths
+	 *
+	 * @param InClassPaths - The class paths to find the nodes for, or an empty list to get all root nodes
+	 * @param OutMatchingNodes - Array to be populated the nodes that correspond to the given class paths
+	 */
+	void GatherMatchingNodesForPaths(const TArray<FName>& InClassPaths, TArray<TSharedRef<FNativeClassHierarchyNode>>& OutMatchingNodes) const;
+
+	/**
+	 * Completely clear and re-populate the known class hierarchy
+	 */
+	void PopulateHierarchy();
+
+	/**
+	 * Append any classes from the given module to the known class hierarchy
+	 *
+	 * @param InModuleName - The name of the module to add
+	 */
+	void AddClassesForModule(const FName& InModuleName);
+
+	/**
+	 * Remove any classes in the given module from the known class hierarchy
+	 *
+	 * @param InModuleName - The name of the module to remove
+	 */
+	void RemoveClassesForModule(const FName& InModuleName);
+
+	/**
+	 * Add a single class to the known class hierarchy, creating any folders as required
+	 *
+	 * @param InClass - The class that is to be added
+	 * @param InGameModules - The list of modules that belong to the "/Classes_Game" root
+	 * @param InPluginModules - The list of modules that belong to neither the "/Classes_Game" or "/Classes_Engine" roots
+	 * @param AddClassMetrics - Metrics to update as new classes and folders are added (used to report population performance)
+	 */
+	void AddClass(UClass* InClass, const TSet<FName>& InGameModules, const TMap<FName, FName>& InPluginModules, FAddClassMetrics& AddClassMetrics);
+
+	/**
+	 * Called when we're notified that a module has changed status
+	 *
+	 * @param InModuleName - The module that changed status
+	 * @param InModuleChangeReason - Why the module changed status
+	 */
+	void OnModulesChanged(FName InModuleName, EModuleChangeReason InModuleChangeReason);
+
+	/**
+	 * Called when we're notified that a module has been hot-reloaded
+	 *
+	 * @param bWasTriggeredAutomatically - True if the hot-reload was automatically triggered, or false if it was from a user action
+	 */
+	void OnHotReload(bool bWasTriggeredAutomatically);
+
+	/**
+	 * Given a class, work out which module it belongs to
+	 *
+	 * @param InClass - The class we want to find the module for
+	 *
+	 * @return The name of the module that holds the class, eg) "CoreUObject"
+	 */
+	static FName GetClassModuleName(UClass* InClass);
+
+	/**
+	 * Given a module, work out which root path it should use as a parent
+	 *
+	 * @param InModuleName - The module we want to find the root path for
+	 * @param InGameModules - The list of modules that belong to the "/Classes_Game" root
+	 * @param InPluginModules - The list of modules that belong to neither the "/Classes_Game" or "/Classes_Engine" roots
+	 *
+	 * @return The name of the root path to use, eg "/Classes_Game"
+	 */
+	static FName GetClassPathRootForModule(const FName& InModuleName, const TSet<FName>& InGameModules, const TMap<FName, FName>& InPluginModules);
+
+	/**
+	 * Calculate the list of modules that belong to the "/Classes_Game" root
+	 */
+	static TSet<FName> GetGameModules();
+
+	/**
+	 * Calculate the list of modules that belong to neither the "/Classes_Game" or "/Classes_Engine" roots
+	 */
+	static TMap<FName, FName> GetPluginModules();
+
+	/** Root level nodes corresponding to the root folders used by the Content Browser, eg) Classes_Engine, Classes_Game, etc */
+	TMap<FName, TSharedPtr<FNativeClassHierarchyNode>> RootNodes;
+
+	/** Delegate called when the class hierarchy is updated */
+	FSimpleMulticastDelegate ClassHierarchyUpdatedDelegate;
+};
diff --git a/Engine/Source/Editor/ContentBrowser/Private/NewAssetContextMenu.cpp b/Engine/Source/Editor/ContentBrowser/Private/NewAssetContextMenu.cpp
deleted file mode 100644
diff --git a/Engine/Source/Editor/ContentBrowser/Private/NewAssetContextMenu.h b/Engine/Source/Editor/ContentBrowser/Private/NewAssetContextMenu.h
deleted file mode 100644
diff --git a/Engine/Source/Editor/ContentBrowser/Private/NewAssetOrClassContextMenu.cpp b/Engine/Source/Editor/ContentBrowser/Private/NewAssetOrClassContextMenu.cpp
new file mode 100644
--- /dev/null
+++ b/Engine/Source/Editor/ContentBrowser/Private/NewAssetOrClassContextMenu.cpp
@@ -0,0 +1,420 @@
+// Copyright 1998-2015 Epic Games, Inc. All Rights Reserved.
+
+#include "ContentBrowserPCH.h"
+#include "NewAssetOrClassContextMenu.h"
+#include "IDocumentation.h"
+#include "EditorClassUtils.h"
+
+#define LOCTEXT_NAMESPACE "ContentBrowser"
+
+struct FFactoryItem
+{
+	UFactory* Factory;
+	FText DisplayName;
+
+	FFactoryItem(UFactory* InFactory, const FText& InDisplayName)
+		: Factory(InFactory), DisplayName(InDisplayName)
+	{}
+};
+
+TArray<FFactoryItem> FindFactoriesInCategory(EAssetTypeCategories::Type AssetTypeCategory)
+{
+	TArray<FFactoryItem> FactoriesInThisCategory;
+	for (TObjectIterator<UClass> It; It; ++It)
+	{
+		UClass* Class = *It;
+		if (Class->IsChildOf(UFactory::StaticClass()) && !Class->HasAnyClassFlags(CLASS_Abstract))
+		{
+			UFactory* Factory = Class->GetDefaultObject<UFactory>();
+			if (Factory->ShouldShowInNewMenu() && ensure(!Factory->GetDisplayName().IsEmpty()))
+			{
+				uint32 FactoryCategories = Factory->GetMenuCategories();
+
+				if (FactoryCategories & AssetTypeCategory)
+				{
+					new(FactoriesInThisCategory)FFactoryItem(Factory, Factory->GetDisplayName());
+				}
+			}
+		}
+	}
+
+	return FactoriesInThisCategory;
+}
+
+struct FAdvancedAssetCategory
+{
+	EAssetTypeCategories::Type CategoryType;
+	FText CategoryName;
+};
+
+class SFactoryMenuEntry : public SCompoundWidget
+{
+public:
+
+	SLATE_BEGIN_ARGS( SFactoryMenuEntry )
+		: _Width(32)
+		, _Height(32)
+		{}
+		SLATE_ARGUMENT( uint32, Width )
+		SLATE_ARGUMENT( uint32, Height )
+	SLATE_END_ARGS()
+
+	/**
+	 * Construct this widget.  Called by the SNew() Slate macro.
+	 *
+	 * @param	InArgs				Declaration used by the SNew() macro to construct this widget
+	 * @param	Factory				The factory this menu entry represents
+	 */
+	void Construct( const FArguments& InArgs, UFactory* Factory )
+	{
+		// Since we are only displaying classes, we need no thumbnail pool.
+		Thumbnail = MakeShareable( new FAssetThumbnail( Factory->GetSupportedClass(), InArgs._Width, InArgs._Height, TSharedPtr<FAssetThumbnailPool>() ) );
+
+		FName ClassThumbnailBrushOverride = Factory->GetNewAssetThumbnailOverride();
+		TSharedPtr<SWidget> ThumbnailWidget;
+		if ( ClassThumbnailBrushOverride != NAME_None )
+		{
+			const bool bAllowFadeIn = false;
+			const bool bForceGenericThumbnail = true;
+			const bool bAllowHintText = false;
+			ThumbnailWidget = Thumbnail->MakeThumbnailWidget(bAllowFadeIn, bForceGenericThumbnail, EThumbnailLabel::AssetName, FText(), FLinearColor(0,0,0,0), bAllowHintText, ClassThumbnailBrushOverride);
+		}
+		else
+		{
+			ThumbnailWidget = Thumbnail->MakeThumbnailWidget();
+		}
+
+		ChildSlot
+		[
+			SNew( SHorizontalBox )
+			+SHorizontalBox::Slot()
+			.Padding( 4, 0, 0, 0 )
+			.VAlign(VAlign_Center)
+			.AutoWidth()
+			[
+				SNew( SBox )
+				.WidthOverride( Thumbnail->GetSize().X + 4 )
+				.HeightOverride( Thumbnail->GetSize().Y + 4 )
+				[
+					ThumbnailWidget.ToSharedRef()
+				]
+			]
+
+			+SHorizontalBox::Slot()
+			.VAlign(VAlign_Center)
+			.Padding(4, 0, 4, 0)
+			[
+				SNew( SVerticalBox )
+				+SVerticalBox::Slot()
+				.Padding(0, 0, 0, 1)
+				.AutoHeight()
+				[
+					SNew(STextBlock)
+					.Font( FEditorStyle::GetFontStyle("LevelViewportContextMenu.AssetLabel.Text.Font") )
+					.Text( Factory->GetDisplayName() )
+				]
+			]
+		];
+
+		
+		SetToolTip(IDocumentation::Get()->CreateToolTip(Factory->GetToolTip(), nullptr, Factory->GetToolTipDocumentationPage(), Factory->GetToolTipDocumentationExcerpt()));
+	}
+
+private:
+	TSharedPtr< FAssetThumbnail > Thumbnail;
+};
+
+void FNewAssetOrClassContextMenu::MakeContextMenu(
+	FMenuBuilder& MenuBuilder, 
+	const TArray<FName>& InSelectedPaths,
+	const FOnNewAssetRequested& InOnNewAssetRequested, 
+	const FOnNewClassRequested& InOnNewClassRequested, 
+	const FOnNewFolderRequested& InOnNewFolderRequested, 
+	const FOnImportAssetRequested& InOnImportAssetRequested, 
+	const FOnGetContentRequested& InOnGetContentRequested
+	)
+{
+	TArray<FString> SelectedStringPaths;
+	SelectedStringPaths.Reserve(InSelectedPaths.Num());
+	for(const FName& Path : InSelectedPaths)
+	{
+		SelectedStringPaths.Add(Path.ToString());
+	}
+
+	MakeContextMenu(MenuBuilder, SelectedStringPaths, InOnNewAssetRequested, InOnNewClassRequested, InOnNewFolderRequested, InOnImportAssetRequested, InOnGetContentRequested);
+}
+
+void FNewAssetOrClassContextMenu::MakeContextMenu(
+	FMenuBuilder& MenuBuilder, 
+	const TArray<FString>& InSelectedPaths,
+	const FOnNewAssetRequested& InOnNewAssetRequested, 
+	const FOnNewClassRequested& InOnNewClassRequested, 
+	const FOnNewFolderRequested& InOnNewFolderRequested, 
+	const FOnImportAssetRequested& InOnImportAssetRequested, 
+	const FOnGetContentRequested& InOnGetContentRequested
+	)
+{
+	int32 NumAssetPaths, NumClassPaths;
+	ContentBrowserUtils::CountPathTypes(InSelectedPaths, NumAssetPaths, NumClassPaths);
+
+	const FString FirstSelectedPath = (InSelectedPaths.Num() > 0) ? InSelectedPaths[0] : FString();
+	const bool bIsValidNewClassPath = ContentBrowserUtils::IsValidPathToCreateNewClass(FirstSelectedPath);
+	const bool bIsValidNewFolderPath = ContentBrowserUtils::IsValidPathToCreateNewFolder(FirstSelectedPath);
+	const bool bHasSinglePathSelected = InSelectedPaths.Num() == 1;
+
+	auto CanExecuteFolderActions = [NumAssetPaths, NumClassPaths, bIsValidNewFolderPath]() -> bool
+	{
+		// We can execute folder actions when we only have a single path selected, and that path is a valid path for creating a folder
+		return (NumAssetPaths + NumClassPaths) == 1 && bIsValidNewFolderPath;
+	};
+	const FCanExecuteAction CanExecuteFolderActionsDelegate = FCanExecuteAction::CreateLambda(CanExecuteFolderActions);
+
+	auto CanExecuteAssetActions = [NumAssetPaths, NumClassPaths]() -> bool
+	{
+		// We can execute asset actions when we only have a single asset path selected
+		return NumAssetPaths == 1 && NumClassPaths == 0;
+	};
+	const FCanExecuteAction CanExecuteAssetActionsDelegate = FCanExecuteAction::CreateLambda(CanExecuteAssetActions);
+
+	auto CanExecuteClassActions = [NumAssetPaths, NumClassPaths, bIsValidNewClassPath]() -> bool
+	{
+		// We can execute class actions when we only have a single class path selected, and that path is a valid path for creating a class
+		return NumClassPaths == 1 && NumAssetPaths == 0 && bIsValidNewClassPath;
+	};
+	const FCanExecuteAction CanExecuteClassActionsDelegate = FCanExecuteAction::CreateLambda(CanExecuteClassActions);
+
+	// Get Content
+	if ( InOnGetContentRequested.IsBound() )
+	{
+		MenuBuilder.BeginSection( "ContentBrowserGetContent", LOCTEXT( "GetContentMenuHeading", "Content" ) );
+		{
+			MenuBuilder.AddMenuEntry(
+				LOCTEXT( "GetContentText", "Add feature or content pack..." ),
+				LOCTEXT( "GetContentTooltip", "Add features and content packs to the project." ),
+				FSlateIcon( FEditorStyle::GetStyleSetName(), "ContentBrowser.AddContent" ),
+				FUIAction(
+					FExecuteAction::CreateStatic( &FNewAssetOrClassContextMenu::ExecuteGetContent, InOnGetContentRequested ),
+					CanExecuteAssetActionsDelegate
+					)
+				);
+		}
+		MenuBuilder.EndSection();
+	}
+
+	// New Folder
+	if(InOnNewFolderRequested.IsBound() && GetDefault<UContentBrowserSettings>()->DisplayFolders)
+	{
+		MenuBuilder.BeginSection("ContentBrowserNewFolder", LOCTEXT("FolderMenuHeading", "Folder") );
+		{
+			FText NewFolderToolTip;
+			if(bHasSinglePathSelected)
+			{
+				if(bIsValidNewFolderPath)
+				{
+					NewFolderToolTip = FText::Format(LOCTEXT("NewFolderTooltip_CreateIn", "Create a new folder in {0}."), FText::FromString(FirstSelectedPath));
+				}
+				else
+				{
+					NewFolderToolTip = FText::Format(LOCTEXT("NewFolderTooltip_InvalidPath", "Cannot create new folders in {0}."), FText::FromString(FirstSelectedPath));
+				}
+			}
+			else
+			{
+				NewFolderToolTip = LOCTEXT("NewFolderTooltip_InvalidNumberOfPaths", "Can only create folders when there is a single path selected.");
+			}
+
+			MenuBuilder.AddMenuEntry(
+				LOCTEXT("NewFolderLabel", "New Folder"),
+				NewFolderToolTip,
+				FSlateIcon(FEditorStyle::GetStyleSetName(), "ContentBrowser.NewFolderIcon"),
+				FUIAction(
+					FExecuteAction::CreateStatic(&FNewAssetOrClassContextMenu::ExecuteNewFolder, FirstSelectedPath, InOnNewFolderRequested),
+					CanExecuteFolderActionsDelegate
+					)
+				);
+		}
+		MenuBuilder.EndSection(); //ContentBrowserNewFolder
+	}
+
+	// Add Class
+	if(InOnNewClassRequested.IsBound())
+	{
+		FText NewClassToolTip;
+		if(bHasSinglePathSelected)
+		{
+			if(bIsValidNewClassPath)
+			{
+				NewClassToolTip = FText::Format(LOCTEXT("NewClassTooltip_CreateIn", "Create a new class in {0}."), FText::FromString(FirstSelectedPath));
+			}
+			else
+			{
+				NewClassToolTip = FText::Format(LOCTEXT("NewClassTooltip_InvalidPath", "Cannot create new classes in {0}."), FText::FromString(FirstSelectedPath));
+			}
+		}
+		else
+		{
+			NewClassToolTip = LOCTEXT("NewClassTooltip_InvalidNumberOfPaths", "Can only create classes when there is a single path selected.");
+		}
+
+		MenuBuilder.BeginSection("ContentBrowserNewClass", LOCTEXT("ClassMenuHeading", "Class") );
+		{
+			MenuBuilder.AddMenuEntry(
+				LOCTEXT("NewClassLabel", "New Class"),
+				NewClassToolTip,
+				FSlateIcon(FEditorStyle::GetStyleSetName(), "MainFrame.AddCodeToProject"),
+				FUIAction(
+					FExecuteAction::CreateStatic(&FNewAssetOrClassContextMenu::ExecuteNewClass, FirstSelectedPath, InOnNewClassRequested),
+					CanExecuteClassActionsDelegate
+					)
+				);
+		}
+		MenuBuilder.EndSection(); //ContentBrowserNewClass
+	}
+
+	// Import
+	if (InOnImportAssetRequested.IsBound() && !FirstSelectedPath.IsEmpty())
+	{
+		MenuBuilder.BeginSection( "ContentBrowserImportAsset", LOCTEXT( "ImportAssetMenuHeading", "Import Asset" ) );
+		{
+			MenuBuilder.AddMenuEntry(
+				FText::Format( LOCTEXT( "ImportAsset", "Import to {0}..." ), FText::FromString( FirstSelectedPath ) ),
+				LOCTEXT( "ImportAssetTooltip", "Imports an asset from file to this folder." ),
+				FSlateIcon( FEditorStyle::GetStyleSetName(), "ContentBrowser.ImportIcon" ),
+				FUIAction(
+					FExecuteAction::CreateStatic( &FNewAssetOrClassContextMenu::ExecuteImportAsset, InOnImportAssetRequested, FirstSelectedPath ),
+					CanExecuteAssetActionsDelegate
+					)
+				);
+		}
+		MenuBuilder.EndSection();
+	}
+
+	// Add Basic Asset
+	if (InOnNewAssetRequested.IsBound())
+	{
+		MenuBuilder.BeginSection("ContentBrowserNewBasicAsset", LOCTEXT("BasicAssetsMenuHeading", "Create Basic Asset") );
+		{
+			CreateNewAssetMenuCategory(
+				MenuBuilder, 
+				EAssetTypeCategories::Basic, 
+				FirstSelectedPath, 
+				InOnNewAssetRequested, 
+				CanExecuteAssetActionsDelegate
+				);
+		}
+		MenuBuilder.EndSection(); //ContentBrowserNewBasicAsset
+	}
+
+	// Add Advanced Asset
+	if (InOnNewAssetRequested.IsBound())
+	{
+		MenuBuilder.BeginSection("ContentBrowserNewAdvancedAsset", LOCTEXT("AdvancedAssetsMenuHeading", "Create Advanced Asset"));
+		{
+			// List of advanced categories to add, in order
+			const FAdvancedAssetCategory AdvancedAssetCategories[] = {
+				//	CategoryType								CategoryName
+				{	EAssetTypeCategories::Animation,			LOCTEXT("AnimationAssetCategory", "Animation")				},
+				{	EAssetTypeCategories::Blueprint,			LOCTEXT("BlueprintAssetCategory", "Blueprints")				},
+				{	EAssetTypeCategories::MaterialsAndTextures,	LOCTEXT("MaterialAssetCategory", "Materials & Textures")	},
+				{	EAssetTypeCategories::Sounds,				LOCTEXT("SoundAssetCategory", "Sounds")						},
+				{	EAssetTypeCategories::Physics,				LOCTEXT("PhysicsAssetCategory", "Physics")					},
+				{	EAssetTypeCategories::UI,					LOCTEXT("UserInterfaceAssetCategory", "User Interface")		},
+				{	EAssetTypeCategories::Misc,					LOCTEXT("MiscellaneousAssetCategory", "Miscellaneous")		},
+				{	EAssetTypeCategories::Gameplay,				LOCTEXT("GameplayAssetCategory", "Gameplay")				},
+			};
+
+			for (const FAdvancedAssetCategory& AdvancedAssetCategory : AdvancedAssetCategories)
+			{
+				TArray<FFactoryItem> AnimationFactories = FindFactoriesInCategory(AdvancedAssetCategory.CategoryType);
+				if (AnimationFactories.Num() > 0)
+				{
+					MenuBuilder.AddSubMenu(
+						AdvancedAssetCategory.CategoryName,
+						FText::GetEmpty(),
+						FNewMenuDelegate::CreateStatic(
+							&FNewAssetOrClassContextMenu::CreateNewAssetMenuCategory, 
+							AdvancedAssetCategory.CategoryType, 
+							FirstSelectedPath, 
+							InOnNewAssetRequested, 
+							FCanExecuteAction() // We handle this at this level, rather than at the sub-menu item level
+							),
+						FUIAction(
+							FExecuteAction(),
+							CanExecuteAssetActionsDelegate
+							),
+						NAME_None,
+						EUserInterfaceActionType::Button
+						);
+				}
+			}
+		}
+		MenuBuilder.EndSection(); //ContentBrowserNewAdvancedAsset
+	}
+}
+
+void FNewAssetOrClassContextMenu::CreateNewAssetMenuCategory(FMenuBuilder& MenuBuilder, EAssetTypeCategories::Type AssetTypeCategory, FString InPath, FOnNewAssetRequested InOnNewAssetRequested, FCanExecuteAction InCanExecuteAction)
+{
+	// Find UFactory classes that can create new objects in this category.
+	TArray<FFactoryItem> FactoriesInThisCategory = FindFactoriesInCategory(AssetTypeCategory);
+
+	// Sort the list
+	struct FCompareFactoryDisplayNames
+	{
+		FORCEINLINE bool operator()( const FFactoryItem& A, const FFactoryItem& B ) const
+		{
+			return A.DisplayName.CompareToCaseIgnored(B.DisplayName) < 0;
+		}
+	};
+	FactoriesInThisCategory.Sort( FCompareFactoryDisplayNames() );
+
+	// Add menu entries for each one
+	for ( auto FactoryIt = FactoriesInThisCategory.CreateConstIterator(); FactoryIt; ++FactoryIt )
+	{
+		UFactory* Factory = (*FactoryIt).Factory;
+		TWeakObjectPtr<UClass> WeakFactoryClass = Factory->GetClass();
+
+		MenuBuilder.AddMenuEntry(
+			FUIAction(
+				FExecuteAction::CreateStatic( &FNewAssetOrClassContextMenu::ExecuteNewAsset, InPath, WeakFactoryClass, InOnNewAssetRequested ),
+				InCanExecuteAction
+				),
+			SNew( SFactoryMenuEntry, Factory )
+			);
+	}
+}
+
+void FNewAssetOrClassContextMenu::ExecuteImportAsset( FOnImportAssetRequested InOnInportAssetRequested, FString InPath )
+{
+	InOnInportAssetRequested.ExecuteIfBound( InPath );
+}
+
+void FNewAssetOrClassContextMenu::ExecuteNewAsset(FString InPath, TWeakObjectPtr<UClass> FactoryClass, FOnNewAssetRequested InOnNewAssetRequested)
+{
+	if ( ensure(FactoryClass.IsValid()) && ensure(!InPath.IsEmpty()) )
+	{
+		InOnNewAssetRequested.ExecuteIfBound(InPath, FactoryClass);
+	}
+}
+
+void FNewAssetOrClassContextMenu::ExecuteNewClass(FString InPath, FOnNewClassRequested InOnNewClassRequested)
+{
+	if(ensure(!InPath.IsEmpty()))
+	{
+		InOnNewClassRequested.ExecuteIfBound(InPath);
+	}
+}
+
+void FNewAssetOrClassContextMenu::ExecuteNewFolder(FString InPath, FOnNewFolderRequested InOnNewFolderRequested)
+{
+	if(ensure(!InPath.IsEmpty()) )
+	{
+		InOnNewFolderRequested.ExecuteIfBound(InPath);
+	}
+}
+
+void FNewAssetOrClassContextMenu::ExecuteGetContent( FOnGetContentRequested InOnGetContentRequested )
+{
+	InOnGetContentRequested.ExecuteIfBound();
+}
+
+#undef LOCTEXT_NAMESPACE
diff --git a/Engine/Source/Editor/ContentBrowser/Private/NewAssetOrClassContextMenu.h b/Engine/Source/Editor/ContentBrowser/Private/NewAssetOrClassContextMenu.h
new file mode 100644
--- /dev/null
+++ b/Engine/Source/Editor/ContentBrowser/Private/NewAssetOrClassContextMenu.h
@@ -0,0 +1,54 @@
+// Copyright 1998-2015 Epic Games, Inc. All Rights Reserved.
+
+#pragma once
+
+class FNewAssetOrClassContextMenu
+{
+public:
+	DECLARE_DELEGATE_TwoParams( FOnNewAssetRequested, const FString& /*SelectedPath*/, TWeakObjectPtr<UClass> /*FactoryClass*/ );
+	DECLARE_DELEGATE_OneParam( FOnNewClassRequested, const FString& /*SelectedPath*/ );
+	DECLARE_DELEGATE_OneParam( FOnNewFolderRequested, const FString& /*SelectedPath*/ );
+	DECLARE_DELEGATE_OneParam( FOnImportAssetRequested, const FString& );
+	DECLARE_DELEGATE( FOnGetContentRequested )
+
+	/** Makes the context menu widget */
+	static void MakeContextMenu(
+		FMenuBuilder& MenuBuilder, 
+		const TArray<FName>& InSelectedPaths, 
+		const FOnNewAssetRequested& InOnNewAssetRequested, 
+		const FOnNewClassRequested& InOnNewClassRequested, 
+		const FOnNewFolderRequested& InOnNewFolderRequested, 
+		const FOnImportAssetRequested& InOnImportAssetRequested, 
+		const FOnGetContentRequested& InOnGetContentRequested
+		);
+
+	/** Makes the context menu widget */
+	static void MakeContextMenu(
+		FMenuBuilder& MenuBuilder, 
+		const TArray<FString>& InSelectedPaths, 
+		const FOnNewAssetRequested& InOnNewAssetRequested, 
+		const FOnNewClassRequested& InOnNewClassRequested, 
+		const FOnNewFolderRequested& InOnNewFolderRequested, 
+		const FOnImportAssetRequested& InOnImportAssetRequested, 
+		const FOnGetContentRequested& InOnGetContentRequested
+		);
+
+private:
+	/** Handle creating a new asset from an asset category */
+	static void CreateNewAssetMenuCategory(FMenuBuilder& MenuBuilder, EAssetTypeCategories::Type AssetTypeCategory, FString InPath, FOnNewAssetRequested InOnNewAssetRequested, FCanExecuteAction InCanExecuteAction);
+
+	/** Handle when the "Import" button is clicked */
+	static void ExecuteImportAsset(FOnImportAssetRequested InOnImportAssetRequested, FString InPath);
+
+	/** Create a new asset using the specified factory at the specified path */
+	static void ExecuteNewAsset(FString InPath, TWeakObjectPtr<UClass> FactoryClass, FOnNewAssetRequested InOnNewAssetRequested);
+
+	/** Create a new class at the specified path */
+	static void ExecuteNewClass(FString InPath, FOnNewClassRequested InOnNewClassRequested);
+
+	/** Create a new folder at the specified path */
+	static void ExecuteNewFolder(FString InPath, FOnNewFolderRequested InOnNewFolderRequested);
+
+	/** Handle when the "Get Content" button is clicked */
+	static void ExecuteGetContent( FOnGetContentRequested InOnGetContentRequested );
+};
diff --git a/Engine/Source/Editor/ContentBrowser/Private/PathContextMenu.cpp b/Engine/Source/Editor/ContentBrowser/Private/PathContextMenu.cpp
--- a/Engine/Source/Editor/ContentBrowser/Private/PathContextMenu.cpp
+++ b/Engine/Source/Editor/ContentBrowser/Private/PathContextMenu.cpp
@@ -11,6 +11,7 @@
 #include "Editor/UnrealEd/Public/PackageTools.h"
 #include "SColorPicker.h"
 #include "GenericCommands.h"
+#include "NativeClassHierarchy.h"
 
 
 #define LOCTEXT_NAMESPACE "ContentBrowser"
@@ -25,12 +26,17 @@ FPathContextMenu::FPathContextMenu(const TWeakPtr<SWidget>& InParentContent)
 	
 }
 
-void FPathContextMenu::SetOnNewAssetRequested(const FNewAssetContextMenu::FOnNewAssetRequested& InOnNewAssetRequested)
+void FPathContextMenu::SetOnNewAssetRequested(const FNewAssetOrClassContextMenu::FOnNewAssetRequested& InOnNewAssetRequested)
 {
 	OnNewAssetRequested = InOnNewAssetRequested;
 }
 
-void FPathContextMenu::SetOnImportAssetRequested( const FNewAssetContextMenu::FOnImportAssetRequested& InOnImportAssetRequested )
+void FPathContextMenu::SetOnNewClassRequested(const FNewAssetOrClassContextMenu::FOnNewClassRequested& InOnNewClassRequested)
+{
+	OnNewClassRequested = InOnNewClassRequested;
+}
+
+void FPathContextMenu::SetOnImportAssetRequested( const FNewAssetOrClassContextMenu::FOnImportAssetRequested& InOnImportAssetRequested )
 {
 	OnImportAssetRequested = InOnImportAssetRequested;
 }
@@ -77,22 +83,83 @@ TSharedRef<FExtender> FPathContextMenu::MakePathViewContextMenuExtender(const TA
 
 void FPathContextMenu::MakePathViewContextMenu(FMenuBuilder& MenuBuilder)
 {
-	const FString& SelectedPath = GetFirstSelectedPath();
+	int32 NumAssetPaths, NumClassPaths;
+	ContentBrowserUtils::CountPathTypes(SelectedPaths, NumAssetPaths, NumClassPaths);
 
 	// Only add something if at least one folder is selected
-	if ( SelectedPath.Len() && SelectedPath != TEXT("/Classes") )
+	if ( SelectedPaths.Num() > 0 )
 	{
+		const bool bHasAssetPaths = NumAssetPaths > 0;
+		const bool bHasClassPaths = NumClassPaths > 0;
+
 		// Common operations section //
 		MenuBuilder.BeginSection("PathViewFolderOptions", LOCTEXT("PathViewOptionsMenuHeading", "Folder Options") );
 		{
-			// New Asset (submenu)
-			MenuBuilder.AddSubMenu(
-				LOCTEXT( "CreateNewAsset", "Add New" ),
-				LOCTEXT( "CreateNewAssetTooltip", "Create a new asset." ),
-				FNewMenuDelegate::CreateRaw( this, &FPathContextMenu::MakeNewAssetSubMenu ),
-				false,
-				FSlateIcon( FEditorStyle::GetStyleSetName(), "ContentBrowser.PathActions.NewAsset" )
-				);
+			if(bHasAssetPaths)
+			{
+				FText NewAssetToolTip;
+				if(SelectedPaths.Num() == 1)
+				{
+					if(CanCreateAsset())
+					{
+						NewAssetToolTip = FText::Format(LOCTEXT("NewAssetTooltip_CreateIn", "Create a new asset in {0}."), FText::FromString(SelectedPaths[0]));
+					}
+					else
+					{
+						NewAssetToolTip = FText::Format(LOCTEXT("NewAssetTooltip_InvalidPath", "Cannot create new assets in {0}."), FText::FromString(SelectedPaths[0]));
+					}
+				}
+				else
+				{
+					NewAssetToolTip = LOCTEXT("NewAssetTooltip_InvalidNumberOfPaths", "Can only create assets when there is a single path selected.");
+				}
+
+				// New Asset (submenu)
+				MenuBuilder.AddSubMenu(
+					LOCTEXT( "NewAssetLabel", "New Asset" ),
+					NewAssetToolTip,
+					FNewMenuDelegate::CreateRaw( this, &FPathContextMenu::MakeNewAssetSubMenu ),
+					FUIAction(
+						FExecuteAction(),
+						FCanExecuteAction::CreateRaw( this, &FPathContextMenu::CanCreateAsset )
+						),
+					NAME_None,
+					EUserInterfaceActionType::Button,
+					false,
+					FSlateIcon( FEditorStyle::GetStyleSetName(), "ContentBrowser.PathActions.NewAsset" )
+					);
+			}
+
+			if(bHasClassPaths)
+			{
+				FText NewClassToolTip;
+				if(SelectedPaths.Num() == 1)
+				{
+					if(CanCreateClass())
+					{
+						NewClassToolTip = FText::Format(LOCTEXT("NewClassTooltip_CreateIn", "Create a new class in {0}."), FText::FromString(SelectedPaths[0]));
+					}
+					else
+					{
+						NewClassToolTip = FText::Format(LOCTEXT("NewClassTooltip_InvalidPath", "Cannot create new classes in {0}."), FText::FromString(SelectedPaths[0]));
+					}
+				}
+				else
+				{
+					NewClassToolTip = LOCTEXT("NewClassTooltip_InvalidNumberOfPaths", "Can only create classes when there is a single path selected.");
+				}
+
+				// New Class
+				MenuBuilder.AddMenuEntry(
+					LOCTEXT("NewClassLabel", "New Class"),
+					NewClassToolTip,
+					FSlateIcon(FEditorStyle::GetStyleSetName(), "MainFrame.AddCodeToProject"),
+					FUIAction(
+						FExecuteAction::CreateRaw( this, &FPathContextMenu::ExecuteCreateClass ),
+						FCanExecuteAction::CreateRaw( this, &FPathContextMenu::CanCreateClass )
+						)
+					);
+			}
 
 			// Explore
 			MenuBuilder.AddMenuEntry(
@@ -132,126 +199,159 @@ void FPathContextMenu::MakePathViewContextMenu(FMenuBuilder& MenuBuilder)
 		}
 		MenuBuilder.EndSection();
 
-		// Bulk operations section //
-		MenuBuilder.BeginSection("PathContextBulkOperations", LOCTEXT("AssetTreeBulkMenuHeading", "Bulk Operations") );
+		if(bHasAssetPaths)
 		{
-		    // Save
-		    MenuBuilder.AddMenuEntry(
-			    LOCTEXT("SaveFolder", "Save All"),
-			    LOCTEXT("SaveFolderTooltip", "Saves all modified assets in this folder."),
-			    FSlateIcon(FEditorStyle::GetStyleSetName(), "MainFrame.SaveAll"),
-			    FUIAction( FExecuteAction::CreateSP( this, &FPathContextMenu::ExecuteSaveFolder ) )
-			    );
+			// Bulk operations section //
+			MenuBuilder.BeginSection("PathContextBulkOperations", LOCTEXT("AssetTreeBulkMenuHeading", "Bulk Operations") );
+			{
+				// Save
+				MenuBuilder.AddMenuEntry(
+					LOCTEXT("SaveFolder", "Save All"),
+					LOCTEXT("SaveFolderTooltip", "Saves all modified assets in this folder."),
+					FSlateIcon(FEditorStyle::GetStyleSetName(), "MainFrame.SaveAll"),
+					FUIAction( FExecuteAction::CreateSP( this, &FPathContextMenu::ExecuteSaveFolder ) )
+					);
     
-		    // Delete
-			MenuBuilder.AddMenuEntry(FGenericCommands::Get().Delete, NAME_None,
-			    LOCTEXT("DeleteFolder", "Delete"),
-			    LOCTEXT("DeleteFolderTooltip", "Removes this folder and all assets it contains."),
-				FSlateIcon( FEditorStyle::GetStyleSetName(), "ContentBrowser.AssetActions.Delete" )
-			    );
-
-			// Reference Viewer
-			MenuBuilder.AddMenuEntry(
-				LOCTEXT("ReferenceViewer", "Reference Viewer..."),
-				LOCTEXT("ReferenceViewerOnFolderTooltip", "Shows a graph of references for this folder."),
-				FSlateIcon(),
-				FUIAction( FExecuteAction::CreateSP( this, &FPathContextMenu::ExecuteReferenceViewer ) )
-				);
+				// Delete
+				MenuBuilder.AddMenuEntry(FGenericCommands::Get().Delete, NAME_None,
+					LOCTEXT("DeleteFolder", "Delete"),
+					LOCTEXT("DeleteFolderTooltip", "Removes this folder and all assets it contains."),
+					FSlateIcon( FEditorStyle::GetStyleSetName(), "ContentBrowser.AssetActions.Delete" )
+					);
+
+				// Reference Viewer
+				MenuBuilder.AddMenuEntry(
+					LOCTEXT("ReferenceViewer", "Reference Viewer..."),
+					LOCTEXT("ReferenceViewerOnFolderTooltip", "Shows a graph of references for this folder."),
+					FSlateIcon(),
+					FUIAction( FExecuteAction::CreateSP( this, &FPathContextMenu::ExecuteReferenceViewer ) )
+					);
     
-			// Fix Up Redirectors in Folder
-			MenuBuilder.AddMenuEntry(
-				LOCTEXT("FixUpRedirectorsInFolder", "Fix Up Redirectors in Folder"),
-				LOCTEXT("FixUpRedirectorsInFolderTooltip", "Finds referencers to all redirectors in the selected folders and resaves them if possible, then deletes any redirectors that had all their referencers fixed."),
-				FSlateIcon(),
-				FUIAction( FExecuteAction::CreateSP( this, &FPathContextMenu::ExecuteFixUpRedirectorsInFolder ) )
-				);
+				// Fix Up Redirectors in Folder
+				MenuBuilder.AddMenuEntry(
+					LOCTEXT("FixUpRedirectorsInFolder", "Fix Up Redirectors in Folder"),
+					LOCTEXT("FixUpRedirectorsInFolderTooltip", "Finds referencers to all redirectors in the selected folders and resaves them if possible, then deletes any redirectors that had all their referencers fixed."),
+					FSlateIcon(),
+					FUIAction( FExecuteAction::CreateSP( this, &FPathContextMenu::ExecuteFixUpRedirectorsInFolder ) )
+					);
 
-		    if ( SelectedPaths.Num() == 1 )
-		    {
-			    // Migrate Folder
-			    MenuBuilder.AddMenuEntry(
-				    LOCTEXT("MigrateFolder", "Migrate..."),
-				    LOCTEXT("MigrateFolderTooltip", "Copies assets found in this folder and their dependencies to another game content folder."),
-				    FSlateIcon(),
-				    FUIAction( FExecuteAction::CreateSP( this, &FPathContextMenu::ExecuteMigrateFolder ) )
-				    );
-		    }
+				if ( NumAssetPaths == 1 && NumClassPaths == 0 )
+				{
+					// Migrate Folder
+					MenuBuilder.AddMenuEntry(
+						LOCTEXT("MigrateFolder", "Migrate..."),
+						LOCTEXT("MigrateFolderTooltip", "Copies assets found in this folder and their dependencies to another game content folder."),
+						FSlateIcon(),
+						FUIAction( FExecuteAction::CreateSP( this, &FPathContextMenu::ExecuteMigrateFolder ) )
+						);
+				}
+			}
+			MenuBuilder.EndSection();
 		}
-		MenuBuilder.EndSection();
-
-		// Source control section //
-		MenuBuilder.BeginSection("PathContextSourceControl", LOCTEXT("AssetTreeSCCMenuHeading", "Source Control") );
 
-		ISourceControlProvider& SourceControlProvider = ISourceControlModule::Get().GetProvider();
-		if ( SourceControlProvider.IsEnabled() )
+		if(bHasAssetPaths)
 		{
-			// Check out
-			MenuBuilder.AddMenuEntry(
-				LOCTEXT("FolderSCCCheckOut", "Check Out"),
-				LOCTEXT("FolderSCCCheckOutTooltip", "Checks out all assets from source control which are in this folder."),
-				FSlateIcon(),
-				FUIAction(
-					FExecuteAction::CreateSP( this, &FPathContextMenu::ExecuteSCCCheckOut ),
-					FCanExecuteAction::CreateSP( this, &FPathContextMenu::CanExecuteSCCCheckOut )
-					)
-				);
+			// Source control section //
+			MenuBuilder.BeginSection("PathContextSourceControl", LOCTEXT("AssetTreeSCCMenuHeading", "Source Control") );
 
-			// Open for Add
-			MenuBuilder.AddMenuEntry(
-				LOCTEXT("FolderSCCOpenForAdd", "Mark For Add"),
-				LOCTEXT("FolderSCCOpenForAddTooltip", "Adds all assets to source control that are in this folder and not already added."),
-				FSlateIcon(),
-				FUIAction(
-					FExecuteAction::CreateSP( this, &FPathContextMenu::ExecuteSCCOpenForAdd ),
-					FCanExecuteAction::CreateSP( this, &FPathContextMenu::CanExecuteSCCOpenForAdd )
-					)
-				);
+			ISourceControlProvider& SourceControlProvider = ISourceControlModule::Get().GetProvider();
+			if ( SourceControlProvider.IsEnabled() )
+			{
+				// Check out
+				MenuBuilder.AddMenuEntry(
+					LOCTEXT("FolderSCCCheckOut", "Check Out"),
+					LOCTEXT("FolderSCCCheckOutTooltip", "Checks out all assets from source control which are in this folder."),
+					FSlateIcon(),
+					FUIAction(
+						FExecuteAction::CreateSP( this, &FPathContextMenu::ExecuteSCCCheckOut ),
+						FCanExecuteAction::CreateSP( this, &FPathContextMenu::CanExecuteSCCCheckOut )
+						)
+					);
 
-			// Check in
-			MenuBuilder.AddMenuEntry(
-				LOCTEXT("FolderSCCCheckIn", "Check In"),
-				LOCTEXT("FolderSCCCheckInTooltip", "Checks in all assets to source control which are in this folder."),
-				FSlateIcon(),
-				FUIAction(
-					FExecuteAction::CreateSP( this, &FPathContextMenu::ExecuteSCCCheckIn ),
-					FCanExecuteAction::CreateSP( this, &FPathContextMenu::CanExecuteSCCCheckIn )
-					)
-				);
+				// Open for Add
+				MenuBuilder.AddMenuEntry(
+					LOCTEXT("FolderSCCOpenForAdd", "Mark For Add"),
+					LOCTEXT("FolderSCCOpenForAddTooltip", "Adds all assets to source control that are in this folder and not already added."),
+					FSlateIcon(),
+					FUIAction(
+						FExecuteAction::CreateSP( this, &FPathContextMenu::ExecuteSCCOpenForAdd ),
+						FCanExecuteAction::CreateSP( this, &FPathContextMenu::CanExecuteSCCOpenForAdd )
+						)
+					);
 
-			// Sync
-			MenuBuilder.AddMenuEntry(
-				LOCTEXT("FolderSCCSync", "Sync"),
-				LOCTEXT("FolderSCCSyncTooltip", "Syncs all the assets in this folder to the latest version."),
-				FSlateIcon(),
-				FUIAction( FExecuteAction::CreateSP( this, &FPathContextMenu::ExecuteSCCSync ) )
-				);
-		}
-		else
-		{
-			MenuBuilder.AddMenuEntry(
-				LOCTEXT("FolderSCCConnect", "Connect To Source Control"),
-				LOCTEXT("FolderSCCConnectTooltip", "Connect to source control to allow source control operations to be performed on content and levels."),
-				FSlateIcon(),
-				FUIAction(
-					FExecuteAction::CreateSP( this, &FPathContextMenu::ExecuteSCCConnect ),
-					FCanExecuteAction::CreateSP( this, &FPathContextMenu::CanExecuteSCCConnect )
-					)
-				);	
-		}
+				// Check in
+				MenuBuilder.AddMenuEntry(
+					LOCTEXT("FolderSCCCheckIn", "Check In"),
+					LOCTEXT("FolderSCCCheckInTooltip", "Checks in all assets to source control which are in this folder."),
+					FSlateIcon(),
+					FUIAction(
+						FExecuteAction::CreateSP( this, &FPathContextMenu::ExecuteSCCCheckIn ),
+						FCanExecuteAction::CreateSP( this, &FPathContextMenu::CanExecuteSCCCheckIn )
+						)
+					);
 
-		MenuBuilder.EndSection();
+				// Sync
+				MenuBuilder.AddMenuEntry(
+					LOCTEXT("FolderSCCSync", "Sync"),
+					LOCTEXT("FolderSCCSyncTooltip", "Syncs all the assets in this folder to the latest version."),
+					FSlateIcon(),
+					FUIAction(
+						FExecuteAction::CreateSP( this, &FPathContextMenu::ExecuteSCCSync ),
+						FCanExecuteAction::CreateSP( this, &FPathContextMenu::CanExecuteSCCSync )
+						)
+					);
+			}
+			else
+			{
+				MenuBuilder.AddMenuEntry(
+					LOCTEXT("FolderSCCConnect", "Connect To Source Control"),
+					LOCTEXT("FolderSCCConnectTooltip", "Connect to source control to allow source control operations to be performed on content and levels."),
+					FSlateIcon(),
+					FUIAction(
+						FExecuteAction::CreateSP( this, &FPathContextMenu::ExecuteSCCConnect ),
+						FCanExecuteAction::CreateSP( this, &FPathContextMenu::CanExecuteSCCConnect )
+						)
+					);	
+			}
+
+			MenuBuilder.EndSection();
+		}
 	}
 }
 
+bool FPathContextMenu::CanCreateAsset() const
+{
+	// We can only create assets when we have a single asset path selected
+	return SelectedPaths.Num() == 1 && !ContentBrowserUtils::IsClassPath(SelectedPaths[0]);
+}
+
 void FPathContextMenu::MakeNewAssetSubMenu(FMenuBuilder& MenuBuilder)
 {
-	const FString& SourcesPath = GetFirstSelectedPath();
-	if ( ensure(SourcesPath.Len()) )
+	if ( SelectedPaths.Num() )
 	{
-		FNewAssetContextMenu::MakeContextMenu(MenuBuilder, SourcesPath, OnNewAssetRequested, FNewAssetContextMenu::FOnNewFolderRequested(), OnImportAssetRequested, FNewAssetContextMenu::FOnGetContentRequested());
+		FNewAssetOrClassContextMenu::MakeContextMenu(
+			MenuBuilder, 
+			SelectedPaths, 
+			OnNewAssetRequested, 
+			FNewAssetOrClassContextMenu::FOnNewClassRequested(), 
+			FNewAssetOrClassContextMenu::FOnNewFolderRequested(), 
+			OnImportAssetRequested, 
+			FNewAssetOrClassContextMenu::FOnGetContentRequested()
+			);
 	}
 }
 
+void FPathContextMenu::ExecuteCreateClass()
+{
+	OnNewClassRequested.ExecuteIfBound(SelectedPaths[0]);
+}
+
+bool FPathContextMenu::CanCreateClass() const
+{
+	// We can only create assets when we have a single class path selected
+	return SelectedPaths.Num() == 1 && ContentBrowserUtils::IsValidPathToCreateNewClass(SelectedPaths[0]);
+}
+
 void FPathContextMenu::MakeSetColorSubMenu(FMenuBuilder& MenuBuilder)
 {
 	// New Color
@@ -339,25 +439,54 @@ void FPathContextMenu::ExecuteExplore()
 	for (int32 PathIdx = 0; PathIdx < SelectedPaths.Num(); ++PathIdx)
 	{
 		const FString& Path = SelectedPaths[PathIdx];
-		const FString FilePath = FPaths::ConvertRelativePathToFull(FPackageName::LongPackageNameToFilename(Path + TEXT("/")));
-
-		// If the folder has not yet been created, make is right before we try to explore to it
-		if ( !IFileManager::Get().DirectoryExists(*FilePath) )
+		FString FilePath;
+		if (ContentBrowserUtils::IsClassPath(Path))
+		{
+			TSharedRef<FNativeClassHierarchy> NativeClassHierarchy = FContentBrowserSingleton::Get().GetNativeClassHierarchy();
+			if (NativeClassHierarchy->GetFileSystemPath(Path, FilePath))
+			{
+				FilePath = IFileManager::Get().ConvertToAbsolutePathForExternalAppForRead(*FilePath);
+			}
+		}
+		else
 		{
-			IFileManager::Get().MakeDirectory(*FilePath, /*Tree=*/true);
+			FilePath = FPaths::ConvertRelativePathToFull(FPackageName::LongPackageNameToFilename(Path + TEXT("/")));
 		}
 
-		FPlatformProcess::ExploreFolder( *FilePath );
+		if (!FilePath.IsEmpty())
+		{
+			// If the folder has not yet been created, make is right before we try to explore to it
+			if (!IFileManager::Get().DirectoryExists(*FilePath))
+			{
+				IFileManager::Get().MakeDirectory(*FilePath, /*Tree=*/true);
+			}
+
+			FPlatformProcess::ExploreFolder(*FilePath);
+		}
 	}
 }
 
 bool FPathContextMenu::CanExecuteRename() const
 {
-	if (SelectedPaths.Num() == 1 && !ContentBrowserUtils::IsAssetRootDir(SelectedPaths[0]))
+	// We can't rename when we have more than one path selected
+	if (SelectedPaths.Num() != 1)
 	{
-		return true;
+		return false;
 	}
-	return false;
+
+	// We can't rename a root folder
+	if (ContentBrowserUtils::IsRootDir(SelectedPaths[0]))
+	{
+		return false;
+	}
+
+	// We can't rename *any* folders that belong to class roots
+	if (ContentBrowserUtils::IsClassPath(SelectedPaths[0]))
+	{
+		return false;
+	}
+
+	return true;
 }
 
 void FPathContextMenu::ExecuteRename()
@@ -479,7 +608,11 @@ void FPathContextMenu::ExecuteSaveFolder()
 
 bool FPathContextMenu::CanExecuteDelete() const
 {
-	return SelectedPaths.Num() > 0;
+	int32 NumAssetPaths, NumClassPaths;
+	ContentBrowserUtils::CountPathTypes(SelectedPaths, NumAssetPaths, NumClassPaths);
+
+	// We can't delete folders containing classes
+	return NumAssetPaths > 0 && NumClassPaths == 0;
 }
 
 void FPathContextMenu::ExecuteDelete()
@@ -769,22 +902,47 @@ void FPathContextMenu::ExecuteSCCConnect() const
 
 bool FPathContextMenu::CanExecuteSCCCheckOut() const
 {
-	return bCanExecuteSCCCheckOut;
+	int32 NumAssetPaths, NumClassPaths;
+	ContentBrowserUtils::CountPathTypes(SelectedPaths, NumAssetPaths, NumClassPaths);
+
+	// Can only perform SCC operations on asset paths
+	return bCanExecuteSCCCheckOut && (NumAssetPaths > 0 && NumClassPaths == 0);
 }
 
 bool FPathContextMenu::CanExecuteSCCOpenForAdd() const
 {
-	return bCanExecuteSCCOpenForAdd;
+	int32 NumAssetPaths, NumClassPaths;
+	ContentBrowserUtils::CountPathTypes(SelectedPaths, NumAssetPaths, NumClassPaths);
+
+	// Can only perform SCC operations on asset paths
+	return bCanExecuteSCCOpenForAdd && (NumAssetPaths > 0 && NumClassPaths == 0);
 }
 
 bool FPathContextMenu::CanExecuteSCCCheckIn() const
 {
-	return bCanExecuteSCCCheckIn;
+	int32 NumAssetPaths, NumClassPaths;
+	ContentBrowserUtils::CountPathTypes(SelectedPaths, NumAssetPaths, NumClassPaths);
+
+	// Can only perform SCC operations on asset paths
+	return bCanExecuteSCCCheckIn && (NumAssetPaths > 0 && NumClassPaths == 0);
+}
+
+bool FPathContextMenu::CanExecuteSCCSync() const
+{
+	int32 NumAssetPaths, NumClassPaths;
+	ContentBrowserUtils::CountPathTypes(SelectedPaths, NumAssetPaths, NumClassPaths);
+
+	// Can only perform SCC operations on asset paths
+	return (NumAssetPaths > 0 && NumClassPaths == 0);
 }
 
 bool FPathContextMenu::CanExecuteSCCConnect() const
 {
-	return !ISourceControlModule::Get().IsEnabled() || !ISourceControlModule::Get().GetProvider().IsAvailable();
+	int32 NumAssetPaths, NumClassPaths;
+	ContentBrowserUtils::CountPathTypes(SelectedPaths, NumAssetPaths, NumClassPaths);
+
+	// Can only perform SCC operations on asset paths
+	return (!ISourceControlModule::Get().IsEnabled() || !ISourceControlModule::Get().GetProvider().IsAvailable()) && (NumAssetPaths > 0 && NumClassPaths == 0);
 }
 
 void FPathContextMenu::CacheCanExecuteVars()
diff --git a/Engine/Source/Editor/ContentBrowser/Private/PathContextMenu.h b/Engine/Source/Editor/ContentBrowser/Private/PathContextMenu.h
--- a/Engine/Source/Editor/ContentBrowser/Private/PathContextMenu.h
+++ b/Engine/Source/Editor/ContentBrowser/Private/PathContextMenu.h
@@ -2,7 +2,7 @@
 
 #pragma once
 
-#include "NewAssetContextMenu.h"
+#include "NewAssetOrClassContextMenu.h"
 
 class FPathContextMenu : public TSharedFromThis<FPathContextMenu>
 {
@@ -11,10 +11,13 @@ public:
 	FPathContextMenu(const TWeakPtr<SWidget>& InParentContent);
 
 	/** Sets the handler for when new assets are requested */
-	void SetOnNewAssetRequested(const FNewAssetContextMenu::FOnNewAssetRequested& InOnNewAssetRequested);
+	void SetOnNewAssetRequested(const FNewAssetOrClassContextMenu::FOnNewAssetRequested& InOnNewAssetRequested);
+
+	/** Sets the handler for when new classes are requested */
+	void SetOnNewClassRequested(const FNewAssetOrClassContextMenu::FOnNewClassRequested& InOnNewClassRequested);
 
 	/** Sets the handler for when importing an asset is requested */
-	void SetOnImportAssetRequested( const FNewAssetContextMenu::FOnImportAssetRequested& InOnImportAssetRequested );
+	void SetOnImportAssetRequested( const FNewAssetOrClassContextMenu::FOnImportAssetRequested& InOnImportAssetRequested );
 
 	/** Delegate for when the context menu requests a rename of a folder */
 	DECLARE_DELEGATE_OneParam(FOnRenameFolderRequested, const FString& /*FolderToRename*/);
@@ -33,9 +36,18 @@ public:
 	/** Makes the asset tree context menu widget */
 	void MakePathViewContextMenu(FMenuBuilder& MenuBuilder);
 
+	/** Handler to check to see if creating a new asset is allowed */
+	bool CanCreateAsset() const;
+
 	/** Makes the new asset submenu */
 	void MakeNewAssetSubMenu(FMenuBuilder& MenuBuilder);
 
+	/** Handler for when "New Class" is selected */
+	void ExecuteCreateClass();
+
+	/** Handler to check to see if creating a new class is allowed */
+	bool CanCreateClass() const;
+
 	/** Makes the set color submenu */
 	void MakeSetColorSubMenu(FMenuBuilder& MenuBuilder);
 
@@ -99,6 +111,9 @@ public:
 	/** Handler to check to see if "Checkin to source control" can be executed */
 	bool CanExecuteSCCCheckIn() const;
 
+	/** Handler to check to see if "Sync" can be executed */
+	bool CanExecuteSCCSync() const;
+
 	/** Handler to check to see if "Connect to source control" can be executed */
 	bool CanExecuteSCCConnect() const;	
 
@@ -127,8 +142,9 @@ private:
 private:
 	TArray<FString> SelectedPaths;
 	TWeakPtr<SWidget> ParentContent;
-	FNewAssetContextMenu::FOnNewAssetRequested OnNewAssetRequested;
-	FNewAssetContextMenu::FOnImportAssetRequested OnImportAssetRequested;
+	FNewAssetOrClassContextMenu::FOnNewAssetRequested OnNewAssetRequested;
+	FNewAssetOrClassContextMenu::FOnNewClassRequested OnNewClassRequested;
+	FNewAssetOrClassContextMenu::FOnImportAssetRequested OnImportAssetRequested;
 	FOnRenameFolderRequested OnRenameFolderRequested;
 	FOnFolderDeleted OnFolderDeleted;
 
diff --git a/Engine/Source/Editor/ContentBrowser/Private/PathViewTypes.h b/Engine/Source/Editor/ContentBrowser/Private/PathViewTypes.h
--- a/Engine/Source/Editor/ContentBrowser/Private/PathViewTypes.h
+++ b/Engine/Source/Editor/ContentBrowser/Private/PathViewTypes.h
@@ -9,11 +9,13 @@ struct FTreeItem : public TSharedFromThis<FTreeItem>
 	{
 		FORCEINLINE bool operator()( TSharedPtr<FTreeItem> A, TSharedPtr<FTreeItem> B ) const
 		{
-			return A->FolderName < B->FolderName;
+			return A->DisplayName.ToString() < B->DisplayName.ToString();
 		}
 	};
 
-	/** The name of the tree item without the path */
+	/** The display name of the tree item (typically the same as FolderName, but may be localized for known folder types) */
+	FText DisplayName;
+	/** The leaf-name of the tree item folder */
 	FString FolderName;
 	/** The path of the tree item including the name */
 	FString FolderPath;
@@ -33,11 +35,12 @@ struct FTreeItem : public TSharedFromThis<FTreeItem>
 	FRenamedRequestEvent OnRenamedRequestEvent;
 
 	/** Constructor */
-	FTreeItem(const FString& InFolderName, const FString& InFolderPath, TSharedPtr<FTreeItem> InParent, bool InNamingFolder = false)
-		: FolderName(InFolderName)
-		, FolderPath(InFolderPath)
+	FTreeItem(FText InDisplayName, FString InFolderName, FString InFolderPath, TSharedPtr<FTreeItem> InParent, bool InNamingFolder = false)
+		: DisplayName(MoveTemp(InDisplayName))
+		, FolderName(MoveTemp(InFolderName))
+		, FolderPath(MoveTemp(InFolderPath))
 		, bNamingFolder(InNamingFolder)
-		, Parent(InParent)
+		, Parent(MoveTemp(InParent))
 	{}
 
 	/** Returns true if this item is a child of the specified item */
@@ -96,4 +99,4 @@ struct FTreeItem : public TSharedFromThis<FTreeItem>
 	{
 		Children.Sort( FCompareFTreeItemByName() );
 	}
-};
\ No newline at end of file
+};
diff --git a/Engine/Source/Editor/ContentBrowser/Private/SAssetView.cpp b/Engine/Source/Editor/ContentBrowser/Private/SAssetView.cpp
--- a/Engine/Source/Editor/ContentBrowser/Private/SAssetView.cpp
+++ b/Engine/Source/Editor/ContentBrowser/Private/SAssetView.cpp
@@ -13,6 +13,7 @@
 #include "ObjectTools.h"
 #include "KismetEditorUtilities.h"
 #include "IPluginManager.h"
+#include "NativeClassHierarchy.h"
 
 #define LOCTEXT_NAMESPACE "ContentBrowser"
 
@@ -44,6 +45,13 @@ SAssetView::~SAssetView()
 	FCoreUObjectDelegates::OnAssetLoaded.RemoveAll(this);
 	FCoreUObjectDelegates::OnObjectPropertyChanged.RemoveAll(this);
 
+	// Unsubscribe from class events
+	if ( bCanShowClasses )
+	{
+		TSharedRef<FNativeClassHierarchy> NativeClassHierarchy = FContentBrowserSingleton::Get().GetNativeClassHierarchy();
+		NativeClassHierarchy->OnClassHierarchyUpdated().RemoveAll( this );
+	}
+
 	// Remove the listener for when view settings are changed
 	UContentBrowserSettings::OnSettingChanged().RemoveAll(this);
 
@@ -88,6 +96,13 @@ void SAssetView::Construct( const FArguments& InArgs )
 	FCoreUObjectDelegates::OnAssetLoaded.AddSP(this, &SAssetView::OnAssetLoaded);
 	FCoreUObjectDelegates::OnObjectPropertyChanged.AddSP(this, &SAssetView::OnObjectPropertyChanged);
 
+	// Listen to find out when the available classes are changed, so that we can refresh our paths
+	if ( bCanShowClasses )
+	{
+		TSharedRef<FNativeClassHierarchy> NativeClassHierarchy = FContentBrowserSingleton::Get().GetNativeClassHierarchy();
+		NativeClassHierarchy->OnClassHierarchyUpdated().AddSP( this, &SAssetView::OnClassHierarchyUpdated );
+	}
+
 	// Listen for when view settings are changed
 	UContentBrowserSettings::OnSettingChanged().AddSP(this, &SAssetView::HandleSettingChanged);
 
@@ -398,7 +413,11 @@ const FSourcesData& SAssetView::GetSourcesData() const
 
 bool SAssetView::IsAssetPathSelected() const
 {
-	return SourcesData.PackagePaths.Num() > 0 && !SourcesData.PackagePaths[0].ToString().StartsWith(TEXT("/Classes"));
+	int32 NumAssetPaths, NumClassPaths;
+	ContentBrowserUtils::CountPathTypes(SourcesData.PackagePaths, NumAssetPaths, NumClassPaths);
+
+	// Check that only asset paths are selected
+	return NumAssetPaths > 0 && NumClassPaths == 0;
 }
 
 void SAssetView::SetBackendFilter(const FARFilter& InBackendFilter)
@@ -1285,12 +1304,13 @@ void SAssetView::RefreshSourceItems()
 
 	const bool bShowAll = SourcesData.IsEmpty() && BackendFilter.IsEmpty();
 
-	bool bWantToShowShowClasses = false;
+	bool bShowClasses = false;
+	TArray<FName> ClassPathsToShow;
 
 	if ( bShowAll )
 	{
 		AssetRegistryModule.Get().GetAllAssets(Items);
-		bWantToShowShowClasses = true;
+		bShowClasses = true;
 		bWereItemsRecursivelyFiltered = true;
 	}
 	else
@@ -1303,8 +1323,17 @@ void SAssetView::RefreshSourceItems()
 
 		bWereItemsRecursivelyFiltered = bRecurse;
 
-		// Remove the classes path if it is in the list. We will add classes to the results later
-		bWantToShowShowClasses = Filter.PackagePaths.Remove(TEXT("/Classes")) > 0;
+		// Move any class paths into their own array
+		Filter.PackagePaths.RemoveAll([&ClassPathsToShow](const FName& PackagePath) -> bool
+		{
+			if(ContentBrowserUtils::IsClassPath(PackagePath.ToString()))
+			{
+				ClassPathsToShow.Add(PackagePath);
+				return true;
+			}
+			return false;
+		});
+		bShowClasses = ClassPathsToShow.Num() > 0;
 
 		if ( SourcesData.Collections.Num() > 0 && Filter.ObjectPaths.Num() == 0 )
 		{
@@ -1338,42 +1367,21 @@ void SAssetView::RefreshSourceItems()
 	}
 
 	// If we are showing classes in the asset list...
-	if (bWantToShowShowClasses && bCanShowClasses)
+	if (bShowClasses && bCanShowClasses)
 	{
-		// Make a map of UClasses to ActorFactories that support them
-		const TArray< UActorFactory *>& ActorFactories = GEditor->ActorFactories;
-		TMap<UClass*, UActorFactory*> ActorFactoryMap;
-		for ( int32 FactoryIdx = 0; FactoryIdx < ActorFactories.Num(); ++FactoryIdx )
-		{
-			UActorFactory* ActorFactory = ActorFactories[FactoryIdx];
+		// Load the native class hierarchy
+		TSharedRef<FNativeClassHierarchy> NativeClassHierarchy = FContentBrowserSingleton::Get().GetNativeClassHierarchy();
 
-			if ( ActorFactory )
-			{
-				ActorFactoryMap.Add(ActorFactory->GetDefaultActorClass( FAssetData() ), ActorFactory);
-			}
-		}
+		FNativeClassHierarchyFilter ClassFilter;
+		ClassFilter.ClassPaths = ClassPathsToShow;
+		ClassFilter.bRecursivePaths = ShouldFilterRecursively() || !IsShowingFolders() || !ClassPathsToShow.Num();
 
-		// Add loaded classes
-		FText UnusedErrorMessage;
-		FAssetData NoAssetData;
-		for (TObjectIterator<UClass> ClassIt; ClassIt; ++ClassIt)
+		// Find all the classes that match the current criteria
+		TArray<UClass*> MatchingClasses;
+		NativeClassHierarchy->GetMatchingClasses(ClassFilter, MatchingClasses);
+		for(UClass* CurrentClass : MatchingClasses)
 		{
-			// Don't offer skeleton classes
-			bool bIsSkeletonClass = FKismetEditorUtilities::IsClassABlueprintSkeleton(*ClassIt);
-
-			if ( !ClassIt->HasAllClassFlags(CLASS_NotPlaceable) &&
-				!ClassIt->HasAnyClassFlags(CLASS_Abstract | CLASS_Deprecated | CLASS_NewerVersionExists) &&
-				ClassIt->IsChildOf(AActor::StaticClass()) &&
-				(!ClassIt->IsChildOf(ABrush::StaticClass()) || ClassIt->IsChildOf(AVolume::StaticClass())) &&
-				!bIsSkeletonClass )
-			{
-				UActorFactory** ActorFactory = ActorFactoryMap.Find(*ClassIt);
-
-				if ( !ActorFactory || (*ActorFactory)->CanCreateActorFrom( NoAssetData, UnusedErrorMessage ) )
-				{
-					Items.Add(FAssetData(*ClassIt));
-				}
-			}
+			Items.Add(FAssetData(CurrentClass));
 		}
 	}
 
@@ -1604,32 +1612,76 @@ void SAssetView::RefreshFilteredItems()
 
 void SAssetView::RefreshFolders()
 {
-	if(IsShowingFolders() && !ShouldFilterRecursively())
+	if(!IsShowingFolders() || ShouldFilterRecursively())
 	{
-		const bool bDisplayDev = GetDefault<UContentBrowserSettings>()->GetDisplayDevelopersFolder();
-		FAssetRegistryModule& AssetRegistryModule = FModuleManager::LoadModuleChecked<FAssetRegistryModule>(TEXT("AssetRegistry"));
-		for(auto SourcePathIt(SourcesData.PackagePaths.CreateConstIterator()); SourcePathIt; SourcePathIt++)
-		{		
-			TArray<FString> SubPaths;
-			AssetRegistryModule.Get().GetSubPaths((*SourcePathIt).ToString(), SubPaths, false);
-			for(auto SubPathIt(SubPaths.CreateConstIterator()); SubPathIt; SubPathIt++)	
+		return;
+	}
+	
+	// Split the selected paths into asset and class paths
+	TArray<FName> AssetPathsToShow;
+	TArray<FName> ClassPathsToShow;
+	for(const FName& PackagePath : SourcesData.PackagePaths)
+	{
+		if(ContentBrowserUtils::IsClassPath(PackagePath.ToString()))
+		{
+			ClassPathsToShow.Add(PackagePath);
+		}
+		else
+		{
+			AssetPathsToShow.Add(PackagePath);
+		}
+	}
+
+	TArray<FString> FoldersToAdd;
+
+	const bool bDisplayDev = GetDefault<UContentBrowserSettings>()->GetDisplayDevelopersFolder();
+	FAssetRegistryModule& AssetRegistryModule = FModuleManager::LoadModuleChecked<FAssetRegistryModule>(TEXT("AssetRegistry"));
+	for(const FName& PackagePath : AssetPathsToShow)
+	{		
+		TArray<FString> SubPaths;
+		AssetRegistryModule.Get().GetSubPaths(PackagePath.ToString(), SubPaths, false);
+		for(const FString& SubPath : SubPaths)
+		{
+			// If this is a developer folder, and we don't want to show them try the next path
+			if(!bDisplayDev && ContentBrowserUtils::IsDevelopersFolder(SubPath))
 			{
-				// If this is a developer folder, and we don't want to show them try the next path
-				if ( !bDisplayDev && ContentBrowserUtils::IsDevelopersFolder(*SubPathIt) )
-				{
-					continue;
-				}
+				continue;
+			}
 
-				if(!Folders.Contains(*SubPathIt))
-				{
-					FilteredAssetItems.Add(MakeShareable(new FAssetViewFolder(*SubPathIt)));
-					RefreshList();
-					Folders.Add(*SubPathIt);
-					bPendingSortFilteredItems = true;
-				}
+			if(!Folders.Contains(SubPath))
+			{
+				FoldersToAdd.Add(SubPath);
 			}
 		}
 	}
+
+	// If we are showing classes in the asset list then we need to show their folders too
+	if(bCanShowClasses && ClassPathsToShow.Num() > 0)
+	{
+		// Load the native class hierarchy
+		TSharedRef<FNativeClassHierarchy> NativeClassHierarchy = FContentBrowserSingleton::Get().GetNativeClassHierarchy();
+
+		FNativeClassHierarchyFilter ClassFilter;
+		ClassFilter.ClassPaths = ClassPathsToShow;
+		ClassFilter.bRecursivePaths = false;
+
+		// Find all the classes that match the current criteria
+		TArray<FString> MatchingFolders;
+		NativeClassHierarchy->GetMatchingFolders(ClassFilter, MatchingFolders);
+		FoldersToAdd.Append(MatchingFolders);
+	}
+
+	if(FoldersToAdd.Num() > 0)
+	{
+		for(const FString& FolderPath : FoldersToAdd)
+		{
+			FilteredAssetItems.Add(MakeShareable(new FAssetViewFolder(FolderPath)));
+			Folders.Add(FolderPath);
+		}
+
+		RefreshList();
+		bPendingSortFilteredItems = true;
+	}
 }
 
 void SAssetView::SetMajorityAssetType(FName NewMajorityAssetType)
@@ -2038,6 +2090,12 @@ void SAssetView::OnObjectPropertyChanged(UObject* Asset, FPropertyChangedEvent&
 	}
 }
 
+void SAssetView::OnClassHierarchyUpdated()
+{
+	// The class hierarchy has changed in some way, so we need to refresh our backend list
+	RequestSlowFullListRefresh();
+}
+
 void SAssetView::OnFrontendFiltersChanged()
 {
 	RequestQuickFrontendListRefresh();
diff --git a/Engine/Source/Editor/ContentBrowser/Private/SAssetView.h b/Engine/Source/Editor/ContentBrowser/Private/SAssetView.h
--- a/Engine/Source/Editor/ContentBrowser/Private/SAssetView.h
+++ b/Engine/Source/Editor/ContentBrowser/Private/SAssetView.h
@@ -313,6 +313,9 @@ private:
 	/** Handler for when an asset's property has changed */
 	void OnObjectPropertyChanged(UObject* Asset, FPropertyChangedEvent& PropertyChangedEvent);
 
+	/** Called when the class hierarchy is updated due to the available modules changing */
+	void OnClassHierarchyUpdated();
+
 	/** Handler for when any frontend filters have been changed */
 	void OnFrontendFiltersChanged();
 
diff --git a/Engine/Source/Editor/ContentBrowser/Private/SContentBrowser.cpp b/Engine/Source/Editor/ContentBrowser/Private/SContentBrowser.cpp
--- a/Engine/Source/Editor/ContentBrowser/Private/SContentBrowser.cpp
+++ b/Engine/Source/Editor/ContentBrowser/Private/SContentBrowser.cpp
@@ -3,7 +3,7 @@
 
 #include "ContentBrowserPCH.h"
 #include "AssetContextMenu.h"
-#include "NewAssetContextMenu.h"
+#include "NewAssetOrClassContextMenu.h"
 #include "PathContextMenu.h"
 #include "ContentBrowserModule.h"
 #include "ContentBrowserCommands.h"
@@ -13,6 +13,8 @@
 #include "GenericCommands.h"
 #include "IAddContentDialogModule.h"
 #include "Engine/Selection.h"
+#include "NativeClassHierarchy.h"
+#include "GameProjectGenerationModule.h"
 
 #define LOCTEXT_NAMESPACE "ContentBrowser"
 
@@ -56,8 +58,9 @@ void SContentBrowser::Construct( const FArguments& InArgs, const FName& InInstan
 	HistoryManager.SetOnUpdateHistoryData(FOnUpdateHistoryData::CreateSP(this, &SContentBrowser::OnUpdateHistoryData));
 
 	PathContextMenu = MakeShareable(new FPathContextMenu( AsShared() ));
-	PathContextMenu->SetOnNewAssetRequested( FNewAssetContextMenu::FOnNewAssetRequested::CreateSP(this, &SContentBrowser::NewAssetRequested) );
-	PathContextMenu->SetOnImportAssetRequested(FNewAssetContextMenu::FOnImportAssetRequested::CreateSP(this, &SContentBrowser::ImportAsset));
+	PathContextMenu->SetOnNewAssetRequested( FNewAssetOrClassContextMenu::FOnNewAssetRequested::CreateSP(this, &SContentBrowser::NewAssetRequested) );
+	PathContextMenu->SetOnNewClassRequested( FNewAssetOrClassContextMenu::FOnNewClassRequested::CreateSP(this, &SContentBrowser::NewClassRequested) );
+	PathContextMenu->SetOnImportAssetRequested(FNewAssetOrClassContextMenu::FOnImportAssetRequested::CreateSP(this, &SContentBrowser::ImportAsset));
 	PathContextMenu->SetOnRenameFolderRequested(FPathContextMenu::FOnRenameFolderRequested::CreateSP(this, &SContentBrowser::OnRenameFolderRequested));
 	PathContextMenu->SetOnFolderDeleted(FPathContextMenu::FOnFolderDeleted::CreateSP(this, &SContentBrowser::OnOpenedFolderDeleted));
 
@@ -109,9 +112,9 @@ void SContentBrowser::Construct( const FArguments& InArgs, const FName& InInstan
 							.ComboButtonStyle( FEditorStyle::Get(), "ContentBrowser.NewAsset.Style" )
 							.ForegroundColor(FLinearColor::White)
 							.ContentPadding(0)
-							.OnGetMenuContent_Lambda( [this]{ return MakeCreateAssetContextMenu( true, false ); } )
-							.ToolTipText( this, &SContentBrowser::GetNewAssetToolTipText )
-							.IsEnabled( this, &SContentBrowser::IsAssetPathSelected )
+							.OnGetMenuContent_Lambda( [this]{ return MakeAddNewContextMenu( true, false ); } )
+							.ToolTipText( this, &SContentBrowser::GetAddNewToolTipText )
+							.IsEnabled( this, &SContentBrowser::IsAddNewEnabled )
 							.AddMetaData<FTagMetaData>(FTagMetaData(TEXT("ContentBrowserNewAsset")))
 							.ButtonContent()
 							[
@@ -149,7 +152,7 @@ void SContentBrowser::Construct( const FArguments& InArgs, const FName& InInstan
 							SNew( SButton )
 							.ButtonStyle( FEditorStyle::Get(), "ToggleButton" )
 							.ToolTipText( this, &SContentBrowser::GetImportTooltipText )
-							.IsEnabled( this, &SContentBrowser::IsAssetPathSelected )
+							.IsEnabled( this, &SContentBrowser::IsImportEnabled )
 							.OnClicked( this, &SContentBrowser::HandleImportClicked )
 							.ContentPadding( 0 )
 							.AddMetaData<FTagMetaData>(FTagMetaData(TEXT("ContentBrowserImportAsset")))
@@ -388,6 +391,7 @@ void SContentBrowser::Construct( const FArguments& InArgs, const FName& InInstan
 							.FocusSearchBoxWhenOpened( false )
 							.ShowTreeTitle( false )
 							.ShowSeparator( false )
+							.AllowClassesFolder( true )
 							.AddMetaData<FTagMetaData>(FTagMetaData(TEXT("ContentBrowserSources")))
 							.SearchContent()
 							[
@@ -480,6 +484,7 @@ void SContentBrowser::Construct( const FArguments& InArgs, const FName& InInstan
 							.ContentPadding(0)
 							.ToolTipText( LOCTEXT( "AddFilterToolTip", "Add an asset filter." ) )
 							.OnGetMenuContent( this, &SContentBrowser::MakeAddFilterMenu )
+							.IsEnabled( this, &SContentBrowser::IsFilterMenuEnabled )
 							.HasDownArrow( true )
 							.ContentPadding( FMargin( 1, 0 ) )
 							.AddMetaData<FTagMetaData>(FTagMetaData(TEXT("ContentBrowserFiltersCombo")))
@@ -621,21 +626,34 @@ void SContentBrowser::CreateNewAsset(const FString& DefaultAssetName, const FStr
 	AssetViewPtr->CreateNewAsset(DefaultAssetName, PackagePath, AssetClass, Factory);
 }
 
+bool SContentBrowser::IsImportEnabled() const
+{
+	const FSourcesData& SourcesData = AssetViewPtr->GetSourcesData();
+	return SourcesData.PackagePaths.Num() == 1 && !ContentBrowserUtils::IsClassPath(SourcesData.PackagePaths[0].ToString());
+}
+
 FText SContentBrowser::GetImportTooltipText() const
 {
-	FString CurrentPath = GetCurrentPath();
+	const FSourcesData& SourcesData = AssetViewPtr->GetSourcesData();
 
-	FText ImportAssetLabel;
-	if ( !CurrentPath.IsEmpty() )
+	if ( SourcesData.PackagePaths.Num() == 1 )
 	{
-		ImportAssetLabel = FText::Format( LOCTEXT( "ImportAsset", "Import to {0}..." ), FText::FromString( CurrentPath ) );
+		const FString CurrentPath = SourcesData.PackagePaths[0].ToString();
+		if ( ContentBrowserUtils::IsClassPath( CurrentPath ) )
+		{
+			return LOCTEXT( "ImportAssetToolTip_InvalidClassPath", "Cannot import assets to class paths." );
+		}
+		else
+		{
+			return FText::Format( LOCTEXT( "ImportAssetToolTip", "Import to {0}..." ), FText::FromString( CurrentPath ) );
+		}
 	}
-	else
+	else if ( SourcesData.PackagePaths.Num() > 1 )
 	{
-		ImportAssetLabel = LOCTEXT( "ImportAsset_NoPath", "Import" );
+		return LOCTEXT( "ImportAssetToolTip_MultiplePaths", "Cannot import assets to multiple paths." );
 	}
-
-	return ImportAssetLabel;
+	
+	return LOCTEXT( "ImportAssetToolTip_NoPath", "No path is selected as an import target." );
 }
 
 FReply SContentBrowser::HandleImportClicked()
@@ -1082,6 +1100,17 @@ void SContentBrowser::NewAssetRequested(const FString& SelectedPath, TWeakObject
 	}
 }
 
+void SContentBrowser::NewClassRequested(const FString& SelectedPath)
+{
+	TSharedRef<FNativeClassHierarchy> NativeClassHierarchy = FContentBrowserSingleton::Get().GetNativeClassHierarchy();
+
+	// Parse out the on disk location for the currently selected path, this will then be used as the default location for the new class (if a valid project module location)
+	FString ExistingFolderPath;
+	NativeClassHierarchy->GetFileSystemPath(SelectedPath, ExistingFolderPath);
+
+	FGameProjectGenerationModule::Get().OpenAddCodeToProjectDialog(nullptr, ExistingFolderPath, FGlobalTabmanager::Get()->GetRootWindow());
+}
+
 void SContentBrowser::NewFolderRequested(const FString& SelectedPath)
 {
 	if( ensure(SelectedPath.Len() > 0) && AssetViewPtr.IsValid() )
@@ -1157,12 +1186,25 @@ TSharedPtr<SWidget> SContentBrowser::OnGetCrumbDelimiterContent(const FString& C
 
 	if( SourcesData.PackagePaths.Num() > 0 )
 	{
-		FAssetRegistryModule& AssetRegistryModule = FModuleManager::LoadModuleChecked<FAssetRegistryModule>(TEXT("AssetRegistry"));
-		IAssetRegistry& AssetRegistry = AssetRegistryModule.Get();
-
 		TArray<FString> SubPaths;
 		const bool bRecurse = false;
-		AssetRegistry.GetSubPaths( CrumbData, SubPaths, bRecurse );
+		if( ContentBrowserUtils::IsClassPath(CrumbData) )
+		{
+			TSharedRef<FNativeClassHierarchy> NativeClassHierarchy = FContentBrowserSingleton::Get().GetNativeClassHierarchy();
+
+			FNativeClassHierarchyFilter ClassFilter;
+			ClassFilter.ClassPaths.Add(FName(*CrumbData));
+			ClassFilter.bRecursivePaths = bRecurse;
+
+			NativeClassHierarchy->GetMatchingFolders(ClassFilter, SubPaths);
+		}
+		else
+		{
+			FAssetRegistryModule& AssetRegistryModule = FModuleManager::LoadModuleChecked<FAssetRegistryModule>(TEXT("AssetRegistry"));
+			IAssetRegistry& AssetRegistry = AssetRegistryModule.Get();
+
+			AssetRegistry.GetSubPaths( CrumbData, SubPaths, bRecurse );
+		}
 
 		if( SubPaths.Num() > 0 )
 		{
@@ -1208,6 +1250,7 @@ TSharedRef<SWidget> SContentBrowser::GetPathPickerContent()
 	
 	PathPickerConfig.OnPathSelected = FOnPathSelected::CreateSP(this, &SContentBrowser::PathPickerPathSelected);
 	PathPickerConfig.bAllowContextMenu = false;
+	PathPickerConfig.bAllowClassesFolder = true;
 
 	return SNew(SBox)
 		.WidthOverride(300)
@@ -1248,10 +1291,8 @@ FString SContentBrowser::GetCurrentPath() const
 	return CurrentPath;
 }
 
-TSharedRef<SWidget> SContentBrowser::MakeCreateAssetContextMenu(bool bShowGetContent, bool bShowImport)
+TSharedRef<SWidget> SContentBrowser::MakeAddNewContextMenu(bool bShowGetContent, bool bShowImport)
 {
-	FString CurrentPath = GetCurrentPath();
-
 	// Get all menu extenders for this context menu from the content browser module
 	FContentBrowserModule& ContentBrowserModule = FModuleManager::GetModuleChecked<FContentBrowserModule>( TEXT("ContentBrowser") );
 	TArray<FContentBrowserMenuExtender> MenuExtenderDelegates = ContentBrowserModule.GetAllAssetContextMenuExtenders();
@@ -1269,23 +1310,50 @@ TSharedRef<SWidget> SContentBrowser::MakeCreateAssetContextMenu(bool bShowGetCon
 	FMenuBuilder MenuBuilder(/*bInShouldCloseWindowAfterMenuSelection=*/true, NULL, MenuExtender);
 
 	// Only add "New Folder" item if we do not have a collection selected
-	FNewAssetContextMenu::FOnNewFolderRequested OnNewFolderRequested;
+	FNewAssetOrClassContextMenu::FOnNewFolderRequested OnNewFolderRequested;
 	if (CollectionViewPtr->GetSelectedCollections().Num() == 0)
 	{
-		OnNewFolderRequested = FNewAssetContextMenu::FOnNewFolderRequested::CreateSP(this, &SContentBrowser::NewFolderRequested);
+		OnNewFolderRequested = FNewAssetOrClassContextMenu::FOnNewFolderRequested::CreateSP(this, &SContentBrowser::NewFolderRequested);
 	}
 
-	FNewAssetContextMenu::MakeContextMenu(
+	const FSourcesData& SourcesData = AssetViewPtr->GetSourcesData();
+
+	int32 NumAssetPaths, NumClassPaths;
+	ContentBrowserUtils::CountPathTypes(SourcesData.PackagePaths, NumAssetPaths, NumClassPaths);
+
+	// Only the asset items if we have an asset path selected
+	FNewAssetOrClassContextMenu::FOnNewAssetRequested OnNewAssetRequested;
+	FNewAssetOrClassContextMenu::FOnImportAssetRequested OnImportAssetRequested;
+	FNewAssetOrClassContextMenu::FOnGetContentRequested OnGetContentRequested;
+	if(NumAssetPaths > 0)
+	{
+		OnNewAssetRequested = FNewAssetOrClassContextMenu::FOnNewAssetRequested::CreateSP(this, &SContentBrowser::NewAssetRequested);
+		if(bShowImport)
+		{
+			OnImportAssetRequested = FNewAssetOrClassContextMenu::FOnImportAssetRequested::CreateSP(this, &SContentBrowser::ImportAsset);
+		}
+		if(bShowGetContent)
+		{
+			OnGetContentRequested = FNewAssetOrClassContextMenu::FOnGetContentRequested::CreateSP(this, &SContentBrowser::OnAddContentRequested);
+		}
+	}
+
+	// Only add the class items if we have a class path selected
+	FNewAssetOrClassContextMenu::FOnNewClassRequested OnNewClassRequested;
+	if(NumClassPaths > 0)
+	{
+		OnNewClassRequested = FNewAssetOrClassContextMenu::FOnNewClassRequested::CreateSP(this, &SContentBrowser::NewClassRequested);
+	}
+
+	FNewAssetOrClassContextMenu::MakeContextMenu(
 		MenuBuilder, 
-		CurrentPath, 
-		FNewAssetContextMenu::FOnNewAssetRequested::CreateSP(this, &SContentBrowser::NewAssetRequested),
+		SourcesData.PackagePaths, 
+		OnNewAssetRequested,
+		OnNewClassRequested,
 		OnNewFolderRequested,
-		bShowImport
-			? FNewAssetContextMenu::FOnImportAssetRequested::CreateSP(this, &SContentBrowser::ImportAsset)
-			: FNewAssetContextMenu::FOnImportAssetRequested(),
-		bShowGetContent
-			? FNewAssetContextMenu::FOnGetContentRequested::CreateSP(this, &SContentBrowser::OnAddContentRequested)
-			: FNewAssetContextMenu::FOnGetContentRequested());
+		OnImportAssetRequested,
+		OnGetContentRequested
+		);
 
 	FDisplayMetrics DisplayMetrics;
 	FSlateApplication::Get().GetDisplayMetrics( DisplayMetrics );
@@ -1304,17 +1372,45 @@ TSharedRef<SWidget> SContentBrowser::MakeCreateAssetContextMenu(bool bShowGetCon
 		];
 }
 
-FText SContentBrowser::GetNewAssetToolTipText() const
+bool SContentBrowser::IsAddNewEnabled() const
 {
 	const FSourcesData& SourcesData = AssetViewPtr->GetSourcesData();
+	return SourcesData.PackagePaths.Num() == 1;
+}
 
-	// At least one source is selected
-	if (SourcesData.PackagePaths.Num() > 0)
+FText SContentBrowser::GetAddNewToolTipText() const
+{
+	const FSourcesData& SourcesData = AssetViewPtr->GetSourcesData();
+
+	if ( SourcesData.PackagePaths.Num() == 1 )
+	{
+		const FString CurrentPath = SourcesData.PackagePaths[0].ToString();
+		if ( ContentBrowserUtils::IsClassPath( CurrentPath ) )
+		{
+			return FText::Format( LOCTEXT("AddNewToolTip_AddNewClass", "Create a new class in {0}..."), FText::FromString(CurrentPath) );
+		}
+		else
+		{
+			return FText::Format( LOCTEXT("AddNewToolTip_AddNewAsset", "Create a new asset in {0}..."), FText::FromString(CurrentPath) );
+		}
+	}
+	else if ( SourcesData.PackagePaths.Num() > 1 )
 	{
-		return FText::Format( LOCTEXT("CreateAssetToolTipFmt", "Create an asset in {0}."), FText::FromName(SourcesData.PackagePaths[0]) );
+		return LOCTEXT( "AddNewToolTip_MultiplePaths", "Cannot add assets or classes to multiple paths." );
 	}
 	
-	return FText::GetEmpty();
+	return LOCTEXT( "AddNewToolTip_NoPath", "No path is selected as an add target." );
+}
+
+bool SContentBrowser::IsFilterMenuEnabled() const
+{
+	const FSourcesData& SourcesData = AssetViewPtr->GetSourcesData();
+
+	int32 NumAssetPaths, NumClassPaths;
+	ContentBrowserUtils::CountPathTypes(SourcesData.PackagePaths, NumAssetPaths, NumClassPaths);
+
+	// We can only filter when have something other than class paths selected (or nothing selected, for collections)
+	return !(NumClassPaths > 0 && NumAssetPaths == 0) || SourcesData.PackagePaths.Num() == 0;
 }
 
 TSharedRef<SWidget> SContentBrowser::MakeAddFilterMenu()
@@ -1420,25 +1516,9 @@ void SContentBrowser::OnAssetsActivated(const TArray<FAssetData>& ActivatedAsset
 
 TSharedPtr<SWidget> SContentBrowser::OnGetAssetContextMenu(const TArray<FAssetData>& SelectedAssets)
 {
-	// If a class is selected do not open a context menu
-	for(auto AssetIt = SelectedAssets.CreateConstIterator(); AssetIt; ++AssetIt)
-	{
-		if( AssetIt->AssetClass == NAME_Class )
-		{
-			return NULL;
-		}
-	}
-
-	// If the Classes folder is selected do not open a context menu
-	const TArray<FString> SelectedPaths = PathViewPtr->GetSelectedPaths();
-	if(SelectedPaths.Contains(TEXT("/Classes")))
-	{
-		return NULL;
-	}
-
 	if ( SelectedAssets.Num() == 0 )
 	{
-		return MakeCreateAssetContextMenu( false, true );
+		return MakeAddNewContextMenu( false, true );
 	}
 	else
 	{
@@ -1605,7 +1685,7 @@ FReply SContentBrowser::OnDirectoryUpClicked()
 void SContentBrowser::OnDirectoryUp()
 {
 	TArray<FString> SelectedPaths = PathViewPtr->GetSelectedPaths();
-	if(SelectedPaths.Num() == 1 && !ContentBrowserUtils::IsAssetRootDir(SelectedPaths[0]))
+	if(SelectedPaths.Num() == 1 && !ContentBrowserUtils::IsRootDir(SelectedPaths[0]))
 	{
 		FString ParentDir = SelectedPaths[0] / TEXT("..");
 		FPaths::CollapseRelativeDirectories(ParentDir);
@@ -1626,7 +1706,7 @@ bool SContentBrowser::IsForwardEnabled() const
 bool SContentBrowser::CanExecuteDirectoryUp() const
 {
 	TArray<FString> SelectedPaths = PathViewPtr->GetSelectedPaths();
-	return (SelectedPaths.Num() == 1 && !ContentBrowserUtils::IsAssetRootDir(SelectedPaths[0]));
+	return (SelectedPaths.Num() == 1 && !ContentBrowserUtils::IsRootDir(SelectedPaths[0]));
 }
 
 FText SContentBrowser::GetHistoryBackTooltip() const
@@ -1650,7 +1730,7 @@ FText SContentBrowser::GetHistoryForwardTooltip() const
 FText SContentBrowser::GetDirectoryUpTooltip() const
 {
 	TArray<FString> SelectedPaths = PathViewPtr->GetSelectedPaths();
-	if(SelectedPaths.Num() == 1 && !ContentBrowserUtils::IsAssetRootDir(SelectedPaths[0]))
+	if(SelectedPaths.Num() == 1 && !ContentBrowserUtils::IsRootDir(SelectedPaths[0]))
 	{
 		FString ParentDir = SelectedPaths[0] / TEXT("..");
 		FPaths::CollapseRelativeDirectories(ParentDir);
@@ -1735,8 +1815,11 @@ void SContentBrowser::UpdatePath()
 		FString CrumbPath = TEXT("/");
 		for ( auto CrumbIt = Crumbs.CreateConstIterator(); CrumbIt; ++CrumbIt )
 		{
+			// If this is the root part of the path, try and get the localized display name to stay in sync with what we see in SPathView
+			const FText DisplayName = (CrumbIt.GetIndex() == 0) ? ContentBrowserUtils::GetRootDirDisplayName(*CrumbIt) : FText::FromString(*CrumbIt);
+
 			CrumbPath += *CrumbIt;
-			PathBreadcrumbTrail->PushCrumb(FText::FromString(*CrumbIt), CrumbPath);
+			PathBreadcrumbTrail->PushCrumb(DisplayName, CrumbPath);
 			CrumbPath += TEXT("/");
 		}
 	}
@@ -1803,12 +1886,6 @@ bool SContentBrowser::IsFilteredBySource() const
 	return SourcesData.PackagePaths.Num() != 0 || SourcesData.Collections.Num() != 0;
 }
 
-bool SContentBrowser::IsAssetPathSelected() const
-{
-	return AssetViewPtr->IsAssetPathSelected();
-}
-
-
 void SContentBrowser::OnAssetRenameCommitted(const TArray<FAssetData>& Assets)
 {
 	// After a rename is committed we allow an implicit sync so as not to
@@ -1911,12 +1988,24 @@ FText SContentBrowser::GetSearchAssetsHintText() const
 		TArray<FString> Paths = PathViewPtr->GetSelectedPaths();
 		if (Paths.Num() != 0)
 		{
-			FString SearchHint = "Search ";
-			for(int i = 0; i < Paths.Num(); i++)
+			FString SearchHint = NSLOCTEXT( "ContentBrowser", "SearchBoxPartialHint", "Search" ).ToString();
+			SearchHint += TEXT(" ");
+			for(int32 i = 0; i < Paths.Num(); i++)
 			{
-				SearchHint += FPaths::GetCleanFilename(Paths[i]);
+				const FString& Path = Paths[i];
+				if (ContentBrowserUtils::IsRootDir(Path))
+				{
+					SearchHint += ContentBrowserUtils::GetRootDirDisplayName(Path).ToString();
+				}
+				else
+				{
+					SearchHint += FPaths::GetCleanFilename(Path);
+				}
+
 				if (i + 1 < Paths.Num())
+				{
 					SearchHint += ", ";
+				}
 			}
 
 			return FText::FromString(SearchHint);
@@ -1955,6 +2044,9 @@ TSharedPtr<SWidget> SContentBrowser::GetFolderContextMenu(const TArray<FString>&
 		AssetViewPtr->ClearSelection();
 	}
 	
+	// Ensure the path context menu has the up-to-date list of paths being worked on
+	PathContextMenu->SetSelectedPaths(SelectedPaths);
+
 	TSharedPtr<FExtender> Extender;
 	if(InMenuExtender.IsBound())
 	{
@@ -1964,12 +2056,35 @@ TSharedPtr<SWidget> SContentBrowser::GetFolderContextMenu(const TArray<FString>&
 	const bool bInShouldCloseWindowAfterSelection = true;
 	FMenuBuilder MenuBuilder(bInShouldCloseWindowAfterSelection, Commands, Extender, true);
 	
+	// We can only create folders when we have a single path selected
+	const bool bCanCreateNewFolder = SelectedPaths.Num() == 1 && ContentBrowserUtils::IsValidPathToCreateNewFolder(SelectedPaths[0]);
+
+	FText NewFolderToolTip;
+	if(SelectedPaths.Num() == 1)
+	{
+		if(bCanCreateNewFolder)
+		{
+			NewFolderToolTip = FText::Format(LOCTEXT("NewFolderTooltip_CreateIn", "Create a new folder in {0}."), FText::FromString(SelectedPaths[0]));
+		}
+		else
+		{
+			NewFolderToolTip = FText::Format(LOCTEXT("NewFolderTooltip_InvalidPath", "Cannot create new folders in {0}."), FText::FromString(SelectedPaths[0]));
+		}
+	}
+	else
+	{
+		NewFolderToolTip = LOCTEXT("NewFolderTooltip_InvalidNumberOfPaths", "Can only create folders when there is a single path selected.");
+	}
+
 	// New Folder
 	MenuBuilder.AddMenuEntry(
 		LOCTEXT("NewFolder", "New Folder"),
-		LOCTEXT("NewSubFolderTooltip", "Creates a new sub-folder in this folder."),
+		NewFolderToolTip,
 		FSlateIcon(FEditorStyle::GetStyleSetName(), "ContentBrowser.NewFolderIcon"),
-		FUIAction( FExecuteAction::CreateSP( this, &SContentBrowser::CreateNewFolder, SelectedPaths.Num() > 0 ? SelectedPaths[0] : FString(), InOnCreateNewFolder ) ),
+		FUIAction(
+			FExecuteAction::CreateSP( this, &SContentBrowser::CreateNewFolder, SelectedPaths.Num() > 0 ? SelectedPaths[0] : FString(), InOnCreateNewFolder ),
+			FCanExecuteAction::CreateLambda( [bCanCreateNewFolder] { return bCanCreateNewFolder; } )
+			),
 		"NewFolder"
 		);
 
diff --git a/Engine/Source/Editor/ContentBrowser/Private/SContentBrowser.h b/Engine/Source/Editor/ContentBrowser/Private/SContentBrowser.h
--- a/Engine/Source/Editor/ContentBrowser/Private/SContentBrowser.h
+++ b/Engine/Source/Editor/ContentBrowser/Private/SContentBrowser.h
@@ -72,6 +72,9 @@ private:
 	/** Called to retrieve the text that should be highlighted on assets */
 	FText GetHighlightedText() const;
 
+	/** Called to work out whether the import button should be enabled */
+	bool IsImportEnabled() const;
+
 	/** Called to retrieve the text that should be in the import tooltip */
 	FText GetImportTooltipText() const;
 
@@ -120,6 +123,9 @@ private:
 	/** Handler for when the path view requests an asset creation */
 	void NewAssetRequested(const FString& SelectedPath, TWeakObjectPtr<UClass> FactoryClass);
 
+	/** Handler for when the path view requests a class creation */
+	void NewClassRequested(const FString& SelectedPath);
+
 	/** Handler for when the path context menu requests a folder creation */
 	void NewFolderRequested(const FString& SelectedPath);
 
@@ -149,11 +155,17 @@ private:
 	/** Gets the content for the path picker combo button */
 	TSharedRef<SWidget> GetPathPickerContent();
 
-	/** Handle creating a context menu to generate a new asset */
-	TSharedRef<SWidget> MakeCreateAssetContextMenu(bool bShowGetContent, bool bShowImport);
+	/** Handle creating a context menu for the "Add New" button */
+	TSharedRef<SWidget> MakeAddNewContextMenu(bool bShowGetContent, bool bShowImport);
+
+	/** Called to work out whether the import button should be enabled */
+	bool IsAddNewEnabled() const;
 
-	/** Gets the tool tip for the new asset button */
-	FText GetNewAssetToolTipText() const;
+	/** Gets the tool tip for the "Add New" button */
+	FText GetAddNewToolTipText() const;
+
+	/** Called to work out whether the "Filters" combo should be enabled */
+	bool IsFilterMenuEnabled() const;
 
 	/** Makes the filters menu */
 	TSharedRef<SWidget> MakeAddFilterMenu();
@@ -266,9 +278,6 @@ private:
 	/** Returns true if currently filtering by a source */
 	bool IsFilteredBySource() const;
 
-	/** Returns true if a real asset path is selected (i.e \Engine\* or \Game\*) */
-	bool IsAssetPathSelected() const;
-
 	/** Handler for when the context menu or asset view requests to find assets in the asset tree */
 	void OnFindInAssetTreeRequested(const TArray<FAssetData>& AssetsToFind);
 
diff --git a/Engine/Source/Editor/ContentBrowser/Private/SPathPicker.cpp b/Engine/Source/Editor/ContentBrowser/Private/SPathPicker.cpp
--- a/Engine/Source/Editor/ContentBrowser/Private/SPathPicker.cpp
+++ b/Engine/Source/Editor/ContentBrowser/Private/SPathPicker.cpp
@@ -24,7 +24,7 @@ void SPathPicker::Construct( const FArguments& InArgs )
 		.OnGetPathContextMenuExtender(InArgs._PathPickerConfig.OnGetPathContextMenuExtender)
 		.FocusSearchBoxWhenOpened(InArgs._PathPickerConfig.bFocusSearchBoxWhenOpened)
 		.AllowContextMenu(InArgs._PathPickerConfig.bAllowContextMenu)
-		.AllowClassesFolder(false)
+		.AllowClassesFolder(InArgs._PathPickerConfig.bAllowClassesFolder)
 		.SelectionMode(ESelectionMode::Single)
 	];
 
@@ -54,12 +54,35 @@ TSharedPtr<SWidget> SPathPicker::GetFolderContextMenu(const TArray<FString>& Sel
 	const bool bCloseSelfOnly = true;
 	FMenuBuilder MenuBuilder(bInShouldCloseWindowAfterSelection, nullptr, Extender, bCloseSelfOnly);
 
+	// We can only create folders when we have a single path selected
+	const bool bCanCreateNewFolder = SelectedPaths.Num() == 1 && ContentBrowserUtils::IsValidPathToCreateNewFolder(SelectedPaths[0]);
+
+	FText NewFolderToolTip;
+	if(SelectedPaths.Num() == 1)
+	{
+		if(bCanCreateNewFolder)
+		{
+			NewFolderToolTip = FText::Format(LOCTEXT("NewFolderTooltip_CreateIn", "Create a new folder in {0}."), FText::FromString(SelectedPaths[0]));
+		}
+		else
+		{
+			NewFolderToolTip = FText::Format(LOCTEXT("NewFolderTooltip_InvalidPath", "Cannot create new folders in {0}."), FText::FromString(SelectedPaths[0]));
+		}
+	}
+	else
+	{
+		NewFolderToolTip = LOCTEXT("NewFolderTooltip_InvalidNumberOfPaths", "Can only create folders when there is a single path selected.");
+	}
+
 	// New Folder
 	MenuBuilder.AddMenuEntry(
 		LOCTEXT("NewFolder", "New Folder"),
-		LOCTEXT("NewSubFolderTooltip", "Creates a new sub-folder in this folder."),
+		NewFolderToolTip,
 		FSlateIcon(FEditorStyle::GetStyleSetName(), "ContentBrowser.NewFolderIcon"),
-		FUIAction(FExecuteAction::CreateSP(this, &SPathPicker::CreateNewFolder, SelectedPaths.Num() > 0 ? SelectedPaths[0] : FString(), InOnCreateNewFolder)),
+		FUIAction(
+			FExecuteAction::CreateSP(this, &SPathPicker::CreateNewFolder, SelectedPaths.Num() > 0 ? SelectedPaths[0] : FString(), InOnCreateNewFolder),
+			FCanExecuteAction::CreateLambda( [bCanCreateNewFolder] { return bCanCreateNewFolder; } )
+			),
 		"NewFolder"
 		);
 
diff --git a/Engine/Source/Editor/ContentBrowser/Private/SPathView.cpp b/Engine/Source/Editor/ContentBrowser/Private/SPathView.cpp
--- a/Engine/Source/Editor/ContentBrowser/Private/SPathView.cpp
+++ b/Engine/Source/Editor/ContentBrowser/Private/SPathView.cpp
@@ -8,6 +8,7 @@
 #include "ObjectTools.h"
 #include "SourcesViewWidgets.h"
 #include "SSearchBox.h"
+#include "NativeClassHierarchy.h"
 
 #define LOCTEXT_NAMESPACE "ContentBrowser"
 
@@ -17,6 +18,13 @@ SPathView::~SPathView()
 	FPackageName::OnContentPathMounted().RemoveAll( this );
 	FPackageName::OnContentPathDismounted().RemoveAll( this );
 
+	// Unsubscribe from class events
+	if ( bAllowClassesFolder )
+	{
+		TSharedRef<FNativeClassHierarchy> NativeClassHierarchy = FContentBrowserSingleton::Get().GetNativeClassHierarchy();
+		NativeClassHierarchy->OnClassHierarchyUpdated().RemoveAll( this );
+	}
+
 	// Load the asset registry module to stop listening for updates
 	FAssetRegistryModule& AssetRegistryModule = FModuleManager::LoadModuleChecked<FAssetRegistryModule>(TEXT("AssetRegistry"));
 	AssetRegistryModule.Get().OnPathAdded().RemoveAll( this );
@@ -51,8 +59,12 @@ void SPathView::Construct( const FArguments& InArgs )
 	FPackageName::OnContentPathMounted().AddSP( this, &SPathView::OnContentPathMountedOrDismounted );
 	FPackageName::OnContentPathDismounted().AddSP( this, &SPathView::OnContentPathMountedOrDismounted );
 
-	ClassesRootName = TEXT("Classes");
-	GameRootName = TEXT("Game");
+	// Listen to find out when the available classes are changed, so that we can refresh our paths
+	if ( bAllowClassesFolder )
+	{
+		TSharedRef<FNativeClassHierarchy> NativeClassHierarchy = FContentBrowserSingleton::Get().GetNativeClassHierarchy();
+		NativeClassHierarchy->OnClassHierarchyUpdated().AddSP( this, &SPathView::OnClassHierarchyUpdated );
+	}
 
 	ChildSlot
 	[
@@ -128,6 +140,7 @@ void SPathView::Construct( const FArguments& InArgs )
 	Populate();
 
 	// Always expand the game root initially
+	static const FString GameRootName = TEXT("Game");
 	for ( auto RootIt = TreeRootItems.CreateConstIterator(); RootIt; ++RootIt )
 	{
 		if ( (*RootIt)->FolderName == GameRootName )
@@ -315,7 +328,7 @@ TSharedPtr<FTreeItem> SPathView::AddPath(const FString& Path, bool bUserNamed)
 						break;
 					}
 
-					ChildItem = MakeShareable( new FTreeItem(FolderName, FolderPath, CurrentItem, bUserNamed) );
+					ChildItem = MakeShareable( new FTreeItem(FText::FromString(FolderName), FolderName, FolderPath, CurrentItem, bUserNamed) );
 					CurrentItem->Children.Add(ChildItem);
 					CurrentItem->SortChildren();
 					TreeViewPtr->RequestTreeRefresh();
@@ -424,12 +437,13 @@ void SPathView::SyncToAssets( const TArray<FAssetData>& AssetDataList, const boo
 	for (auto AssetDataIt = AssetDataList.CreateConstIterator(); AssetDataIt; ++AssetDataIt)
 	{
 		FString Path;
-		if ( AssetDataIt->AssetClass == FName("Class") )
+		if ( AssetDataIt->AssetClass == NAME_Class )
 		{
 			if ( bAllowClassesFolder )
 			{
-				// Classes are found in the /Classes root
-				Path = FString::Printf(TEXT("/%s"), *ClassesRootName);
+				// Classes are found in the /Classes_ roots
+				TSharedRef<FNativeClassHierarchy> NativeClassHierarchy = FContentBrowserSingleton::Get().GetNativeClassHierarchy();
+				NativeClassHierarchy->GetClassPath(Cast<UClass>(AssetDataIt->GetAsset()), Path, false/*bIncludeClassName*/);
 			}
 		}
 		else
@@ -640,17 +654,12 @@ TSharedPtr<SWidget> SPathView::MakePathViewContextMenu()
 		return NULL;
 	}
 
-	const TArray<FString> SelectedPaths = GetSelectedPaths();
-	if(SelectedPaths.Contains(TEXT("/Classes")))
-	{
-		return NULL;
-	}
-
 	if(!OnGetFolderContextMenu.IsBound())
 	{
 		return NULL;
 	}
 
+	const TArray<FString> SelectedPaths = GetSelectedPaths();
 	return OnGetFolderContextMenu.Execute(SelectedPaths, OnGetPathContextMenuExtender, FOnCreateNewFolder::CreateSP(this, &SPathView::OnCreateNewFolder));
 }
 
@@ -740,7 +749,8 @@ TSharedPtr<struct FTreeItem> SPathView::AddRootItem( const FString& InFolderName
 		const bool bDisplayPlugins = GetDefault<UContentBrowserSettings>()->GetDisplayPluginFolders();
 		if ( bDisplayPlugins || !ContentBrowserUtils::IsPluginFolder(InFolderName) )
 		{
-			NewItem = MakeShareable( new FTreeItem(InFolderName, FString(TEXT("/")) + InFolderName, TSharedPtr<FTreeItem>()));
+			const FText DisplayName = ContentBrowserUtils::GetRootDirDisplayName(InFolderName);
+			NewItem = MakeShareable( new FTreeItem(DisplayName, InFolderName, FString(TEXT("/")) + InFolderName, TSharedPtr<FTreeItem>()));
 			TreeRootItems.Add( NewItem );
 			TreeViewPtr->RequestTreeRefresh();
 		}
@@ -883,16 +893,28 @@ void SPathView::Populate()
 	TreeRootItems.Empty();
 	TreeViewPtr->ClearSelection();
 
+	// Load the native class hierarchy to listen for updates
+	TSharedRef<FNativeClassHierarchy> NativeClassHierarchy = FContentBrowserSingleton::Get().GetNativeClassHierarchy();
+
 	const bool bFilteringByText = !SearchBoxFolderFilter->GetRawFilterText().IsEmpty();
 	
+	const bool bDisplayEngine = GetDefault<UContentBrowserSettings>()->GetDisplayEngineFolder();
+	const bool bDisplayPlugins = GetDefault<UContentBrowserSettings>()->GetDisplayPluginFolders();
+
+	TArray<FName> ClassRoots;
+	TArray<FString> ClassFolders;
+	if ( bAllowClassesFolder )
+	{
+		NativeClassHierarchy->GetClassFolders(ClassRoots, ClassFolders, bDisplayEngine, bDisplayPlugins);
+	}
+
 	if ( !bFilteringByText )
 	{
 		// If we aren't filtering, add default folders to the asset tree
 
-		if ( bAllowClassesFolder )
+		for(const FName& ClassRoot : ClassRoots)
 		{
-			// Add the 'Classes' root.  Its a special case since it doesn't really have any content in it.
-			AddRootItem(ClassesRootName);
+			AddRootItem(ClassRoot.ToString());
 		}
 
 		// Add all of the content paths we know about.  Note that this can change on the fly (if say, a plugin
@@ -912,9 +934,7 @@ void SPathView::Populate()
 			{
 				CleanRootPathName = CleanRootPathName.Mid( 0, CleanRootPathName.Len() - 1 );
 			}
-
-			// @todo plugin content: These should probably be sorted alphabetically
-			AddRootItem( CleanRootPathName );
+			AddRootItem(CleanRootPathName);
 		}
 	}
 	
@@ -925,6 +945,9 @@ void SPathView::Populate()
 	TArray<FString> PathList;
 	AssetRegistryModule.Get().GetAllCachedPaths(PathList);
 
+	// Add any class paths we discovered
+	PathList.Append(ClassFolders);
+
 	// Add the user developer folder
 	const FString UserDeveloperFolder = FPackageName::FilenameToLongPackageName(FPaths::GameUserDeveloperDir().LeftChop(1));
 	PathList.Add(UserDeveloperFolder);
@@ -971,6 +994,12 @@ void SPathView::Populate()
 		}
 	}
 
+	// Sort the root tree items using their display names
+	TreeRootItems.Sort([](const TSharedPtr<FTreeItem>& One, const TSharedPtr<FTreeItem>& Two) -> bool
+	{
+		return One->DisplayName.ToString() < Two->DisplayName.ToString();
+	});
+
 	TreeViewPtr->RequestTreeRefresh();
 }
 
@@ -1056,25 +1085,36 @@ void SPathView::FolderNameChanged( const TSharedPtr< FTreeItem >& TreeItem, cons
 			TreeViewPtr->RequestScrollIntoView(TreeItem);
 		}
 
-		// Update the asset registry so this folder will persist
-		FAssetRegistryModule& AssetRegistryModule = FModuleManager::LoadModuleChecked<FAssetRegistryModule>("AssetRegistry");
-		if (AssetRegistryModule.Get().AddPath(TreeItem->FolderPath) && TreeItem->FolderPath != OldPath)
+		// Update either the asset registry or the native class hierarchy so this folder will persist
+		/*
+		if (ContentBrowserUtils::IsClassPath(TreeItem->FolderPath))
 		{
-			// move any assets in our folder
-			TArray<FAssetData> AssetsInFolder;
-			AssetRegistryModule.Get().GetAssetsByPath(*OldPath, AssetsInFolder, true);
-			TArray<UObject*> ObjectsInFolder;
-			ContentBrowserUtils::GetObjectsInAssetData(AssetsInFolder, ObjectsInFolder);
-			ContentBrowserUtils::MoveAssets(ObjectsInFolder, TreeItem->FolderPath, OldPath);
-
-			// Now check to see if the original folder is empty, if so we can delete it
-			TArray<FAssetData> AssetsInOriginalFolder;
-			AssetRegistryModule.Get().GetAssetsByPath(*OldPath, AssetsInOriginalFolder, true);
-			if (AssetsInOriginalFolder.Num() == 0)
+			// todo: jdale - CLASS - This will need updating to support renaming of class folders (SAssetView has similar logic - needs abstracting)
+			TSharedRef<FNativeClassHierarchy> NativeClassHierarchy = FContentBrowserSingleton::Get().GetNativeClassHierarchy();
+			NativeClassHierarchy->AddFolder(TreeItem->FolderPath);
+		}
+		else
+		*/
+		{
+			FAssetRegistryModule& AssetRegistryModule = FModuleManager::LoadModuleChecked<FAssetRegistryModule>("AssetRegistry");
+			if (AssetRegistryModule.Get().AddPath(TreeItem->FolderPath) && TreeItem->FolderPath != OldPath)
 			{
-				TArray<FString> FoldersToDelete;
-				FoldersToDelete.Add(OldPath);
-				ContentBrowserUtils::DeleteFolders(FoldersToDelete);
+				// move any assets in our folder
+				TArray<FAssetData> AssetsInFolder;
+				AssetRegistryModule.Get().GetAssetsByPath(*OldPath, AssetsInFolder, true);
+				TArray<UObject*> ObjectsInFolder;
+				ContentBrowserUtils::GetObjectsInAssetData(AssetsInFolder, ObjectsInFolder);
+				ContentBrowserUtils::MoveAssets(ObjectsInFolder, TreeItem->FolderPath, OldPath);
+
+				// Now check to see if the original folder is empty, if so we can delete it
+				TArray<FAssetData> AssetsInOriginalFolder;
+				AssetRegistryModule.Get().GetAssetsByPath(*OldPath, AssetsInOriginalFolder, true);
+				if (AssetsInOriginalFolder.Num() == 0)
+				{
+					TArray<FString> FoldersToDelete;
+					FoldersToDelete.Add(OldPath);
+					ContentBrowserUtils::DeleteFolders(FoldersToDelete);
+				}
 			}
 		}
 	}
@@ -1166,7 +1206,7 @@ void SPathView::TreeAssetsDropped(const TArray<FAssetData>& AssetList, const TSh
 	}
 
 	FMenuBuilder MenuBuilder(/*bInShouldCloseWindowAfterMenuSelection=*/true, NULL);
-	const FText MoveCopyHeaderString = FText::Format( LOCTEXT("AssetTreeDropMenuHeading", "Move/Copy to {0}"), FText::FromString( TreeItem->FolderName ) );
+	const FText MoveCopyHeaderString = FText::Format( LOCTEXT("AssetTreeDropMenuHeading", "Move/Copy to {0}"), TreeItem->DisplayName );
 	MenuBuilder.BeginSection("PathAssetMoveCopy", MoveCopyHeaderString);
 	{
 		MenuBuilder.AddMenuEntry(
@@ -1202,7 +1242,7 @@ void SPathView::TreeAssetsDropped(const TArray<FAssetData>& AssetList, const TSh
 void SPathView::TreeFoldersDropped(const TArray<FString>& PathNames, const TSharedPtr<FTreeItem>& TreeItem)
 {
 	FMenuBuilder MenuBuilder(/*bInShouldCloseWindowAfterMenuSelection=*/true, NULL);
-	const FText MoveCopyHeaderString = FText::Format( LOCTEXT("AssetTreeDropMenuHeading", "Move/Copy to {0}"), FText::FromString( TreeItem->FolderName ) );
+	const FText MoveCopyHeaderString = FText::Format( LOCTEXT("AssetTreeDropMenuHeading", "Move/Copy to {0}"), TreeItem->DisplayName );
 	MenuBuilder.BeginSection("PathFolderMoveCopy", MoveCopyHeaderString);
 	{
 		MenuBuilder.AddMenuEntry(
@@ -1337,6 +1377,12 @@ void SPathView::OnContentPathMountedOrDismounted( const FString& AssetPath, cons
 	RegisterActiveTimer(0.f, FWidgetActiveTimerDelegate::CreateSP(this, &SPathView::TriggerRepopulate));
 }
 
+void SPathView::OnClassHierarchyUpdated()
+{
+	// The class hierarchy has changed in some way, so we need to refresh our set of paths
+	RegisterActiveTimer(0.f, FWidgetActiveTimerDelegate::CreateSP(this, &SPathView::TriggerRepopulate));
+}
+
 void SPathView::HandleSettingChanged(FName PropertyName)
 {
 	if ((PropertyName == "DisplayDevelopersFolder") ||
diff --git a/Engine/Source/Editor/ContentBrowser/Private/SPathView.h b/Engine/Source/Editor/ContentBrowser/Private/SPathView.h
--- a/Engine/Source/Editor/ContentBrowser/Private/SPathView.h
+++ b/Engine/Source/Editor/ContentBrowser/Private/SPathView.h
@@ -218,6 +218,9 @@ private:
 	/** Called from an engine core event when a new content path has been added or removed, so that we can refresh our root set of paths */
 	void OnContentPathMountedOrDismounted( const FString& AssetPath, const FString& FileSystemPath );
 
+	/** Called when the class hierarchy is updated due to the available modules changing */
+	void OnClassHierarchyUpdated();
+
 	/** Delegate called when an editor setting is changed */
 	void HandleSettingChanged(FName PropertyName);
 
@@ -275,12 +278,6 @@ private:
 	/** The paths that were last reported by OnPathExpanded event. Used in preserving expansion when filtering folders */
 	TSet< FString > LastExpandedPaths;
 
-	/** The name of the root tree item for placeable classes */
-	FString ClassesRootName;
-
-	/** The name of the root tree item for game content */
-	FString GameRootName;
-
 	/** If not empty, this is the path of the folders to sync once they are available while assets are still being discovered */
 	TArray<FString> PendingInitialPaths;
 
diff --git a/Engine/Source/Editor/ContentBrowser/Private/SourcesViewWidgets.cpp b/Engine/Source/Editor/ContentBrowser/Private/SourcesViewWidgets.cpp
--- a/Engine/Source/Editor/ContentBrowser/Private/SourcesViewWidgets.cpp
+++ b/Engine/Source/Editor/ContentBrowser/Private/SourcesViewWidgets.cpp
@@ -61,6 +61,7 @@ void SAssetTreeItem::Construct( const FArguments& InArgs )
 			[
 				SAssignNew(InlineRenameWidget, SInlineEditableTextBlock)
 					.Text(this, &SAssetTreeItem::GetNameText)
+					.ToolTipText(this, &SAssetTreeItem::GetToolTipText)
 					.Font( FEditorStyle::GetFontStyle(bIsRoot ? "ContentBrowser.SourceTreeRootItemFont" : "ContentBrowser.SourceTreeItemFont") )
 					.HighlightText( InArgs._HighlightText )
 					.OnTextCommitted(this, &SAssetTreeItem::HandleNameCommitted)
@@ -285,6 +286,7 @@ void SAssetTreeItem::HandleNameCommitted( const FText& NewText, ETextCommit::Typ
 			const FString OldPath = TreeItemPtr->FolderPath;
 			FString Path;
 			TreeItemPtr->FolderPath.Split(TEXT("/"), &Path, NULL, ESearchCase::CaseSensitive, ESearchDir::FromEnd);
+			TreeItemPtr->DisplayName = NewText;
 			TreeItemPtr->FolderName = NewText.ToString();
 			TreeItemPtr->FolderPath = Path + TEXT("/") + NewText.ToString();
 
@@ -314,7 +316,7 @@ bool SAssetTreeItem::IsValidAssetPath() const
 	if ( TreeItem.IsValid() )
 	{
 		// The classes folder is not a real path
-		return TreeItem.Pin()->FolderPath != TEXT("/Classes");
+		return !ContentBrowserUtils::IsClassPath(TreeItem.Pin()->FolderPath);
 	}
 	else
 	{
@@ -353,9 +355,23 @@ FSlateColor SAssetTreeItem::GetFolderColor() const
 
 FText SAssetTreeItem::GetNameText() const
 {
-	if ( TreeItem.IsValid() )
+	TSharedPtr<FTreeItem> TreeItemPin = TreeItem.Pin();
+	if ( TreeItemPin.IsValid() )
+	{
+		return TreeItemPin->DisplayName;
+	}
+	else
+	{
+		return FText();
+	}
+}
+
+FText SAssetTreeItem::GetToolTipText() const
+{
+	TSharedPtr<FTreeItem> TreeItemPin = TreeItem.Pin();
+	if ( TreeItemPin.IsValid() )
 	{
-		return FText::FromString(TreeItem.Pin()->FolderName);
+		return FText::FromString(TreeItemPin->FolderPath);
 	}
 	else
 	{
diff --git a/Engine/Source/Editor/ContentBrowser/Private/SourcesViewWidgets.h b/Engine/Source/Editor/ContentBrowser/Private/SourcesViewWidgets.h
--- a/Engine/Source/Editor/ContentBrowser/Private/SourcesViewWidgets.h
+++ b/Engine/Source/Editor/ContentBrowser/Private/SourcesViewWidgets.h
@@ -78,6 +78,9 @@ private:
 	/** Returns the text of the folder name */
 	FText GetNameText() const;
 
+	/** Returns the text to use for the folder tooltip */
+	FText GetToolTipText() const;
+
 	/** Returns the image for the border around this item. Used for drag/drop operations */
 	const FSlateBrush* GetBorderImage() const;
 
diff --git a/Engine/Source/Editor/ContentBrowser/Public/IContentBrowserSingleton.h b/Engine/Source/Editor/ContentBrowser/Public/IContentBrowserSingleton.h
--- a/Engine/Source/Editor/ContentBrowser/Public/IContentBrowserSingleton.h
+++ b/Engine/Source/Editor/ContentBrowser/Public/IContentBrowserSingleton.h
@@ -179,12 +179,16 @@ struct FPathPickerConfig
 	/** If false, the context menu will not open when an item is right clicked */
 	bool bAllowContextMenu;
 
+	/** If true, will allow class folders to be shown in the picker */
+	bool bAllowClassesFolder;
+
 	/** If true, will add the path specified in DefaultPath to the tree if it doesn't exist already */
 	bool bAddDefaultPath;
 
 	FPathPickerConfig()
 		: bFocusSearchBoxWhenOpened(true)
 		, bAllowContextMenu(true)
+		, bAllowClassesFolder(false)
 		, bAddDefaultPath(false)
 	{}
 };
diff --git a/Engine/Source/Editor/GameProjectGeneration/Private/GameProjectGenerationModule.cpp b/Engine/Source/Editor/GameProjectGeneration/Private/GameProjectGenerationModule.cpp
--- a/Engine/Source/Editor/GameProjectGeneration/Private/GameProjectGenerationModule.cpp
+++ b/Engine/Source/Editor/GameProjectGeneration/Private/GameProjectGenerationModule.cpp
@@ -49,15 +49,21 @@ TSharedRef<class SWidget> FGameProjectGenerationModule::CreateGameProjectDialog(
 }
 
 
-TSharedRef<class SWidget> FGameProjectGenerationModule::CreateNewClassDialog(class UClass* InClass)
+TSharedRef<class SWidget> FGameProjectGenerationModule::CreateNewClassDialog(const UClass* InClass)
 {
 	return SNew(SNewClassDialog).Class(InClass);
 }
 
 
-void FGameProjectGenerationModule::OpenAddCodeToProjectDialog()
+void FGameProjectGenerationModule::OpenAddCodeToProjectDialog(const TSharedPtr<SWindow>& InParentWindow)
 {
-	GameProjectUtils::OpenAddCodeToProjectDialog();
+	GameProjectUtils::OpenAddCodeToProjectDialog(nullptr, FString(), InParentWindow);
+	AddCodeToProjectDialogOpenedEvent.Broadcast();
+}
+
+void FGameProjectGenerationModule::OpenAddCodeToProjectDialog(const UClass* InClass, const FString& InInitialPath, const TSharedPtr<SWindow>& InParentWindow)
+{
+	GameProjectUtils::OpenAddCodeToProjectDialog(InClass, InInitialPath, InParentWindow);
 	AddCodeToProjectDialogOpenedEvent.Broadcast();
 }
 
@@ -106,6 +112,16 @@ TArray<FModuleContextInfo> FGameProjectGenerationModule::GetCurrentProjectModule
 	return GameProjectUtils::GetCurrentProjectModules();
 }
 
+bool FGameProjectGenerationModule::IsValidBaseClassForCreation(const UClass* InClass, const FModuleContextInfo& InModuleInfo)
+{
+	return GameProjectUtils::IsValidBaseClassForCreation(InClass, InModuleInfo);
+}
+
+bool FGameProjectGenerationModule::IsValidBaseClassForCreation(const UClass* InClass, const TArray<FModuleContextInfo>& InModuleInfoArray)
+{
+	return GameProjectUtils::IsValidBaseClassForCreation(InClass, InModuleInfoArray);
+}
+
 void FGameProjectGenerationModule::GetProjectSourceDirectoryInfo(int32& OutNumFiles, int64& OutDirectorySize)
 {
 	GameProjectUtils::GetProjectSourceDirectoryInfo(OutNumFiles, OutDirectorySize);
diff --git a/Engine/Source/Editor/GameProjectGeneration/Private/GameProjectUtils.cpp b/Engine/Source/Editor/GameProjectGeneration/Private/GameProjectUtils.cpp
--- a/Engine/Source/Editor/GameProjectGeneration/Private/GameProjectUtils.cpp
+++ b/Engine/Source/Editor/GameProjectGeneration/Private/GameProjectUtils.cpp
@@ -695,21 +695,31 @@ bool GameProjectUtils::UpdateGameProject(const FString& ProjectFile, const FStri
 	return UpdateGameProjectFile(ProjectFile, EngineIdentifier, NULL, OutFailReason);
 }
 
-void GameProjectUtils::OpenAddCodeToProjectDialog()
+void GameProjectUtils::OpenAddCodeToProjectDialog(const UClass* InClass, const FString& InInitialPath, const TSharedPtr<SWindow>& InParentWindow)
 {
+	// If we've been given a class then we only show the second page of the dialog, so we can make the window smaller as that page doesn't have as much content
+	const FVector2D WindowSize = (InClass) ? FVector2D(940, 380) : FVector2D(940, 540);
+
 	TSharedRef<SWindow> AddCodeWindow =
 		SNew(SWindow)
 		.Title(LOCTEXT( "AddCodeWindowHeader", "Add Code"))
-		.ClientSize( FVector2D(940, 540) )
+		.ClientSize( WindowSize )
 		.SizingRule( ESizingRule::FixedSize )
 		.SupportsMinimize(false) .SupportsMaximize(false);
 
-	AddCodeWindow->SetContent( SNew(SNewClassDialog) );
+	AddCodeWindow->SetContent( SNew(SNewClassDialog).Class(InClass).InitialPath(InInitialPath) );
+
+	TSharedPtr<SWindow> ParentWindow = InParentWindow;
+	if (!ParentWindow.IsValid())
+	{
+		static const FName MainFrameModuleName = "MainFrame";
+		IMainFrameModule& MainFrameModule = FModuleManager::LoadModuleChecked<IMainFrameModule>(MainFrameModuleName);
+		ParentWindow = MainFrameModule.GetParentWindow();
+	}
 
-	IMainFrameModule& MainFrameModule = FModuleManager::LoadModuleChecked<IMainFrameModule>(TEXT("MainFrame"));
-	if (MainFrameModule.GetParentWindow().IsValid())
+	if (ParentWindow.IsValid())
 	{
-		FSlateApplication::Get().AddWindowAsNativeChild(AddCodeWindow, MainFrameModule.GetParentWindow().ToSharedRef());
+		FSlateApplication::Get().AddWindowAsNativeChild(AddCodeWindow, ParentWindow.ToSharedRef());
 	}
 	else
 	{
@@ -791,6 +801,55 @@ bool GameProjectUtils::IsValidClassNameForCreation(const FString& NewClassName,
 	return true;
 }
 
+bool GameProjectUtils::IsValidBaseClassForCreation(const UClass* InClass, const FModuleContextInfo& InModuleInfo)
+{
+	auto DoesClassNeedAPIExport = [&InModuleInfo](const FString& InClassModuleName) -> bool
+	{
+		return InModuleInfo.ModuleName != InClassModuleName;
+	};
+
+	return IsValidBaseClassForCreation_Internal(InClass, FDoesClassNeedAPIExportCallback::CreateLambda(DoesClassNeedAPIExport));
+}
+
+bool GameProjectUtils::IsValidBaseClassForCreation(const UClass* InClass, const TArray<FModuleContextInfo>& InModuleInfoArray)
+{
+	auto DoesClassNeedAPIExport = [&InModuleInfoArray](const FString& InClassModuleName) -> bool
+	{
+		for(const FModuleContextInfo& ModuleInfo : InModuleInfoArray)
+		{
+			if(ModuleInfo.ModuleName == InClassModuleName)
+			{
+				return false;
+			}
+		}
+		return true;
+	};
+
+	return IsValidBaseClassForCreation_Internal(InClass, FDoesClassNeedAPIExportCallback::CreateLambda(DoesClassNeedAPIExport));
+}
+
+bool GameProjectUtils::IsValidBaseClassForCreation_Internal(const UClass* InClass, const FDoesClassNeedAPIExportCallback& InDoesClassNeedAPIExport)
+{
+	// You may not make native classes based on blueprint generated classes
+	const bool bIsBlueprintClass = (InClass->ClassGeneratedBy != nullptr);
+
+	// UObject is special cased to be extensible since it would otherwise not be since it doesn't pass the API check (intrinsic class).
+	const bool bIsExplicitlyUObject = (InClass == UObject::StaticClass());
+
+	// You need API if you are not UObject itself, and you're in a module that was validated as needing API export
+	const FString ClassModuleName = InClass->GetOutermost()->GetName().RightChop( FString(TEXT("/Script/")).Len() );
+	const bool bNeedsAPI = !bIsExplicitlyUObject && InDoesClassNeedAPIExport.Execute(ClassModuleName);
+
+	// You may not make a class that is not DLL exported.
+	// MinimalAPI classes aren't compatible with the DLL export macro, but can still be used as a valid base
+	const bool bHasAPI = InClass->HasAnyClassFlags(CLASS_RequiredAPI) || InClass->HasAnyClassFlags(CLASS_MinimalAPI);
+
+	// @todo should we support interfaces?
+	const bool bIsInterface = InClass->IsChildOf(UInterface::StaticClass());
+
+	return !bIsBlueprintClass && (!bNeedsAPI || bHasAPI) && !bIsInterface;
+}
+
 bool GameProjectUtils::AddCodeToProject(const FString& NewClassName, const FString& NewClassPath, const FModuleContextInfo& ModuleInfo, const FNewClassInfo ParentClassInfo, const TSet<FString>& DisallowedHeaderNames, FString& OutHeaderFilePath, FString& OutCppFilePath, FText& OutFailReason)
 {
 	const bool bAddCodeSuccessful = AddCodeToProject_Internal(NewClassName, NewClassPath, ModuleInfo, ParentClassInfo, DisallowedHeaderNames, OutHeaderFilePath, OutCppFilePath, OutFailReason);
@@ -2890,7 +2949,7 @@ bool GameProjectUtils::AddCodeToProject_Internal(const FString& NewClassName, co
 		return false;
 	}
 
-	FScopedSlowTask SlowTask( 5,  LOCTEXT( "AddingCodeToProject", "Adding code to project..." ) );
+	FScopedSlowTask SlowTask( 5, LOCTEXT( "AddingCodeToProject", "Adding code to project..." ) );
 	SlowTask.MakeDialog();
 
 	SlowTask.EnterProgressFrame();
@@ -2982,6 +3041,8 @@ bool GameProjectUtils::AddCodeToProject_Internal(const FString& NewClassName, co
 	{
 		// This is the first time we add code to this project so compile its game DLL
 		const FString GameModuleName = FApp::GetGameName();
+		check(ModuleInfo.ModuleName == GameModuleName);
+
 		IHotReloadInterface& HotReloadSupport = FModuleManager::LoadModuleChecked<IHotReloadInterface>("HotReload");
 		const bool bReloadAfterCompiling = true;
 		const bool bForceCodeProject = true;
@@ -2991,6 +3052,53 @@ bool GameProjectUtils::AddCodeToProject_Internal(const FString& NewClassName, co
 			OutFailReason = LOCTEXT("FailedToCompileNewGameModule", "Failed to compile newly created game module.");
 			return false;
 		}
+
+		// Notify that we've created a brand new module
+		FSourceCodeNavigation::AccessOnNewModuleAdded().Broadcast(*GameModuleName);
+	}
+	else if (GEditor->AccessEditorUserSettings().bAutomaticallyHotReloadNewClasses)
+	{
+		FModuleStatus ModuleStatus;
+		const FName ModuleFName = *ModuleInfo.ModuleName;
+		if (ensure(FModuleManager::Get().QueryModule(ModuleFName, ModuleStatus)))
+		{
+			// Compile the module that the class was added to so that the newly added class with appear in the Content Browser
+			TArray<UPackage*> PackagesToRebind;
+			if (ModuleStatus.bIsLoaded)
+			{
+				const bool bIsHotReloadable = FModuleManager::Get().DoesLoadedModuleHaveUObjects(ModuleFName);
+				if (bIsHotReloadable)
+				{
+					// Is there a UPackage with the same name as this module?
+					const FString PotentialPackageName = FString(TEXT("/Script/")) + ModuleInfo.ModuleName;
+					UPackage* Package = FindPackage(nullptr, *PotentialPackageName);
+					if (Package)
+					{
+						PackagesToRebind.Add(Package);
+					}
+				}
+			}
+
+			IHotReloadInterface& HotReloadSupport = FModuleManager::LoadModuleChecked<IHotReloadInterface>("HotReload");
+			if (PackagesToRebind.Num() > 0)
+			{
+				// Perform a hot reload
+				const bool bWaitForCompletion = true;			
+				HotReloadSupport.RebindPackages(PackagesToRebind, TArray<FName>(), bWaitForCompletion, *GWarn);
+			}
+			else
+			{
+				// Perform a regular unload, then reload
+				const bool bReloadAfterRecompile = true;
+				const bool bForceCodeProject = false;
+				const bool bFailIfGeneratedCodeChanges = true;
+				if (!HotReloadSupport.RecompileModule(ModuleFName, bReloadAfterRecompile, *GWarn, bFailIfGeneratedCodeChanges, bForceCodeProject))
+				{
+					OutFailReason = FText::Format(LOCTEXT("FailedToCompileModuleFmt", "Failed to automatically compile the '{0}' module."), FText::FromString(ModuleInfo.ModuleName));
+					return false;
+				}
+			}
+		}
 	}
 
 	return true;
diff --git a/Engine/Source/Editor/GameProjectGeneration/Private/GameProjectUtils.h b/Engine/Source/Editor/GameProjectGeneration/Private/GameProjectUtils.h
--- a/Engine/Source/Editor/GameProjectGeneration/Private/GameProjectUtils.h
+++ b/Engine/Source/Editor/GameProjectGeneration/Private/GameProjectUtils.h
@@ -172,12 +172,24 @@ public:
 	/** Updates the given project file to an engine identifier. Returns true if the project was updated successfully or if no update was needed */
 	static bool UpdateGameProject(const FString& ProjectFile, const FString& EngineIdentifier, FText& OutFailReason);
 
-	/** Opens a dialog to add code files to a project */
-	static void OpenAddCodeToProjectDialog();
+	/** 
+	 * Opens a dialog to add code files to the current project. 
+	 *
+	 * @param	InClass				The class we should force the user to use as their base class type, or null to allow the user to choose their base class in the UI
+	 * @param	InInitialPath		The initial path we should use as the destination for the new class header file, or an empty string to choose a suitable default based upon the module path
+	 * @param	InParentWindow		The parent window the dialog should use, or null to choose a suitable default parent window (the main frame, if available)
+	 */
+	static void OpenAddCodeToProjectDialog(const UClass* InClass, const FString& InInitialPath, const TSharedPtr<SWindow>& InParentWindow);
 
 	/** Returns true if the specified class name is properly formed and does not conflict with another class */
 	static bool IsValidClassNameForCreation(const FString& NewClassName, const FModuleContextInfo& ModuleInfo, const TSet<FString>& DisallowedHeaderNames, FText& OutFailReason);
 
+	/** Returns true if the specified class is a valid base class for the given module */
+	static bool IsValidBaseClassForCreation(const UClass* InClass, const FModuleContextInfo& InModuleInfo);
+
+	/** Returns true if the specified class is a valid base class for any of the given modules */
+	static bool IsValidBaseClassForCreation(const UClass* InClass, const TArray<FModuleContextInfo>& InModuleInfoArray);
+
 	/** Adds new source code to the project. When returning true, OutSyncFileAndLineNumber will be the the preferred target file to sync in the users code editing IDE, formatted for use with GenericApplication::GotoLineInSource */
 	static bool AddCodeToProject(const FString& NewClassName, const FString& NewClassPath, const FModuleContextInfo& ModuleInfo, const FNewClassInfo ParentClassInfo, const TSet<FString>& DisallowedHeaderNames, FString& OutHeaderFilePath, FString& OutCppFilePath, FText& OutFailReason);
 
@@ -407,6 +419,10 @@ private:
 	/** Internal handler for AddCodeToProject*/
 	static bool AddCodeToProject_Internal(const FString& NewClassName, const FString& NewClassPath, const FModuleContextInfo& ModuleInfo, const FNewClassInfo ParentClassInfo, const TSet<FString>& DisallowedHeaderNames, FString& OutHeaderFilePath, FString& OutCppFilePath, FText& OutFailReason);
 
+	/** Internal handler for IsValidBaseClassForCreation */
+	DECLARE_DELEGATE_RetVal_OneParam(bool, FDoesClassNeedAPIExportCallback, const FString& /*ClassModuleName*/);
+	static bool IsValidBaseClassForCreation_Internal(const UClass* InClass, const FDoesClassNeedAPIExportCallback& InDoesClassNeedAPIExport);
+
 	/** Handler for the user confirming they've read the name legnth warning */
 	static void OnWarningReasonOk();
 
diff --git a/Engine/Source/Editor/GameProjectGeneration/Private/SNewClassDialog.cpp b/Engine/Source/Editor/GameProjectGeneration/Private/SNewClassDialog.cpp
--- a/Engine/Source/Editor/GameProjectGeneration/Private/SNewClassDialog.cpp
+++ b/Engine/Source/Editor/GameProjectGeneration/Private/SNewClassDialog.cpp
@@ -40,35 +40,15 @@ struct FParentClassItem
 class FNativeClassParentFilter : public IClassViewerFilter
 {
 public:
-	FNativeClassParentFilter(const TSharedPtr<FModuleContextInfo>* InSelectedModuleInfoPtr)
-		: SelectedModuleInfoPtr(InSelectedModuleInfoPtr)
+	FNativeClassParentFilter()
 	{
+		ProjectModules = GameProjectUtils::GetCurrentProjectModules();
 	}
 
 	virtual bool IsClassAllowed(const FClassViewerInitializationOptions& InInitOptions, const UClass* InClass, TSharedRef< FClassViewerFilterFuncs > InFilterFuncs ) override
 	{
-		// You may not make native classes based on blueprint generated classes
-		const bool bIsBlueprintClass = (InClass->ClassGeneratedBy != NULL);
-
-		// UObject is special cased to be extensible since it would otherwise not be since it doesn't pass the API check (intrinsic class).
-		const bool bIsExplicitlyUObject = (InClass == UObject::StaticClass());
-
-		// Is this class in the same module as our current module?
-		const FString ClassModuleName = InClass->GetOutermost()->GetName().RightChop( FString(TEXT("/Script/")).Len() );
-		const TSharedPtr<FModuleContextInfo>& ModuleInfo = *SelectedModuleInfoPtr;
-		const bool bIsInDestinationModule = (ModuleInfo.IsValid() && ModuleInfo->ModuleName == ClassModuleName);
-
-		// You need API if you are either not UObject itself and you are not in the destination module
-		const bool bNeedsAPI = !bIsExplicitlyUObject && !bIsInDestinationModule;
-
-		// You may not make a class that is not DLL exported.
-		// MinimalAPI classes aren't compatible with the DLL export macro, but can still be used as a valid base
-		const bool bHasAPI = InClass->HasAnyClassFlags(CLASS_RequiredAPI) || InClass->HasAnyClassFlags(CLASS_MinimalAPI);
-
-		// @todo should we support interfaces?
-		const bool bIsInterface = InClass->IsChildOf(UInterface::StaticClass());
-
-		return !bIsBlueprintClass && (!bNeedsAPI || bHasAPI) && !bIsInterface;
+		// We allow a class that belongs to any module in the current project, as you don't actually choose the destination module until after you've selected your parent class
+		return GameProjectUtils::IsValidBaseClassForCreation(InClass, ProjectModules);
 	}
 
 	virtual bool IsUnloadedClassAllowed(const FClassViewerInitializationOptions& InInitOptions, const TSharedRef< const IUnloadedBlueprintData > InUnloadedClassData, TSharedRef< FClassViewerFilterFuncs > InFilterFuncs) override
@@ -77,8 +57,8 @@ public:
 	}
 
 private:
-	/** Pointer to the currently selected module in the new class dialog */
-	const TSharedPtr<FModuleContextInfo>* SelectedModuleInfoPtr;
+	/** The list of currently available modules for this project */
+	TArray<FModuleContextInfo> ProjectModules;
 };
 
 static void FindPublicEngineHeaderFiles(TArray<FString>& OutFiles, const FString& Path)
@@ -106,8 +86,25 @@ void SNewClassDialog::Construct( const FArguments& InArgs )
 		}
 	}
 
+	// If we've been given an initial path that maps to a valid project module, use that as our initial module and path
+	if(!InArgs._InitialPath.IsEmpty())
+	{
+		const FString AbsoluteInitialPath = FPaths::ConvertRelativePathToFull(InArgs._InitialPath);
+		for(const auto& AvailableModule : AvailableModules)
+		{
+			if(AbsoluteInitialPath.StartsWith(AvailableModule->ModuleSourcePath))
+			{
+				SelectedModuleInfo = AvailableModule;
+				NewClassPath = AbsoluteInitialPath;
+				break;
+			}
+		}
+	}
+
+	// If we didn't get given a valid path override (see above), try and automatically work out the best default module
 	// If we have a runtime module with the same name as our project, then use that
 	// Otherwise, set out default target module as the first runtime module in the list
+	if(!SelectedModuleInfo.IsValid())
 	{
 		const FString ProjectName = FApp::GetGameName();
 		for(const auto& AvailableModule : AvailableModules)
@@ -130,9 +127,10 @@ void SNewClassDialog::Construct( const FArguments& InArgs )
 			// No runtime modules? Just take the first available module then
 			SelectedModuleInfo = AvailableModules[0];
 		}
+
+		NewClassPath = SelectedModuleInfo->ModuleSourcePath;
 	}
 
-	NewClassPath = SelectedModuleInfo->ModuleSourcePath;
 	ClassLocation = GameProjectUtils::EClassLocation::UserDefined; // the first call to UpdateInputValidity will set this correctly based on NewClassPath
 
 	ParentClassInfo = GameProjectUtils::FNewClassInfo(InArgs._Class);
@@ -158,7 +156,7 @@ void SNewClassDialog::Construct( const FArguments& InArgs )
 	Options.bShowObjectRootClass = true;
 
 	// Prevent creating native classes based on blueprint classes
-	Options.ClassFilter = MakeShareable(new FNativeClassParentFilter(&SelectedModuleInfo));
+	Options.ClassFilter = MakeShareable(new FNativeClassParentFilter());
 
 	ClassViewer = StaticCastSharedRef<SClassViewer>(FModuleManager::LoadModuleChecked<FClassViewerModule>("ClassViewer").CreateClassViewer(Options, FOnClassPicked::CreateSP(this, &SNewClassDialog::OnAdvancedClassSelected)));
 
@@ -170,7 +168,7 @@ void SNewClassDialog::Construct( const FArguments& InArgs )
 	ChildSlot
 	[
 		SNew(SBorder)
-		.Padding(FMargin(26, 8))
+		.Padding(18)
 		.BorderImage( FEditorStyle::GetBrush("Docking.Tab.ContentAreaBrush") )
 		[
 			SNew(SVerticalBox)
@@ -186,16 +184,55 @@ void SNewClassDialog::Construct( const FArguments& InArgs )
 				.OnCanceled(this, &SNewClassDialog::CancelClicked)
 				.OnFinished(this, &SNewClassDialog::FinishClicked)
 				.InitialPageIndex(ParentClassInfo.IsSet() ? 1 : 0)
+				.PageFooter()
+				[
+					// Get IDE information
+					SNew(SBorder)
+					.Visibility( this, &SNewClassDialog::GetGlobalErrorLabelVisibility )
+					.BorderImage( FEditorStyle::GetBrush("NewClassDialog.ErrorLabelBorder") )
+					.Padding(FMargin(0, 5))
+					.Content()
+					[
+						SNew(SHorizontalBox)
+
+						+SHorizontalBox::Slot()
+						.VAlign(VAlign_Center)
+						.Padding(2.f)
+						.AutoWidth()
+						[
+							SNew(SImage)
+							.Image(FEditorStyle::GetBrush("MessageLog.Warning"))
+						]
+
+						+SHorizontalBox::Slot()
+						.VAlign(VAlign_Center)
+						[
+							SNew(STextBlock)
+							.Text( this, &SNewClassDialog::GetGlobalErrorLabelText )
+							.TextStyle( FEditorStyle::Get(), "NewClassDialog.ErrorLabelFont" )
+						]
+
+						+SHorizontalBox::Slot()
+						.VAlign(VAlign_Center)
+						.HAlign(HAlign_Center)
+						.AutoWidth()
+						.Padding(5.f, 0.f)
+						[
+							SNew(SGetSuggestedIDEWidget)
+						]
+					]
+				]
 				
 				// Choose parent class
 				+SWizard::Page()
+				.CanShow(!ParentClassInfo.IsSet()) // We can't move to this widget page if we've been given a parent class to use
 				[
 					SNew(SVerticalBox)
 
 					// Title
 					+SVerticalBox::Slot()
 					.AutoHeight()
-					.Padding(0, 20, 0, 0)
+					.Padding(0)
 					[
 						SNew(STextBlock)
 						.TextStyle( FEditorStyle::Get(), "NewClassDialog.PageTitle" )
@@ -241,46 +278,6 @@ void SNewClassDialog::Construct( const FArguments& InArgs )
 						]
 					]
 
-					// Get IDE information
-					+SVerticalBox::Slot()
-					.Padding(0, 5)
-					.AutoHeight()
-					[
-						SNew(SBorder)
-						.Visibility( this, &SNewClassDialog::GetGlobalErrorLabelVisibility )
-						.BorderImage( FEditorStyle::GetBrush("NewClassDialog.ErrorLabelBorder") )
-						.Content()
-						[
-							SNew(SHorizontalBox)
-
-							+SHorizontalBox::Slot()
-							.VAlign(VAlign_Center)
-							.Padding(2.f)
-							.AutoWidth()
-							[
-								SNew(SImage)
-								.Image(FEditorStyle::GetBrush("MessageLog.Warning"))
-							]
-
-							+SHorizontalBox::Slot()
-							.VAlign(VAlign_Center)
-							[
-								SNew(STextBlock)
-								.Text( this, &SNewClassDialog::GetGlobalErrorLabelText )
-								.TextStyle( FEditorStyle::Get(), "NewClassDialog.ErrorLabelFont" )
-							]
-
-							+SHorizontalBox::Slot()
-							.VAlign(VAlign_Center)
-							.HAlign(HAlign_Center)
-							.AutoWidth()
-							.Padding(5.f, 0.f)
-							[
-								SNew(SGetSuggestedIDEWidget)
-							]
-						]
-					]
-
 					// Add Code list
 					+SVerticalBox::Slot()
 					.FillHeight(1.f)
@@ -404,7 +401,7 @@ void SNewClassDialog::Construct( const FArguments& InArgs )
 					// Title
 					+SVerticalBox::Slot()
 					.AutoHeight()
-					.Padding(0, 20, 0, 0)
+					.Padding(0)
 					[
 						SNew(STextBlock)
 						.TextStyle( FEditorStyle::Get(), "NewClassDialog.PageTitle" )
@@ -999,16 +996,25 @@ void SNewClassDialog::FinishClicked()
 		// Prevent periodic validity checks. This is to prevent a brief error message about the class already existing while you are exiting.
 		bPreventPeriodicValidityChecksUntilNextChange = true;
 
+		// Display a nag if we didn't automatically hot-reload for the newly added class
+		const bool bShowHotReloadNag = !GEditor->AccessEditorUserSettings().bAutomaticallyHotReloadNewClasses;
+
 		if ( HeaderFilePath.IsEmpty() || CppFilePath.IsEmpty() || !FSlateApplication::Get().SupportsSourceAccess() )
 		{
 			// Code successfully added, notify the user. We are either running on a platform that does not support source access or a file was not given so don't ask about editing the file
-			const FText Message = FText::Format( LOCTEXT("AddCodeSuccess", "Successfully added class {0}."), FText::FromString(NewClassName) );
+			const FText Message = FText::Format( (bShowHotReloadNag) 
+				? LOCTEXT("AddCodeSuccessWithHotReload", "Successfully added class {0}, however you must recompile {1} before it will appear in the Content Browser.")
+				: LOCTEXT("AddCodeSuccess", "Successfully added class {0}.")
+				, FText::FromString(NewClassName), FText::FromString(SelectedModuleInfo->ModuleName) );
 			FMessageDialog::Open(EAppMsgType::Ok, Message);
 		}
 		else
 		{
 			// Code successfully added, notify the user and ask about opening the IDE now
-			const FText Message = FText::Format( LOCTEXT("AddCodeSuccessWithSync", "Successfully added class {0}. Would you like to edit the code now?"), FText::FromString(NewClassName) );
+			const FText Message = FText::Format( (bShowHotReloadNag) 
+				? LOCTEXT("AddCodeSuccessWithHotReloadAndSync", "Successfully added class {0}, however you must recompile {1} before it will appear in the Content Browser.\n\nWould you like to edit the code now?")
+				: LOCTEXT("AddCodeSuccessWithSync", "Successfully added class {0}. Would you like to edit the code now?")
+				, FText::FromString(NewClassName), FText::FromString(SelectedModuleInfo->ModuleName) );
 			if ( FMessageDialog::Open(EAppMsgType::YesNo, Message) == EAppReturnType::Yes )
 			{
 				TArray<FString> SourceFiles;
@@ -1223,6 +1229,21 @@ void SNewClassDialog::UpdateInputValidity()
 		bLastInputValidityCheckSuccessful = GameProjectUtils::IsValidClassNameForCreation(NewClassName, *SelectedModuleInfo, DisallowedHeaderNames, LastInputValidityErrorText);
 	}
 
+	// Validate that the class is valid for the currently selected module
+	// As a project can have multiple modules, this lets us update the class validity as the user changes the target module
+	if ( bLastInputValidityCheckSuccessful && ParentClassInfo.BaseClass )
+	{
+		bLastInputValidityCheckSuccessful = GameProjectUtils::IsValidBaseClassForCreation(ParentClassInfo.BaseClass, *SelectedModuleInfo);
+		if ( !bLastInputValidityCheckSuccessful )
+		{
+			LastInputValidityErrorText = FText::Format(
+				LOCTEXT("NewClassError_InvalidBaseClassForModule", "{0} cannot be used as a base class in the {1} module. Please make sure that {0} is API exported."),
+				FText::FromString(ParentClassInfo.BaseClass->GetName()),
+				FText::FromString(SelectedModuleInfo->ModuleName)
+				);
+		}
+	}
+
 	LastPeriodicValidityCheckTime = FSlateApplication::Get().GetCurrentTime();
 
 	// Since this function was invoked, periodic validity checks should be re-enabled if they were disabled.
diff --git a/Engine/Source/Editor/GameProjectGeneration/Private/SNewClassDialog.h b/Engine/Source/Editor/GameProjectGeneration/Private/SNewClassDialog.h
--- a/Engine/Source/Editor/GameProjectGeneration/Private/SNewClassDialog.h
+++ b/Engine/Source/Editor/GameProjectGeneration/Private/SNewClassDialog.h
@@ -15,7 +15,10 @@ public:
 	{}
 
 	/** The class we want to build our new class from. If this is not specified then the wizard will display classes to the user. */
-	SLATE_ARGUMENT(UClass*, Class)
+	SLATE_ARGUMENT(const UClass*, Class)
+
+	/** The initial path to use as the destination for the new class. If this is not specified, we will work out a suitable default from the available project modules */
+	SLATE_ARGUMENT(FString, InitialPath)
 
 	SLATE_END_ARGS()
 
diff --git a/Engine/Source/Editor/GameProjectGeneration/Public/GameProjectGenerationModule.h b/Engine/Source/Editor/GameProjectGeneration/Public/GameProjectGenerationModule.h
--- a/Engine/Source/Editor/GameProjectGeneration/Public/GameProjectGenerationModule.h
+++ b/Engine/Source/Editor/GameProjectGeneration/Public/GameProjectGenerationModule.h
@@ -49,17 +49,31 @@ public:
 	 */
 	static inline FGameProjectGenerationModule& Get()
 	{
-		return FModuleManager::LoadModuleChecked< FGameProjectGenerationModule >( "GameProjectGeneration" );
+		static const FName ModuleName = "GameProjectGeneration";
+		return FModuleManager::LoadModuleChecked< FGameProjectGenerationModule >( ModuleName );
 	}
 
 	/** Creates the game project dialog */
 	virtual TSharedRef<class SWidget> CreateGameProjectDialog(bool bAllowProjectOpening, bool bAllowProjectCreate);
 
 	/** Creates a new class dialog for creating classes based on the passed-in class. */
-	virtual TSharedRef<class SWidget> CreateNewClassDialog(class UClass* InClass);
+	virtual TSharedRef<class SWidget> CreateNewClassDialog(const UClass* InClass);
 
-	/** Opens a dialog to adds code files to the current project. */
-	virtual void OpenAddCodeToProjectDialog();
+	/** 
+	 * Opens a dialog to add code files to the current project. 
+	 *
+	 * @param	InParentWindow		The parent window the dialog should use, or null to choose a suitable default parent window (the main frame, if available)
+	 */
+	virtual void OpenAddCodeToProjectDialog(const TSharedPtr<SWindow>& InParentWindow = nullptr);
+
+	/** 
+	 * Opens a dialog to add code files to the current project. 
+	 *
+	 * @param	InClass				The class we should force the user to use as their base class type, or null to allow the user to choose their base class in the UI
+	 * @param	InInitialPath		The initial path we should use as the destination for the new class header file, or an empty string to choose a suitable default based upon the module path
+	 * @param	InParentWindow		The parent window the dialog should use, or null to choose a suitable default parent window (the main frame, if available)
+	 */
+	virtual void OpenAddCodeToProjectDialog(const UClass* InClass, const FString& InInitialPath, const TSharedPtr<SWindow>& InParentWindow = nullptr);
 
 	/** Delegate for when the AddCodeToProject dialog is opened */
 	DECLARE_EVENT(FGameProjectGenerationModule, FAddCodeToProjectDialogOpenedEvent);
@@ -83,8 +97,15 @@ public:
 	/** Returns the path to the module's include header */
 	virtual FString DetermineModuleIncludePath(const FModuleContextInfo& ModuleInfo, const FString& FileRelativeTo);
 
+	/** Get the information about any modules referenced in the .uproject file of the currently loaded project */
 	virtual TArray<FModuleContextInfo> GetCurrentProjectModules();
 
+	/** Returns true if the specified class is a valid base class for the given module */
+	virtual bool IsValidBaseClassForCreation(const UClass* InClass, const FModuleContextInfo& InModuleInfo);
+
+	/** Returns true if the specified class is a valid base class for any of the given modules */
+	virtual bool IsValidBaseClassForCreation(const UClass* InClass, const TArray<FModuleContextInfo>& InModuleInfoArray);
+
 	/** Gets file and size info about the source directory */
 	virtual void GetProjectSourceDirectoryInfo(int32& OutNumFiles, int64& OutDirectorySize);
 
diff --git a/Engine/Source/Editor/UnrealEd/Classes/Editor/EditorUserSettings.h b/Engine/Source/Editor/UnrealEd/Classes/Editor/EditorUserSettings.h
--- a/Engine/Source/Editor/UnrealEd/Classes/Editor/EditorUserSettings.h
+++ b/Engine/Source/Editor/UnrealEd/Classes/Editor/EditorUserSettings.h
@@ -43,6 +43,10 @@ class UEditorUserSettings : public UObject
 	UPROPERTY(EditAnywhere, config, Category=Performance)
 	uint32 bMonitorEditorPerformance:1;
 
+	/** If enabled, any newly added classes will trigger a hot-reload of the module they were added to */
+	UPROPERTY(EditAnywhere, config, Category=HotReload)
+	uint32 bAutomaticallyHotReloadNewClasses:1;
+
 	/** Select to make Distributions use the curves, not the baked lookup tables. */
 	UPROPERTY(config)
 	uint32 bUseCurvesForDistributions:1; //(GDistributionType == 0)
diff --git a/Engine/Source/Editor/UnrealEd/Private/SourceCodeNavigation.cpp b/Engine/Source/Editor/UnrealEd/Private/SourceCodeNavigation.cpp
--- a/Engine/Source/Editor/UnrealEd/Private/SourceCodeNavigation.cpp
+++ b/Engine/Source/Editor/UnrealEd/Private/SourceCodeNavigation.cpp
@@ -173,7 +173,31 @@ private:
 
 
 FSourceFileDatabase::FSourceFileDatabase()
+	: bIsDirty(true)
 {
+	// Register to be notified when new .Build.cs files are added to the project
+	FSourceCodeNavigation::AccessOnNewModuleAdded().AddRaw(this, &FSourceFileDatabase::OnNewModuleAdded);
+
+	UpdateIfNeeded();
+}
+
+FSourceFileDatabase::~FSourceFileDatabase()
+{
+	FSourceCodeNavigation::AccessOnNewModuleAdded().RemoveAll(this);
+}
+
+void FSourceFileDatabase::UpdateIfNeeded()
+{
+	if (!bIsDirty)
+	{
+		return;
+	}
+
+	bIsDirty = false;
+
+	ModuleNames.Reset();
+	DisallowedHeaderNames.Empty();
+
 	// Find all the build rules within the game and engine directories
 	FindRootFilesRecursive(ModuleNames, *(FPaths::EngineDir() / TEXT("Source") / TEXT("Developer")), TEXT("*.Build.cs"));
 	FindRootFilesRecursive(ModuleNames, *(FPaths::EngineDir() / TEXT("Source") / TEXT("Editor")), TEXT("*.Build.cs"));
@@ -238,6 +262,11 @@ void FSourceFileDatabase::FindRootFilesRecursive(TArray<FString> &FileNames, con
 	}
 }
 
+void FSourceFileDatabase::OnNewModuleAdded(FName InModuleName)
+{
+	bIsDirty = true;
+}
+
 
 DECLARE_DELEGATE_RetVal( bool, FShouldAbortDelegate );
 
@@ -367,6 +396,9 @@ private:
 	/** Multi-cast delegate that fires after a compiler is not found. */
 	FSourceCodeNavigation::FOnCompilerNotFound OnCompilerNotFound;
 
+	/** Multi-cast delegate that fires after a new module (.Build.cs file) has been added */
+	FSourceCodeNavigation::FOnNewModuleAdded OnNewModuleAdded;
+
 	friend class FSourceCodeNavigation;
 };
 
@@ -653,6 +685,7 @@ const FSourceFileDatabase& FSourceCodeNavigation::GetSourceFileDatabase()
 	FScopeLock Lock(&CriticalSection);
 
 	static FSourceFileDatabase Instance;
+	Instance.UpdateIfNeeded();
 	return Instance;
 }
 
@@ -1538,6 +1571,11 @@ FSourceCodeNavigation::FOnCompilerNotFound& FSourceCodeNavigation::AccessOnCompi
 	return FSourceCodeNavigationImpl::Get().OnCompilerNotFound;
 }
 
+FSourceCodeNavigation::FOnNewModuleAdded& FSourceCodeNavigation::AccessOnNewModuleAdded()
+{
+	return FSourceCodeNavigationImpl::Get().OnNewModuleAdded;
+}
+
 bool FSourceCodeNavigation::FindModulePath( const FString& ModuleName, FString &OutModulePath )
 {
 	// Try to find a file matching the module name
diff --git a/Engine/Source/Editor/UnrealEd/Public/SourceCodeNavigation.h b/Engine/Source/Editor/UnrealEd/Public/SourceCodeNavigation.h
--- a/Engine/Source/Editor/UnrealEd/Public/SourceCodeNavigation.h
+++ b/Engine/Source/Editor/UnrealEd/Public/SourceCodeNavigation.h
@@ -20,17 +20,27 @@ public:
 	/** Constructs database */
 	FSourceFileDatabase();
 
+	/** Destructs database */
+	~FSourceFileDatabase();
+
 	/** Return array of module names used by the engine and current project, including those in plugins */
 	const TArray<FString>& GetModuleNames() const { return ModuleNames; }
 
 	/** Return set of public header names used by engine modules, which are disallowed as project header names */
 	const TSet<FString>& GetDisallowedHeaderNames() const { return DisallowedHeaderNames; }
 
+	/** Update the list of known modules and disallowed header names if they've become stale */
+	void UpdateIfNeeded();
+
 private:
 
 	/** Return array of filenames matching the given wildcard, recursing into subdirectories if no results are yielded from the base directory */
 	void FindRootFilesRecursive(TArray<FString> &FileNames, const FString &BaseDirectory, const FString &Wildcard);
 
+	/** Called when a new .Build.cs file is added to the current project */
+	void OnNewModuleAdded(FName InModuleName);
+
+	bool bIsDirty;
 	TArray<FString> ModuleNames;
 	TSet<FString> DisallowedHeaderNames;
 };
@@ -178,6 +188,12 @@ public:
 
 	/** Call this to access the multi-cast delegate that you can register a callback with */
 	UNREALED_API static FOnCompilerNotFound& AccessOnCompilerNotFound();
+
+	/** Delegate that's triggered when a new module (.Build.cs file) has been added */
+	DECLARE_MULTICAST_DELEGATE_OneParam( FOnNewModuleAdded, FName );
+
+	/** Call this to access the multi-cast delegate that you can register a callback with */
+	UNREALED_API static FOnNewModuleAdded& AccessOnNewModuleAdded();
 };
 
 
diff --git a/Engine/Source/Runtime/AppFramework/Private/Widgets/Workflow/SWizard.cpp b/Engine/Source/Runtime/AppFramework/Private/Widgets/Workflow/SWizard.cpp
--- a/Engine/Source/Runtime/AppFramework/Private/Widgets/Workflow/SWizard.cpp
+++ b/Engine/Source/Runtime/AppFramework/Private/Widgets/Workflow/SWizard.cpp
@@ -57,6 +57,13 @@ void SWizard::Construct( const FArguments& InArgs )
 					]
 			]
 
+		+ SVerticalBox::Slot()
+			.AutoHeight()
+			.Padding(0)
+			[
+				InArgs._PageFooter.Widget
+			]
+
 		+ SVerticalBox::Slot()
 			.AutoHeight()
 			.HAlign(HAlign_Right)
@@ -195,7 +202,8 @@ void SWizard::Construct( const FArguments& InArgs )
 			];
 	}
 
-	WidgetSwitcher->SetActiveWidgetIndex(InArgs._InitialPageIndex.Get());
+	WidgetSwitcher->SetActiveWidgetIndex(INDEX_NONE);
+	ShowPage(InArgs._InitialPageIndex.Get());
 }
 
 
diff --git a/Engine/Source/Runtime/AppFramework/Public/Widgets/Workflow/SWizard.h b/Engine/Source/Runtime/AppFramework/Public/Widgets/Workflow/SWizard.h
--- a/Engine/Source/Runtime/AppFramework/Public/Widgets/Workflow/SWizard.h
+++ b/Engine/Source/Runtime/AppFramework/Public/Widgets/Workflow/SWizard.h
@@ -145,6 +145,7 @@ public:
 		, _DesiredSize(FVector2D(0, 0))
 		, _ShowPageList(true)
 		, _ShowCancelButton(true)
+		, _PageFooter()
 	{ }
 
 		SLATE_SUPPORTS_SLOT_WITH_ARGS(FWizardPage)
@@ -185,6 +186,9 @@ public:
 		/** Holds a flag indicating whether the cancel button should be shown (default = true). */
 		SLATE_ARGUMENT(bool, ShowCancelButton)
 
+		/** Holds an optional widget containing the contents to place above the buttons for all wizard pages */
+		SLATE_NAMED_SLOT(FArguments, PageFooter)
+
 	SLATE_END_ARGS()
 
 public:
diff --git a/Engine/Source/Runtime/Core/Private/Modules/ModuleManager.cpp b/Engine/Source/Runtime/Core/Private/Modules/ModuleManager.cpp
--- a/Engine/Source/Runtime/Core/Private/Modules/ModuleManager.cpp
+++ b/Engine/Source/Runtime/Core/Private/Modules/ModuleManager.cpp
@@ -877,7 +877,7 @@ void FModuleManager::UnloadOrAbandonModuleWithCallback( const FName InModuleName
 	}
 	else
 	{
-		Ar.Logf( TEXT( "Modile being reloaded does not support dynamic unloading -- abandoning existing loaded module so that we can load the recompiled version!" ) );
+		Ar.Logf( TEXT( "Module being reloaded does not support dynamic unloading -- abandoning existing loaded module so that we can load the recompiled version!" ) );
 		AbandonModule( InModuleName );
 	}
 }
diff --git a/Engine/Source/Runtime/Projects/Private/PluginManager.cpp b/Engine/Source/Runtime/Projects/Private/PluginManager.cpp
--- a/Engine/Source/Runtime/Projects/Private/PluginManager.cpp
+++ b/Engine/Source/Runtime/Projects/Private/PluginManager.cpp
@@ -534,6 +534,7 @@ TArray< FPluginStatus > FPluginManager::QueryStatusForAllPlugins() const
 		PluginStatus.bIsEnabledByDefault = PluginInfo.bEnabledByDefault;
 		PluginStatus.bIsBetaVersion = PluginInfo.bIsBetaVersion;
 		PluginStatus.bHasContentFolder = PluginInfo.bCanContainContent;
+		PluginStatus.Modules = PluginInfo.Modules;
 
 		// @todo plugedit: Maybe we should do the FileExists check ONCE at plugin load time and not at query time
 		const FString Icon128FilePath = FPaths::GetPath(Plugin->FileName) / PluginSystemDefs::RelativeIcon128FilePath;
diff --git a/Engine/Source/Runtime/Projects/Public/Interfaces/IPluginManager.h b/Engine/Source/Runtime/Projects/Public/Interfaces/IPluginManager.h
--- a/Engine/Source/Runtime/Projects/Public/Interfaces/IPluginManager.h
+++ b/Engine/Source/Runtime/Projects/Public/Interfaces/IPluginManager.h
@@ -59,6 +59,9 @@ public:
 
 	/** Whether the plug-in has a content folder. */
 	bool bHasContentFolder;
+
+	/** List of all modules associated with this plugin */
+	TArray<FModuleDescriptor> Modules;
 };
 
 
