[INTEGRATE] Change 2412355 by [REDACTED_PERSON][REDACTED_HANDLE] on 2015/01/20 11:51:01

	Fixed bad padding and borders in "Open Level" dialog

diff --git a/Engine/Source/Editor/ContentBrowser/Private/SAssetDialog.cpp b/Engine/Source/Editor/ContentBrowser/Private/SAssetDialog.cpp
--- a/Engine/Source/Editor/ContentBrowser/Private/SAssetDialog.cpp
+++ b/Engine/Source/Editor/ContentBrowser/Private/SAssetDialog.cpp
@@ -187,14 +187,14 @@ void SAssetDialog::Construct(const FArguments& InArgs, const FSharedAssetDialogC
 		+ SHorizontalBox::Slot()
 		.AutoWidth()
 		.HAlign(HAlign_Right)
-		.VAlign(VAlign_Fill)
-		.Padding(bIncludeNameBox ? 80 : 4, 3, 4, 3)
+		.VAlign(VAlign_Bottom)
+		.Padding(bIncludeNameBox ? 80 : 4, 20, 4, 3)
 		[
 			LabelsBox
 		]
 		+ SHorizontalBox::Slot()
 		.FillWidth(1)
-		.VAlign(VAlign_Fill)
+		.VAlign(VAlign_Bottom)
 		.Padding(4, 3)
 		[
 			ContentBox
@@ -206,6 +206,7 @@ void SAssetDialog::Construct(const FArguments& InArgs, const FSharedAssetDialogC
 		[
 			SNew(SButton)
 			.Text(ConfirmButtonText)
+			.ContentPadding(FMargin(8, 2, 8, 2))
 			.IsEnabled(this, &SAssetDialog::IsConfirmButtonEnabled)
 			.OnClicked(this, &SAssetDialog::OnConfirmClicked)
 		]
@@ -215,6 +216,7 @@ void SAssetDialog::Construct(const FArguments& InArgs, const FSharedAssetDialogC
 		.Padding(4, 3)
 		[
 			SNew(SButton)
+			.ContentPadding(FMargin(8, 2, 8, 2))
 			.Text(LOCTEXT("AssetDialogCancelButton", "Cancel"))
 			.OnClicked(this, &SAssetDialog::OnCancelClicked)
 		];
@@ -222,13 +224,9 @@ void SAssetDialog::Construct(const FArguments& InArgs, const FSharedAssetDialogC
 	MainVerticalBox->AddSlot()
 		.AutoHeight()
 		.HAlign(HAlign_Fill)
-		.Padding(8, 2, 8, 4)
+		.Padding(0)
 		[
-			SNew(SBorder)
-			.BorderImage(FEditorStyle::GetBrush("ToolPanel.GroupBorder"))
-			[
-				ButtonsAndNameBox
-			]
+			ButtonsAndNameBox
 		];
 
 	ChildSlot
