[INTEGRATE] Change 2415257 by [REDACTED_PERSON][REDACTED_HANDLE] on 2015/01/22 09:06:28

	Stopped SFactoryMenuEntry using SAssetThumbnail directly

	UE-7184 - Show game c++ classes and engine c++ classes in the content browser

	It now just shows the class icon from the icon finder rather than construct the complex asset thumbnail widget - this prevents the class icon changes I'm making for the content browser from affecting the "Add New" context menu for assets.

diff --git a/Engine/Source/Editor/ContentBrowser/Private/NewAssetOrClassContextMenu.cpp b/Engine/Source/Editor/ContentBrowser/Private/NewAssetOrClassContextMenu.cpp
--- a/Engine/Source/Editor/ContentBrowser/Private/NewAssetOrClassContextMenu.cpp
+++ b/Engine/Source/Editor/ContentBrowser/Private/NewAssetOrClassContextMenu.cpp
@@ -4,6 +4,7 @@
 #include "NewAssetOrClassContextMenu.h"
 #include "IDocumentation.h"
 #include "EditorClassUtils.h"
+#include "ClassIconFinder.h"
 
 #define LOCTEXT_NAMESPACE "ContentBrowser"
 
@@ -67,21 +68,27 @@ public:
 	 */
 	void Construct( const FArguments& InArgs, UFactory* Factory )
 	{
-		// Since we are only displaying classes, we need no thumbnail pool.
-		Thumbnail = MakeShareable( new FAssetThumbnail( Factory->GetSupportedClass(), InArgs._Width, InArgs._Height, TSharedPtr<FAssetThumbnailPool>() ) );
-
-		FName ClassThumbnailBrushOverride = Factory->GetNewAssetThumbnailOverride();
-		TSharedPtr<SWidget> ThumbnailWidget;
-		if ( ClassThumbnailBrushOverride != NAME_None )
+		const FName ClassThumbnailBrushOverride = Factory->GetNewAssetThumbnailOverride();
+		const FSlateBrush* ClassThumbnail = nullptr;
+		if (ClassThumbnailBrushOverride.IsNone())
 		{
-			const bool bAllowFadeIn = false;
-			const bool bForceGenericThumbnail = true;
-			const bool bAllowHintText = false;
-			ThumbnailWidget = Thumbnail->MakeThumbnailWidget(bAllowFadeIn, bForceGenericThumbnail, EThumbnailLabel::AssetName, FText(), FLinearColor(0,0,0,0), bAllowHintText, ClassThumbnailBrushOverride);
+			ClassThumbnail = FClassIconFinder::FindThumbnailForClass(Factory->GetSupportedClass());
 		}
 		else
 		{
-			ThumbnailWidget = Thumbnail->MakeThumbnailWidget();
+			// Instead of getting the override thumbnail directly from the editor style here get it from the
+			// ClassIconFinder since it may have additional styles registered which can be searched by passing
+			// it as a default with no class to search for.
+			ClassThumbnail = FClassIconFinder::FindThumbnailForClass(nullptr, ClassThumbnailBrushOverride);
+		}
+
+		FAssetToolsModule& AssetToolsModule = FAssetToolsModule::GetModule();
+		TWeakPtr<IAssetTypeActions> AssetTypeActions = AssetToolsModule.Get().GetAssetTypeActionsForClass(Factory->GetSupportedClass());
+
+		FLinearColor AssetColor = FLinearColor::White;
+		if ( AssetTypeActions.IsValid() )
+		{
+			AssetColor = AssetTypeActions.Pin()->GetTypeColor();
 		}
 
 		ChildSlot
@@ -92,11 +99,35 @@ public:
 			.VAlign(VAlign_Center)
 			.AutoWidth()
 			[
-				SNew( SBox )
-				.WidthOverride( Thumbnail->GetSize().X + 4 )
-				.HeightOverride( Thumbnail->GetSize().Y + 4 )
+				SNew( SOverlay )
+
+				+SOverlay::Slot()
 				[
-					ThumbnailWidget.ToSharedRef()
+					SNew( SBox )
+					.WidthOverride( InArgs._Width + 4 )
+					.HeightOverride( InArgs._Height + 4 )
+					[
+						SNew( SBorder )
+						.BorderImage( FEditorStyle::GetBrush("AssetThumbnail.AssetBackground") )
+						.BorderBackgroundColor(AssetColor.CopyWithNewOpacity(0.3f))
+						.Padding( 2.0f )
+						.VAlign( VAlign_Center )
+						.HAlign( HAlign_Center )
+						[
+							SNew( SImage )
+							.Image( ClassThumbnail )
+						]
+					]
+				]
+
+				+SOverlay::Slot()
+				.HAlign(HAlign_Fill)
+				.VAlign(VAlign_Bottom)
+				[
+					SNew( SBorder )
+					.BorderImage( FEditorStyle::GetBrush("WhiteBrush") )
+					.BorderBackgroundColor( AssetColor )
+					.Padding( FMargin(0, FMath::Max(FMath::CeilToFloat(InArgs._Width*0.025f), 3.0f), 0, 0) )
 				]
 			]
 
@@ -119,9 +150,6 @@ public:
 		
 		SetToolTip(IDocumentation::Get()->CreateToolTip(Factory->GetToolTip(), nullptr, Factory->GetToolTipDocumentationPage(), Factory->GetToolTipDocumentationExcerpt()));
 	}
-
-private:
-	TSharedPtr< FAssetThumbnail > Thumbnail;
 };
 
 void FNewAssetOrClassContextMenu::MakeContextMenu(
