[INTEGRATE] Change 2415260 by [REDACTED_PERSON][REDACTED_HANDLE] on 2015/01/22 09:06:57

	Sorted content before classes, and modules in the order game -> engine -> plugins

	UE-7184 - Show game c++ classes and engine c++ classes in the content browser

	Also renamed "Game C++ Classes" to just "C++ Classes", and "Game Content" to just "Content" based on previous feedback.

diff --git a/Engine/Source/Editor/ContentBrowser/Private/ContentBrowserUtils.cpp b/Engine/Source/Editor/ContentBrowser/Private/ContentBrowserUtils.cpp
--- a/Engine/Source/Editor/ContentBrowser/Private/ContentBrowserUtils.cpp
+++ b/Engine/Source/Editor/ContentBrowser/Private/ContentBrowserUtils.cpp
@@ -1178,13 +1178,18 @@ FText ContentBrowserUtils::GetRootDirDisplayName(const FString& FolderPath)
 	}
 	else if(CleanFolderPath == GameFolderName)
 	{
-		LocalizedFolderName = LOCTEXT("GameFolderName", "Game");
+		//LocalizedFolderName = LOCTEXT("GameFolderName", "Game");
 	}
 	else
 	{
 		LocalizedFolderName = FText::FromString(CleanFolderPath);
 	}
 
+	if(LocalizedFolderName.IsEmpty())
+	{
+		return (bIsClassDir) ? LOCTEXT("ClassesFolder", "C++ Classes") : LOCTEXT("ContentFolder", "Content");
+	}
+
 	return FText::Format((bIsClassDir) ? LOCTEXT("ClassesFolderFmt", "{0} C++ Classes") : LOCTEXT("ContentFolderFmt", "{0} Content"), LocalizedFolderName);
 }
 
diff --git a/Engine/Source/Editor/ContentBrowser/Private/SPathView.cpp b/Engine/Source/Editor/ContentBrowser/Private/SPathView.cpp
--- a/Engine/Source/Editor/ContentBrowser/Private/SPathView.cpp
+++ b/Engine/Source/Editor/ContentBrowser/Private/SPathView.cpp
@@ -994,12 +994,67 @@ void SPathView::Populate()
 		}
 	}
 
-	// Sort the root tree items using their display names
+	SortRootItems();
+
+	bNeedsRepopulate = false;
+}
+
+void SPathView::SortRootItems()
+{
+	// First sort the root items by their display name, but also making sure that content to appears before classes
 	TreeRootItems.Sort([](const TSharedPtr<FTreeItem>& One, const TSharedPtr<FTreeItem>& Two) -> bool
 	{
+		static const FString ClassesPrefix = TEXT("Classes_");
+
+		FString OneModuleName = One->FolderName;
+		const bool bOneIsClass = OneModuleName.StartsWith(ClassesPrefix);
+		if(bOneIsClass)
+		{
+			OneModuleName = OneModuleName.Mid(ClassesPrefix.Len());
+		}
+
+		FString TwoModuleName = Two->FolderName;
+		const bool bTwoIsClass = TwoModuleName.StartsWith(ClassesPrefix);
+		if(bTwoIsClass)
+		{
+			TwoModuleName = TwoModuleName.Mid(ClassesPrefix.Len());
+		}
+
+		// We want to sort content before classes if both items belong to the same module
+		if(OneModuleName == TwoModuleName)
+		{
+			if(!bOneIsClass && bTwoIsClass)
+			{
+				return true;
+			}
+			return false;
+		}
+
 		return One->DisplayName.ToString() < Two->DisplayName.ToString();
 	});
 
+	// We have some manual sorting requirements that game must come before engine, and engine before everything else - we do that here after sorting everything by name
+	// The array below is in the inverse order as we iterate through and move each match to the beginning of the root items array
+	static const FString InverseSortOrder[] = {
+		TEXT("Classes_Engine"),
+		TEXT("Engine"),
+		TEXT("Classes_Game"),
+		TEXT("Game"),
+	};
+	for(const FString& SortItem : InverseSortOrder)
+	{
+		const int32 FoundItemIndex = TreeRootItems.IndexOfByPredicate([&SortItem](const TSharedPtr<FTreeItem>& TreeItem) -> bool
+		{
+			return TreeItem->FolderName == SortItem;
+		});
+		if(FoundItemIndex != INDEX_NONE)
+		{
+			TSharedPtr<FTreeItem> ItemToMove = TreeRootItems[FoundItemIndex];
+			TreeRootItems.RemoveAt(FoundItemIndex);
+			TreeRootItems.Insert(ItemToMove, 0);
+		}
+	}
+
 	TreeViewPtr->RequestTreeRefresh();
 }
 
diff --git a/Engine/Source/Editor/ContentBrowser/Private/SPathView.h b/Engine/Source/Editor/ContentBrowser/Private/SPathView.h
--- a/Engine/Source/Editor/ContentBrowser/Private/SPathView.h
+++ b/Engine/Source/Editor/ContentBrowser/Private/SPathView.h
@@ -113,6 +113,9 @@ public:
 	void Populate();
 
 private:
+	/** Sort the root items into the correct order */
+	void SortRootItems();
+
 	/** Gets the content for a context menu */
 	TSharedPtr<SWidget> MakePathViewContextMenu();
 
