[INTEGRATE] Change 2415343 by [REDACTED_PERSON][REDACTED_HANDLE] on 2015/01/22 11:08:02

	Made the Content Browser filters work with C++ classes

	UE-7184 - Show game c++ classes and engine c++ classes in the content browser

	You can now filter on the "C++ Class" type, and the asset view no longer shows you all assets when you have a filter applied but only class paths selected.

diff --git a/Engine/Source/Developer/AssetTools/Private/AssetTypeActions/AssetTypeActions_Class.h b/Engine/Source/Developer/AssetTools/Private/AssetTypeActions/AssetTypeActions_Class.h
--- a/Engine/Source/Developer/AssetTools/Private/AssetTypeActions/AssetTypeActions_Class.h
+++ b/Engine/Source/Developer/AssetTools/Private/AssetTypeActions/AssetTypeActions_Class.h
@@ -11,7 +11,7 @@ public:
 	virtual FText GetName() const override { return NSLOCTEXT("AssetTypeActions", "AssetTypeActions_Class", "C++ Class"); }
 	virtual FColor GetTypeColor() const override { return FColor(72, 155, 129); }
 	virtual UClass* GetSupportedClass() const override { return UClass::StaticClass(); }
-	virtual uint32 GetCategories() override { return EAssetTypeCategories::None; }
+	virtual uint32 GetCategories() override { return EAssetTypeCategories::Basic; }
 	
 	virtual bool HasActions(const TArray<UObject*>& InObjects) const override;
 	virtual void GetActions(const TArray<UObject*>& InObjects, FMenuBuilder& MenuBuilder) override;
diff --git a/Engine/Source/Editor/ContentBrowser/Private/SAssetView.cpp b/Engine/Source/Editor/ContentBrowser/Private/SAssetView.cpp
--- a/Engine/Source/Editor/ContentBrowser/Private/SAssetView.cpp
+++ b/Engine/Source/Editor/ContentBrowser/Private/SAssetView.cpp
@@ -1321,6 +1321,9 @@ void SAssetView::RefreshSourceItems()
 		const bool bUsingFolders = IsShowingFolders();
 		FARFilter Filter = SourcesData.MakeFilter(bRecurse, bUsingFolders);
 
+		// Add the backend filters from the filter list
+		Filter.Append(BackendFilter);
+
 		bWereItemsRecursivelyFiltered = bRecurse;
 
 		// Move any class paths into their own array
@@ -1333,35 +1336,42 @@ void SAssetView::RefreshSourceItems()
 			}
 			return false;
 		});
-		bShowClasses = ClassPathsToShow.Num() > 0;
+
+		// Only show classes if we have class paths, and the filter allows classes to be shown
+		const bool bFilterAllowsClasses = Filter.ClassNames.Num() == 0 || Filter.ClassNames.Contains(NAME_Class);
+		bShowClasses = ClassPathsToShow.Num() > 0 && bFilterAllowsClasses;
 
 		if ( SourcesData.Collections.Num() > 0 && Filter.ObjectPaths.Num() == 0 )
 		{
 			// This is an empty collection, no asset will pass the check
 		}
+		else if ( ClassPathsToShow.Num() > 0 && Filter.PackagePaths.Num() == 0 )
+		{
+			// Only class paths are selected, no asset will pass the check
+		}
 		else
 		{
-			// Add the backend filters from the filter list
-			Filter.Append(BackendFilter);
-
 			// Add assets found in the asset registry
 			AssetRegistryModule.Get().GetAssets(Filter, Items);
 		}
 
-		TArray< FName > ClassPaths;
-		FCollectionManagerModule& CollectionManagerModule = FModuleManager::GetModuleChecked<FCollectionManagerModule>(TEXT("CollectionManager"));
-		for (int Index = 0; Index < SourcesData.Collections.Num(); Index++)
+		if ( bFilterAllowsClasses )
 		{
-			CollectionManagerModule.Get().GetClassesInCollection( SourcesData.Collections[Index].Name, SourcesData.Collections[Index].Type, ClassPaths );
-		}
-
-		for (int Index = 0; Index < ClassPaths.Num(); Index++)
-		{
-			UClass* Class = FindObject<UClass>(ANY_PACKAGE, *ClassPaths[Index].ToString());
+			TArray< FName > ClassPaths;
+			FCollectionManagerModule& CollectionManagerModule = FModuleManager::GetModuleChecked<FCollectionManagerModule>(TEXT("CollectionManager"));
+			for (int32 Index = 0; Index < SourcesData.Collections.Num(); Index++)
+			{
+				CollectionManagerModule.Get().GetClassesInCollection( SourcesData.Collections[Index].Name, SourcesData.Collections[Index].Type, ClassPaths );
+			}
 
-			if ( Class != NULL )
+			for (int32 Index = 0; Index < ClassPaths.Num(); Index++)
 			{
-				Items.Add( Class );
+				UClass* Class = FindObject<UClass>(ANY_PACKAGE, *ClassPaths[Index].ToString());
+
+				if ( Class != NULL )
+				{
+					Items.Add( Class );
+				}
 			}
 		}
 	}
diff --git a/Engine/Source/Editor/ContentBrowser/Private/SContentBrowser.cpp b/Engine/Source/Editor/ContentBrowser/Private/SContentBrowser.cpp
--- a/Engine/Source/Editor/ContentBrowser/Private/SContentBrowser.cpp
+++ b/Engine/Source/Editor/ContentBrowser/Private/SContentBrowser.cpp
@@ -484,7 +484,6 @@ void SContentBrowser::Construct( const FArguments& InArgs, const FName& InInstan
 							.ContentPadding(0)
 							.ToolTipText( LOCTEXT( "AddFilterToolTip", "Add an asset filter." ) )
 							.OnGetMenuContent( this, &SContentBrowser::MakeAddFilterMenu )
-							.IsEnabled( this, &SContentBrowser::IsFilterMenuEnabled )
 							.HasDownArrow( true )
 							.ContentPadding( FMargin( 1, 0 ) )
 							.AddMetaData<FTagMetaData>(FTagMetaData(TEXT("ContentBrowserFiltersCombo")))
@@ -1402,17 +1401,6 @@ FText SContentBrowser::GetAddNewToolTipText() const
 	return LOCTEXT( "AddNewToolTip_NoPath", "No path is selected as an add target." );
 }
 
-bool SContentBrowser::IsFilterMenuEnabled() const
-{
-	const FSourcesData& SourcesData = AssetViewPtr->GetSourcesData();
-
-	int32 NumAssetPaths, NumClassPaths;
-	ContentBrowserUtils::CountPathTypes(SourcesData.PackagePaths, NumAssetPaths, NumClassPaths);
-
-	// We can only filter when have something other than class paths selected (or nothing selected, for collections)
-	return !(NumClassPaths > 0 && NumAssetPaths == 0) || SourcesData.PackagePaths.Num() == 0;
-}
-
 TSharedRef<SWidget> SContentBrowser::MakeAddFilterMenu()
 {
 	return FilterListPtr->ExternalMakeAddFilterMenu();
diff --git a/Engine/Source/Editor/ContentBrowser/Private/SContentBrowser.h b/Engine/Source/Editor/ContentBrowser/Private/SContentBrowser.h
--- a/Engine/Source/Editor/ContentBrowser/Private/SContentBrowser.h
+++ b/Engine/Source/Editor/ContentBrowser/Private/SContentBrowser.h
@@ -164,9 +164,6 @@ private:
 	/** Gets the tool tip for the "Add New" button */
 	FText GetAddNewToolTipText() const;
 
-	/** Called to work out whether the "Filters" combo should be enabled */
-	bool IsFilterMenuEnabled() const;
-
 	/** Makes the filters menu */
 	TSharedRef<SWidget> MakeAddFilterMenu();
 	
