[INTEGRATE] Change 2416705 by [REDACTED_PERSON][REDACTED_HANDLE] on 2015/01/23 07:28:46

	The "New Class" option is now always available in the "Add New" context menu

	UE-7184 - Show game c++ classes and engine c++ classes in the content browser

	It will use your project's source folder if you don't have a valid class path selected.

diff --git a/Engine/Source/Editor/ContentBrowser/Private/NewAssetOrClassContextMenu.cpp b/Engine/Source/Editor/ContentBrowser/Private/NewAssetOrClassContextMenu.cpp
--- a/Engine/Source/Editor/ContentBrowser/Private/NewAssetOrClassContextMenu.cpp
+++ b/Engine/Source/Editor/ContentBrowser/Private/NewAssetOrClassContextMenu.cpp
@@ -204,10 +204,11 @@ void FNewAssetOrClassContextMenu::MakeContextMenu(
 	};
 	const FCanExecuteAction CanExecuteAssetActionsDelegate = FCanExecuteAction::CreateLambda(CanExecuteAssetActions);
 
-	auto CanExecuteClassActions = [NumAssetPaths, NumClassPaths, bIsValidNewClassPath]() -> bool
+	auto CanExecuteClassActions = [NumAssetPaths, NumClassPaths]() -> bool
 	{
-		// We can execute class actions when we only have a single class path selected, and that path is a valid path for creating a class
-		return NumClassPaths == 1 && NumAssetPaths == 0 && bIsValidNewClassPath;
+		// We can execute class actions when we only have a single path selected
+		// This menu always lets you create classes, but uses your default project source folder if the selected path is invalid for creating classes
+		return (NumAssetPaths + NumClassPaths) == 1;
 	};
 	const FCanExecuteAction CanExecuteClassActionsDelegate = FCanExecuteAction::CreateLambda(CanExecuteClassActions);
 
@@ -267,16 +268,18 @@ void FNewAssetOrClassContextMenu::MakeContextMenu(
 	// Add Class
 	if(InOnNewClassRequested.IsBound())
 	{
+		FString ClassCreationPath = FirstSelectedPath;
 		FText NewClassToolTip;
 		if(bHasSinglePathSelected)
 		{
 			if(bIsValidNewClassPath)
 			{
-				NewClassToolTip = FText::Format(LOCTEXT("NewClassTooltip_CreateIn", "Create a new class in {0}."), FText::FromString(FirstSelectedPath));
+				NewClassToolTip = FText::Format(LOCTEXT("NewClassTooltip_CreateIn", "Create a new class in {0}."), FText::FromString(ClassCreationPath));
 			}
 			else
 			{
-				NewClassToolTip = FText::Format(LOCTEXT("NewClassTooltip_InvalidPath", "Cannot create new classes in {0}."), FText::FromString(FirstSelectedPath));
+				NewClassToolTip = LOCTEXT("NewClassTooltip_CreateInDefault", "Create a new class in your project's source folder.");
+				ClassCreationPath.Empty(); // An empty path override will cause the class wizard to use the default project path
 			}
 		}
 		else
@@ -284,14 +287,14 @@ void FNewAssetOrClassContextMenu::MakeContextMenu(
 			NewClassToolTip = LOCTEXT("NewClassTooltip_InvalidNumberOfPaths", "Can only create classes when there is a single path selected.");
 		}
 
-		MenuBuilder.BeginSection("ContentBrowserNewClass", LOCTEXT("ClassMenuHeading", "Class") );
+		MenuBuilder.BeginSection("ContentBrowserNewClass", LOCTEXT("ClassMenuHeading", "C++ Class") );
 		{
 			MenuBuilder.AddMenuEntry(
 				LOCTEXT("NewClassLabel", "New Class"),
 				NewClassToolTip,
 				FSlateIcon(FEditorStyle::GetStyleSetName(), "MainFrame.AddCodeToProject"),
 				FUIAction(
-					FExecuteAction::CreateStatic(&FNewAssetOrClassContextMenu::ExecuteNewClass, FirstSelectedPath, InOnNewClassRequested),
+					FExecuteAction::CreateStatic(&FNewAssetOrClassContextMenu::ExecuteNewClass, ClassCreationPath, InOnNewClassRequested),
 					CanExecuteClassActionsDelegate
 					)
 				);
@@ -426,10 +429,8 @@ void FNewAssetOrClassContextMenu::ExecuteNewAsset(FString InPath, TWeakObjectPtr
 
 void FNewAssetOrClassContextMenu::ExecuteNewClass(FString InPath, FOnNewClassRequested InOnNewClassRequested)
 {
-	if(ensure(!InPath.IsEmpty()))
-	{
-		InOnNewClassRequested.ExecuteIfBound(InPath);
-	}
+	// An empty path override will cause the class wizard to use the default project path
+	InOnNewClassRequested.ExecuteIfBound(InPath);
 }
 
 void FNewAssetOrClassContextMenu::ExecuteNewFolder(FString InPath, FOnNewFolderRequested InOnNewFolderRequested)
diff --git a/Engine/Source/Editor/ContentBrowser/Private/SContentBrowser.cpp b/Engine/Source/Editor/ContentBrowser/Private/SContentBrowser.cpp
--- a/Engine/Source/Editor/ContentBrowser/Private/SContentBrowser.cpp
+++ b/Engine/Source/Editor/ContentBrowser/Private/SContentBrowser.cpp
@@ -1101,11 +1101,13 @@ void SContentBrowser::NewAssetRequested(const FString& SelectedPath, TWeakObject
 
 void SContentBrowser::NewClassRequested(const FString& SelectedPath)
 {
-	TSharedRef<FNativeClassHierarchy> NativeClassHierarchy = FContentBrowserSingleton::Get().GetNativeClassHierarchy();
-
 	// Parse out the on disk location for the currently selected path, this will then be used as the default location for the new class (if a valid project module location)
 	FString ExistingFolderPath;
-	NativeClassHierarchy->GetFileSystemPath(SelectedPath, ExistingFolderPath);
+	if (!SelectedPath.IsEmpty())
+	{
+		TSharedRef<FNativeClassHierarchy> NativeClassHierarchy = FContentBrowserSingleton::Get().GetNativeClassHierarchy();
+		NativeClassHierarchy->GetFileSystemPath(SelectedPath, ExistingFolderPath);
+	}
 
 	FGameProjectGenerationModule::Get().OpenAddCodeToProjectDialog(nullptr, ExistingFolderPath, FGlobalTabmanager::Get()->GetRootWindow());
 }
@@ -1337,12 +1339,8 @@ TSharedRef<SWidget> SContentBrowser::MakeAddNewContextMenu(bool bShowGetContent,
 		}
 	}
 
-	// Only add the class items if we have a class path selected
-	FNewAssetOrClassContextMenu::FOnNewClassRequested OnNewClassRequested;
-	if(NumClassPaths > 0)
-	{
-		OnNewClassRequested = FNewAssetOrClassContextMenu::FOnNewClassRequested::CreateSP(this, &SContentBrowser::NewClassRequested);
-	}
+	// This menu always lets you create classes, but uses your default project source folder if the selected path is invalid for creating classes
+	FNewAssetOrClassContextMenu::FOnNewClassRequested OnNewClassRequested = FNewAssetOrClassContextMenu::FOnNewClassRequested::CreateSP(this, &SContentBrowser::NewClassRequested);
 
 	FNewAssetOrClassContextMenu::MakeContextMenu(
 		MenuBuilder, 
