Fix for merge: bNeedsRepopulate doesn't exist any more in main (presumably doesn't need resetting).

diff --git a/Engine/Source/Editor/ContentBrowser/Private/SPathView.cpp b/Engine/Source/Editor/ContentBrowser/Private/SPathView.cpp
--- a/Engine/Source/Editor/ContentBrowser/Private/SPathView.cpp
+++ b/Engine/Source/Editor/ContentBrowser/Private/SPathView.cpp
@@ -995,8 +995,6 @@ void SPathView::Populate()
 	}
 
 	SortRootItems();
-
-	bNeedsRepopulate = false;
 }
 
 void SPathView::SortRootItems()
