[INTEGRATE] Change 2417927 by [REDACTED_PERSON][REDACTED_HANDLE] on 2015/01/24 15:52:29

	Cleaned up "New C++ Class" button labels in Content Browser

diff --git a/Engine/Source/Editor/ContentBrowser/Private/NewAssetOrClassContextMenu.cpp b/Engine/Source/Editor/ContentBrowser/Private/NewAssetOrClassContextMenu.cpp
--- a/Engine/Source/Editor/ContentBrowser/Private/NewAssetOrClassContextMenu.cpp
+++ b/Engine/Source/Editor/ContentBrowser/Private/NewAssetOrClassContextMenu.cpp
@@ -290,7 +290,7 @@ void FNewAssetOrClassContextMenu::MakeContextMenu(
 		MenuBuilder.BeginSection("ContentBrowserNewClass", LOCTEXT("ClassMenuHeading", "C++ Class") );
 		{
 			MenuBuilder.AddMenuEntry(
-				LOCTEXT("NewClassLabel", "New Class"),
+				LOCTEXT("NewClassLabel", "New C++ Class..."),
 				NewClassToolTip,
 				FSlateIcon(FEditorStyle::GetStyleSetName(), "MainFrame.AddCodeToProject"),
 				FUIAction(
diff --git a/Engine/Source/Editor/ContentBrowser/Private/PathContextMenu.cpp b/Engine/Source/Editor/ContentBrowser/Private/PathContextMenu.cpp
--- a/Engine/Source/Editor/ContentBrowser/Private/PathContextMenu.cpp
+++ b/Engine/Source/Editor/ContentBrowser/Private/PathContextMenu.cpp
@@ -151,7 +151,7 @@ void FPathContextMenu::MakePathViewContextMenu(FMenuBuilder& MenuBuilder)
 
 				// New Class
 				MenuBuilder.AddMenuEntry(
-					LOCTEXT("NewClassLabel", "New Class"),
+					LOCTEXT("NewClassLabel", "New C++ Class..."),
 					NewClassToolTip,
 					FSlateIcon(FEditorStyle::GetStyleSetName(), "MainFrame.AddCodeToProject"),
 					FUIAction(
