[INTEGRATE] Change 2418620 by [REDACTED_PERSON][REDACTED_HANDLE] on 2015/01/26 10:35:34

	Fix capitaliization for 'Add Feature or Content Pack...'

diff --git a/Engine/Source/Editor/ContentBrowser/Private/NewAssetOrClassContextMenu.cpp b/Engine/Source/Editor/ContentBrowser/Private/NewAssetOrClassContextMenu.cpp
--- a/Engine/Source/Editor/ContentBrowser/Private/NewAssetOrClassContextMenu.cpp
+++ b/Engine/Source/Editor/ContentBrowser/Private/NewAssetOrClassContextMenu.cpp
@@ -218,7 +218,7 @@ void FNewAssetOrClassContextMenu::MakeContextMenu(
 		MenuBuilder.BeginSection( "ContentBrowserGetContent", LOCTEXT( "GetContentMenuHeading", "Content" ) );
 		{
 			MenuBuilder.AddMenuEntry(
-				LOCTEXT( "GetContentText", "Add feature or content pack..." ),
+				LOCTEXT( "GetContentText", "Add Feature or Content Pack..." ),
 				LOCTEXT( "GetContentTooltip", "Add features and content packs to the project." ),
 				FSlateIcon( FEditorStyle::GetStyleSetName(), "ContentBrowser.AddContent" ),
 				FUIAction(
