[INTEGRATE] Change 2418662 by [REDACTED_PERSON][REDACTED_HANDLE] on 2015/01/26 11:12:12

	Added a "C++" overlay for C++ source code folders in SPathView

	UE-7184 - Show game c++ classes and engine c++ classes in the content browser

diff --git a/Engine/Source/Editor/ContentBrowser/Private/SourcesViewWidgets.cpp b/Engine/Source/Editor/ContentBrowser/Private/SourcesViewWidgets.cpp
--- a/Engine/Source/Editor/ContentBrowser/Private/SourcesViewWidgets.cpp
+++ b/Engine/Source/Editor/ContentBrowser/Private/SourcesViewWidgets.cpp
@@ -30,9 +30,19 @@ void SAssetTreeItem::Construct( const FArguments& InArgs )
 
 	FolderOpenBrush = FEditorStyle::GetBrush("ContentBrowser.AssetTreeFolderOpen");
 	FolderClosedBrush = FEditorStyle::GetBrush("ContentBrowser.AssetTreeFolderClosed");
+	FolderOpenCodeBrush = FEditorStyle::GetBrush("ContentBrowser.AssetTreeFolderOpenCode");
+	FolderClosedCodeBrush = FEditorStyle::GetBrush("ContentBrowser.AssetTreeFolderClosedCode");
 	FolderDeveloperBrush = FEditorStyle::GetBrush("ContentBrowser.AssetTreeFolderDeveloper");
 		
-	bDeveloperFolder = ContentBrowserUtils::IsDevelopersFolder(InArgs._TreeItem->FolderPath);
+	FolderType = EFolderType::Normal;
+	if( ContentBrowserUtils::IsDevelopersFolder(InArgs._TreeItem->FolderPath) )
+	{
+		FolderType = EFolderType::Developer;
+	}
+	else if( ContentBrowserUtils::IsClassPath/*IsClassRootDir*/(InArgs._TreeItem->FolderPath) )
+	{
+		FolderType = EFolderType::Code;
+	}
 
 	bool bIsRoot = !InArgs._TreeItem->Parent.IsValid();
 
@@ -326,17 +336,16 @@ bool SAssetTreeItem::IsValidAssetPath() const
 
 const FSlateBrush* SAssetTreeItem::GetFolderIcon() const
 {
-	if ( bDeveloperFolder )
+	switch( FolderType )
 	{
+	case EFolderType::Code:
+		return ( IsItemExpanded.Get() ) ? FolderOpenCodeBrush : FolderClosedCodeBrush;
+
+	case EFolderType::Developer:
 		return FolderDeveloperBrush;
-	}
-	else if ( IsItemExpanded.Get() )
-	{
-		return FolderOpenBrush;
-	}
-	else
-	{
-		return FolderClosedBrush;
+
+	default:
+		return ( IsItemExpanded.Get() ) ? FolderOpenBrush : FolderClosedBrush;
 	}
 }
 
diff --git a/Engine/Source/Editor/ContentBrowser/Private/SourcesViewWidgets.h b/Engine/Source/Editor/ContentBrowser/Private/SourcesViewWidgets.h
--- a/Engine/Source/Editor/ContentBrowser/Private/SourcesViewWidgets.h
+++ b/Engine/Source/Editor/ContentBrowser/Private/SourcesViewWidgets.h
@@ -91,6 +91,13 @@ private:
 	EVisibility GetStaticTextVisibility() const;
 
 private:
+	enum class EFolderType : uint8
+	{
+		Normal,
+		Code,
+		Developer,
+	};
+
 	/** The data for this item */
 	TWeakPtr<FTreeItem> TreeItem;
 
@@ -121,13 +128,15 @@ private:
 	/** Brushes for the different folder states */
 	const FSlateBrush* FolderOpenBrush;
 	const FSlateBrush* FolderClosedBrush;
+	const FSlateBrush* FolderOpenCodeBrush;
+	const FSlateBrush* FolderClosedCodeBrush;
 	const FSlateBrush* FolderDeveloperBrush;
 
 	/** True when a drag is over this item with a valid operation for drop */
 	bool bDraggedOver;
 
-	/** True when this item represents a folder which is in the developer folder or is the developer folder itself */
-	bool bDeveloperFolder;
+	/** What type of stuff does this folder hold */
+	EFolderType FolderType;
 
 	/** Widget to display the name of the asset item and allows for renaming */
 	TSharedPtr< SInlineEditableTextBlock > InlineRenameWidget;
diff --git a/Engine/Source/Editor/EditorStyle/Private/SlateEditorStyle.cpp b/Engine/Source/Editor/EditorStyle/Private/SlateEditorStyle.cpp
--- a/Engine/Source/Editor/EditorStyle/Private/SlateEditorStyle.cpp
+++ b/Engine/Source/Editor/EditorStyle/Private/SlateEditorStyle.cpp
@@ -5093,6 +5093,8 @@ void FSlateEditorStyle::FStyle::SetupContentBrowserStyle()
 		Set( "ContentBrowser.BreadcrumbPathPickerFolder", new IMAGE_BRUSH( "Icons/FolderClosed", FVector2D(18, 16), FLinearColor::Gray ) );
 		Set( "ContentBrowser.AssetTreeFolderOpen", new IMAGE_BRUSH( "Icons/FolderOpen", FVector2D(18, 16) ) );
 		Set( "ContentBrowser.AssetTreeFolderDeveloper", new IMAGE_BRUSH( "Icons/FolderDeveloper", FVector2D(18, 16) ) );
+		Set( "ContentBrowser.AssetTreeFolderOpenCode", new IMAGE_BRUSH( "Icons/FolderOpen_Code", FVector2D(18, 16) ) );
+		Set( "ContentBrowser.AssetTreeFolderClosedCode", new IMAGE_BRUSH( "Icons/FolderClosed_Code", FVector2D(18, 16) ) );
 		Set( "ContentBrowser.AddCollectionButtonIcon", new IMAGE_BRUSH( "Icons/PlusSymbol_12x", Icon12x12, FLinearColor::Black ) );
 
 		Set( "ContentBrowser.Splitter", FSplitterStyle()
