UE-7132 Renaming folders with subfolders breaks the content browser

diff --git a/Engine/Source/Editor/ContentBrowser/Private/SPathView.cpp b/Engine/Source/Editor/ContentBrowser/Private/SPathView.cpp
--- a/Engine/Source/Editor/ContentBrowser/Private/SPathView.cpp
+++ b/Engine/Source/Editor/ContentBrowser/Private/SPathView.cpp
@@ -341,6 +341,11 @@ TSharedPtr<FTreeItem> SPathView::AddPath(const FString& Path, bool bUserNamed)
 						TreeViewPtr->RequestScrollIntoView(ChildItem);
 					}
 				}
+				else
+				{
+					//If the child item does exist, ensure its folder path is correct (may differ when renaming parent folder)
+					ChildItem->FolderPath = CurrentItem->FolderPath + "/" + PathItemName;
+				}
 
 				CurrentItem = ChildItem;
 			}
