Fixed some issues where CB browser sync would fail when syncing to C++ classes

UE-8574 - Sync to content browser does not work for C++ classes

Two issues:
1) The Content Browser wasn't enabling the Engine or Plugin view filters if you synced to a class in those categories - this meant that you couldn't see the synced class as its parent folder was still hidden.
2) FNativeClassHierarchy::GetClassPath was returning a path with a trailing / if the class was in the module root folder - this never matched anything so nothing was synced.

diff --git a/Engine/Source/Editor/ContentBrowser/Private/ContentBrowserUtils.cpp b/Engine/Source/Editor/ContentBrowser/Private/ContentBrowserUtils.cpp
--- a/Engine/Source/Editor/ContentBrowser/Private/ContentBrowserUtils.cpp
+++ b/Engine/Source/Editor/ContentBrowser/Private/ContentBrowserUtils.cpp
@@ -1028,6 +1028,53 @@ bool ContentBrowserUtils::AssetHasCustomThumbnail( const FAssetData& AssetData )
 	return false;
 }
 
+ContentBrowserUtils::ECBFolderCategory ContentBrowserUtils::GetFolderCategory( const FString& InPath )
+{
+	static const FString ClassesPrefix = TEXT("/Classes_");
+	static const FString GameClassesPrefix = TEXT("/Classes_Game");
+	static const FString EngineClassesPrefix = TEXT("/Classes_Engine");
+
+	const bool bIsClassDir = InPath.StartsWith(ClassesPrefix);
+	if(bIsClassDir)
+	{
+		const bool bIsGameClassDir = InPath.StartsWith(GameClassesPrefix);
+		if(bIsGameClassDir)
+		{
+			return ECBFolderCategory::GameClasses;
+		}
+
+		const bool bIsEngineClassDir = InPath.StartsWith(EngineClassesPrefix);
+		if(bIsEngineClassDir)
+		{
+			return ECBFolderCategory::EngineClasses;
+		}
+
+		return ECBFolderCategory::PluginClasses;
+	}
+	else
+	{
+		const bool bIsEngineContent = IsEngineFolder(InPath);
+		if(bIsEngineContent)
+		{
+			return ECBFolderCategory::EngineContent;
+		}
+
+		const bool bIsPluginContent = IsPluginFolder(InPath);
+		if(bIsPluginContent)
+		{
+			return ECBFolderCategory::PluginContent;
+		}
+
+		const bool bIsDeveloperContent = IsDevelopersFolder(InPath);
+		if(bIsDeveloperContent)
+		{
+			return ECBFolderCategory::DeveloperContent;
+		}
+
+		return ECBFolderCategory::GameContent;
+	}
+}
+
 bool ContentBrowserUtils::IsEngineFolder( const FString& InPath )
 {
 	return InPath.StartsWith(TEXT("/Engine")) || InPath == TEXT("Engine");
diff --git a/Engine/Source/Editor/ContentBrowser/Private/ContentBrowserUtils.h b/Engine/Source/Editor/ContentBrowser/Private/ContentBrowserUtils.h
--- a/Engine/Source/Editor/ContentBrowser/Private/ContentBrowserUtils.h
+++ b/Engine/Source/Editor/ContentBrowser/Private/ContentBrowserUtils.h
@@ -4,6 +4,18 @@
 
 namespace ContentBrowserUtils
 {
+	enum class ECBFolderCategory : uint8
+	{
+		GameContent,
+		EngineContent,
+		PluginContent,
+		DeveloperContent,
+
+		GameClasses,
+		EngineClasses,
+		PluginClasses,
+	};
+
 	/** Loads the specified object if needed and opens the asset editor for it */
 	bool OpenEditorForAsset(const FString& ObjectPath);
 
@@ -120,6 +132,9 @@ namespace ContentBrowserUtils
 	/** Returns true if the specified asset that uses shared thumbnails has a thumbnail assigned to it */
 	bool AssetHasCustomThumbnail( const FAssetData& AssetData );
 
+	/** Extract the category of the given path */
+	ECBFolderCategory GetFolderCategory( const FString& InPath );
+
 	/** Returns true if the passed-in path is a engine folder */
 	bool IsEngineFolder( const FString& InPath );
 
diff --git a/Engine/Source/Editor/ContentBrowser/Private/NativeClassHierarchy.cpp b/Engine/Source/Editor/ContentBrowser/Private/NativeClassHierarchy.cpp
--- a/Engine/Source/Editor/ContentBrowser/Private/NativeClassHierarchy.cpp
+++ b/Engine/Source/Editor/ContentBrowser/Private/NativeClassHierarchy.cpp
@@ -393,7 +393,12 @@ bool FNativeClassHierarchy::GetClassPath(UClass* InClass, FString& OutClassPath,
 
 	// Work out the final path to this class within the hierarchy (which isn't the same as the path on disk)
 	const FString ClassModuleRelativePath = ClassModuleRelativeIncludePath.Left(ClassModuleRelativeIncludePath.Find(TEXT("/"), ESearchCase::CaseSensitive, ESearchDir::FromEnd));
-	OutClassPath = FString(TEXT("/")) + RootNodeName.ToString() + TEXT("/") + ClassModuleName.ToString() + TEXT("/") + ClassModuleRelativePath;
+	OutClassPath = FString(TEXT("/")) + RootNodeName.ToString() + TEXT("/") + ClassModuleName.ToString();
+
+	if(!ClassModuleRelativePath.IsEmpty())
+	{
+		OutClassPath += TEXT("/") + ClassModuleRelativePath;
+	}
 
 	if(bIncludeClassName)
 	{
diff --git a/Engine/Source/Editor/ContentBrowser/Private/SContentBrowser.cpp b/Engine/Source/Editor/ContentBrowser/Private/SContentBrowser.cpp
--- a/Engine/Source/Editor/ContentBrowser/Private/SContentBrowser.cpp
+++ b/Engine/Source/Editor/ContentBrowser/Private/SContentBrowser.cpp
@@ -683,19 +683,34 @@ void SContentBrowser::SyncToAssets( const TArray<FAssetData>& AssetDataList, con
 		for (int32 AssetIdx = AssetDataList.Num() - 1; AssetIdx >= 0 && ( !bDisplayDev || !bDisplayEngine || !bDisplayPlugins ); --AssetIdx)
 		{
 			const FAssetData& Item = AssetDataList[AssetIdx];
-			if ( !bDisplayDev && ContentBrowserUtils::IsDevelopersFolder( Item.PackagePath.ToString() ) )
+
+			FString PackagePath;
+			if ( Item.AssetClass == NAME_Class )
+			{
+				// Classes are found in the /Classes_ roots
+				TSharedRef<FNativeClassHierarchy> NativeClassHierarchy = FContentBrowserSingleton::Get().GetNativeClassHierarchy();
+				NativeClassHierarchy->GetClassPath(Cast<UClass>(Item.GetAsset()), PackagePath, false/*bIncludeClassName*/);
+			}
+			else
+			{
+				// All other assets are found by their package path
+				PackagePath = Item.PackagePath.ToString();
+			}
+
+			const ContentBrowserUtils::ECBFolderCategory FolderCategory = ContentBrowserUtils::GetFolderCategory( PackagePath );
+			if ( !bDisplayDev && FolderCategory == ContentBrowserUtils::ECBFolderCategory::DeveloperContent )
 			{
 				bDisplayDev = true;
 				GetMutableDefault<UContentBrowserSettings>()->SetDisplayDevelopersFolder(true, true);
 				bRepopulate = true;
 			}
-			else if ( !bDisplayEngine && ContentBrowserUtils::IsEngineFolder( Item.PackagePath.ToString() ) )
+			else if ( !bDisplayEngine && (FolderCategory == ContentBrowserUtils::ECBFolderCategory::EngineContent || FolderCategory == ContentBrowserUtils::ECBFolderCategory::EngineClasses) )
 			{
 				bDisplayEngine = true;
 				GetMutableDefault<UContentBrowserSettings>()->SetDisplayEngineFolder(true, true);
 				bRepopulate = true;
 			}
-			else if ( !bDisplayPlugins && ContentBrowserUtils::IsPluginFolder( Item.PackagePath.ToString() ) )
+			else if ( !bDisplayPlugins && (FolderCategory == ContentBrowserUtils::ECBFolderCategory::PluginContent || FolderCategory == ContentBrowserUtils::ECBFolderCategory::PluginClasses) )
 			{
 				bDisplayPlugins = true;
 				GetMutableDefault<UContentBrowserSettings>()->SetDisplayPluginFolders(true, true);
