Increased the time that the quick jump text in the Content Browser remains active before being reset.

diff --git a/Engine/Source/Editor/ContentBrowser/Private/SAssetView.cpp b/Engine/Source/Editor/ContentBrowser/Private/SAssetView.cpp
--- a/Engine/Source/Editor/ContentBrowser/Private/SAssetView.cpp
+++ b/Engine/Source/Editor/ContentBrowser/Private/SAssetView.cpp
@@ -23,7 +23,7 @@ namespace
 	const double TimeBetweenAddingNewAssets = 4.0;
 
 	/** Time delay between performing the last jump, and the jump term being reset */
-	const double JumpDelaySeconds = 0.6;
+	const double JumpDelaySeconds = 2.0;
 }
 
 #define MAX_THUMBNAIL_SIZE 4096
