Made the warning text in SAssetView HitTestInvisible so that it doesn't block summoning the context menu

UE-6644 - HitTestInvisible doesn't work properly with SRichTextBlock

diff --git a/Engine/Source/Editor/ContentBrowser/Private/SAssetView.cpp b/Engine/Source/Editor/ContentBrowser/Private/SAssetView.cpp
--- a/Engine/Source/Editor/ContentBrowser/Private/SAssetView.cpp
+++ b/Engine/Source/Editor/ContentBrowser/Private/SAssetView.cpp
@@ -258,8 +258,9 @@ void SAssetView::Construct( const FArguments& InArgs )
 			.Padding(FMargin(0, 14, 0, 0))
 			[
 				// A warning to display when there are no assets to show
-				SAssignNew( WarningTextWidget, SRichTextBlock )
+				SNew( STextBlock )
 				.Justification( ETextJustify::Center )
+				.Text( this, &SAssetView::GetAssetShowWarningText )
 				.Visibility( this, &SAssetView::IsAssetShowWarningTextVisible )
 				.AutoWrapText( true )
 			]
@@ -930,16 +931,6 @@ void SAssetView::Tick( const FGeometry& AllottedGeometry, const double InCurrent
 	{
 		ResetQuickJump();
 	}
-
-	if ( IsAssetShowWarningTextVisible() == EVisibility::Visible )
-	{
-		FText WarningText = GetAssetShowWarningText();
-		if ( WarningText.CompareTo( CachedWarningText ) != 0)
-		{
-			CachedWarningText = WarningText;
-			WarningTextWidget->SetText( WarningText );
-		}
-	}
 }
 
 void SAssetView::CalculateFillScale( const FGeometry& AllottedGeometry )
@@ -3625,7 +3616,7 @@ void SAssetView::OnSortColumnHeader(const EColumnSortPriority::Type SortPriority
 
 EVisibility SAssetView::IsAssetShowWarningTextVisible() const
 {
-	return FilteredAssetItems.Num() > 0 ? EVisibility::Collapsed : EVisibility::Visible;
+	return FilteredAssetItems.Num() > 0 ? EVisibility::Collapsed : EVisibility::HitTestInvisible;
 }
 
 FText SAssetView::GetAssetShowWarningText() const
diff --git a/Engine/Source/Editor/ContentBrowser/Private/SAssetView.h b/Engine/Source/Editor/ContentBrowser/Private/SAssetView.h
--- a/Engine/Source/Editor/ContentBrowser/Private/SAssetView.h
+++ b/Engine/Source/Editor/ContentBrowser/Private/SAssetView.h
@@ -867,10 +867,4 @@ private:
 
 	/** Data for the asset quick-jump */
 	FQuickJumpData QuickJumpData;
-
-	/** Cached warning text that is checked against each tick when the warning block is visible. */
-	FText CachedWarningText;
-
-	/** The Warning text widget. */
-	TSharedPtr<SRichTextBlock> WarningTextWidget;
 };
