Fixed issue where "Assets could not be loaded" notification was misleadingly appearing when fixing up redirectors on a folder.

diff --git a/Engine/Source/Editor/ContentBrowser/Private/ContentBrowserUtils.cpp b/Engine/Source/Editor/ContentBrowser/Private/ContentBrowserUtils.cpp
--- a/Engine/Source/Editor/ContentBrowser/Private/ContentBrowserUtils.cpp
+++ b/Engine/Source/Editor/ContentBrowser/Private/ContentBrowserUtils.cpp
@@ -277,7 +277,7 @@ bool ContentBrowserUtils::OpenEditorForAsset(const TArray<UObject*>& Assets)
 	return false;
 }
 
-bool ContentBrowserUtils::LoadAssetsIfNeeded(const TArray<FString>& ObjectPaths, TArray<UObject*>& LoadedObjects, bool bAllowedToPromptToLoadAssets)
+bool ContentBrowserUtils::LoadAssetsIfNeeded(const TArray<FString>& ObjectPaths, TArray<UObject*>& LoadedObjects, bool bAllowedToPromptToLoadAssets, bool bLoadRedirects)
 {
 	bool bAnyObjectsWereLoadedOrUpdated = false;
 
@@ -328,15 +328,16 @@ bool ContentBrowserUtils::LoadAssetsIfNeeded(const TArray<FString>& ObjectPaths,
 
 		GIsEditorLoadingPackage = true;
 
+		// We usually don't want to follow redirects when loading objects for the Content Browser.  It would
+		// allow a user to interact with a ghost/unverified asset as if it were still alive.
+		// This can be overridden by providing bLoadRedirects = true as a parameter.
+		const ELoadFlags LoadFlags = bLoadRedirects ? LOAD_None : LOAD_NoRedirects;
+
 		bool bSomeObjectsFailedToLoad = false;
 		for (int32 PathIdx = 0; PathIdx < UnloadedObjectPaths.Num(); ++PathIdx)
 		{
 			const FString& ObjectPath = UnloadedObjectPaths[PathIdx];
 
-			// We never want to follow redirects when loading objects for the Content Browser.  It would
-			// allow a user to interact with a ghost/unverified asset as if it were still alive.
-			const ELoadFlags LoadFlags = LOAD_NoRedirects;
-
 			// Load up the object
 			UObject* LoadedObject = LoadObject<UObject>(NULL, *ObjectPath, NULL, LoadFlags, NULL);
 			if ( LoadedObject )
diff --git a/Engine/Source/Editor/ContentBrowser/Private/ContentBrowserUtils.h b/Engine/Source/Editor/ContentBrowser/Private/ContentBrowserUtils.h
--- a/Engine/Source/Editor/ContentBrowser/Private/ContentBrowserUtils.h
+++ b/Engine/Source/Editor/ContentBrowser/Private/ContentBrowserUtils.h
@@ -32,7 +32,7 @@ namespace ContentBrowserUtils
 	  * @param LoadedObjects The returned list of objects that were already loaded or loaded by this method.
 	  * @return false if user canceled after being warned about loading very many packages.
 	  */
-	bool LoadAssetsIfNeeded(const TArray<FString>& ObjectPaths, TArray<UObject*>& LoadedObjects, bool bAllowedToPromptToLoadAssets = true);
+	bool LoadAssetsIfNeeded(const TArray<FString>& ObjectPaths, TArray<UObject*>& LoadedObjects, bool bAllowedToPromptToLoadAssets = true, bool bLoadRedirects = false);
 
 	/**
 	 * Determines if enough assets are unloaded that we should prompt the user to load them instead of loading them automatically
diff --git a/Engine/Source/Editor/ContentBrowser/Private/PathContextMenu.cpp b/Engine/Source/Editor/ContentBrowser/Private/PathContextMenu.cpp
--- a/Engine/Source/Editor/ContentBrowser/Private/PathContextMenu.cpp
+++ b/Engine/Source/Editor/ContentBrowser/Private/PathContextMenu.cpp
@@ -684,7 +684,9 @@ void FPathContextMenu::ExecuteFixUpRedirectorsInFolder()
 		}
 
 		TArray<UObject*> Objects;
-		if (ContentBrowserUtils::LoadAssetsIfNeeded(ObjectPaths, Objects))
+		const bool bAllowedToPromptToLoadAssets = true;
+		const bool bLoadRedirects = true;
+		if (ContentBrowserUtils::LoadAssetsIfNeeded(ObjectPaths, Objects, bAllowedToPromptToLoadAssets, bLoadRedirects))
 		{
 			// Transform Objects array to ObjectRedirectors array
 			TArray<UObjectRedirector*> Redirectors;
