[INTEGRATE] Change 2426931 by [REDACTED_PERSON][REDACTED_HANDLE] on 2015/01/30 23:38:33

	Hide collisions by default and add a menu entry in the content browser to show them

diff --git a/Engine/Source/Editor/ContentBrowser/Private/SContentBrowser.cpp b/Engine/Source/Editor/ContentBrowser/Private/SContentBrowser.cpp
--- a/Engine/Source/Editor/ContentBrowser/Private/SContentBrowser.cpp
+++ b/Engine/Source/Editor/ContentBrowser/Private/SContentBrowser.cpp
@@ -364,7 +364,7 @@ void SContentBrowser::Construct( const FArguments& InArgs, const FName& InInstan
 
 			// Sources View
 			+ SSplitter::Slot()
-			.Value(0.3f)
+			.Value(0.1f)
 			[
 				SNew(SVerticalBox)
 				.Visibility( this, &SContentBrowser::GetSourcesViewVisibility )
@@ -419,9 +419,10 @@ void SContentBrowser::Construct( const FArguments& InArgs, const FName& InInstan
 
 					// Collection View
 					+ SSplitter::Slot()
-					.Value(0.1f)
+					.Value(0.9f)
 					[
 						SNew(SBorder)
+						.Visibility( this, &SContentBrowser::GetCollectionViewVisibility )
 						.Padding(FMargin(3))
 						.BorderImage( FEditorStyle::GetBrush("ToolPanel.GroupBorder") )
 						[
@@ -560,6 +561,7 @@ void SContentBrowser::Construct( const FArguments& InArgs, const FName& InInstan
 	AssetContextMenu->SetOnDuplicateRequested( FAssetContextMenu::FOnDuplicateRequested::CreateSP(this, &SContentBrowser::OnDuplicateRequested) );
 	AssetContextMenu->SetOnAssetViewRefreshRequested( FAssetContextMenu::FOnAssetViewRefreshRequested::CreateSP( this, &SContentBrowser::OnAssetViewRefreshRequested) );
 
+
 	// Select /Game by default
 	FSourcesData DefaultSourcesData;
 	TArray<FString> SelectedPaths;
@@ -583,6 +585,16 @@ void SContentBrowser::Construct( const FArguments& InArgs, const FName& InInstan
 	FAssetRegistryModule& AssetRegistryModule = FModuleManager::LoadModuleChecked<FAssetRegistryModule>(TEXT("AssetRegistry"));
 	AssetRegistryModule.Get().OnPathRemoved().AddSP(this, &SContentBrowser::HandlePathRemoved);
 
+	FContentBrowserModule& ContentBrowserModule = FModuleManager::GetModuleChecked<FContentBrowserModule>( TEXT("ContentBrowser") );
+	ContentBrowserModule.GetAllAssetViewViewMenuExtenders().Add( 
+		FContentBrowserMenuExtender::CreateLambda( [&]()
+		{
+			TSharedRef<FExtender> Extender = MakeShareable( new FExtender );
+
+			Extender->AddMenuExtension( FName("Folders"), EExtensionHook::After, nullptr, FMenuExtensionDelegate::CreateSP( this, &SContentBrowser::ExtendAssetViewMenu ) );
+			return Extender;
+		})
+	);
 	// Update the breadcrumb trail path
 	UpdatePath();
 }
@@ -615,6 +627,37 @@ void SContentBrowser::BindCommands()
 		));
 }
 
+void SContentBrowser::ExtendAssetViewMenu( FMenuBuilder& MenuBuilder )
+{
+	MenuBuilder.AddMenuEntry(
+		LOCTEXT("ShowCollectionOption", "Show Collections"),
+		LOCTEXT("ShowCollectionOptionToolTip", "Show the collections list in the view."),
+		FSlateIcon(),
+		FUIAction(
+		FExecuteAction::CreateSP(this, &SContentBrowser::ToggleShowCollections),
+		FCanExecuteAction(),
+		FIsActionChecked::CreateSP(this, &SContentBrowser::IsShowingCollections)
+		),
+		NAME_None,
+		EUserInterfaceActionType::ToggleButton
+		);
+}
+
+void SContentBrowser::ToggleShowCollections()
+{
+	return GetMutableDefault<UContentBrowserSettings>()->SetDisplayCollections(!GetDefault<UContentBrowserSettings>()->GetDisplayCollections());
+}
+
+bool SContentBrowser::IsShowingCollections() const
+{
+	return GetDefault<UContentBrowserSettings>()->GetDisplayCollections();
+}
+
+EVisibility SContentBrowser::GetCollectionViewVisibility() const
+{
+	return IsShowingCollections() ? EVisibility::Visible : EVisibility::Collapsed;
+}
+
 FText SContentBrowser::GetHighlightedText() const
 {
 	return TextFilter->GetRawFilterText();
diff --git a/Engine/Source/Editor/ContentBrowser/Private/SContentBrowser.h b/Engine/Source/Editor/ContentBrowser/Private/SContentBrowser.h
--- a/Engine/Source/Editor/ContentBrowser/Private/SContentBrowser.h
+++ b/Engine/Source/Editor/ContentBrowser/Private/SContentBrowser.h
@@ -320,6 +320,17 @@ private:
 	/** Bind our UI commands */
 	void BindCommands();
 
+	/** Adds menu extensions to the asset view */
+	void ExtendAssetViewMenu( FMenuBuilder& MenuBuilder );
+
+	/** Enables or disables collections from being shown */
+	void ToggleShowCollections();
+
+	/** @return Whether or not we are showing collections */
+	bool IsShowingCollections() const;
+
+	/** Gets the visibility of the collection view */
+	EVisibility GetCollectionViewVisibility() const;
 private:
 
 	/** The tab that contains this browser */
diff --git a/Engine/Source/Editor/UnrealEd/Classes/Settings/ContentBrowserSettings.h b/Engine/Source/Editor/UnrealEd/Classes/Settings/ContentBrowserSettings.h
--- a/Engine/Source/Editor/UnrealEd/Classes/Settings/ContentBrowserSettings.h
+++ b/Engine/Source/Editor/UnrealEd/Classes/Settings/ContentBrowserSettings.h
@@ -71,6 +71,18 @@ public:
 		return ( ( bExcludeOverride ? false : OverrideDisplayPluginFolders ) | DisplayPluginFolders );
 	}
 
+	/** Sets whether we are allowed to display collection folders or not */
+	void SetDisplayCollections(bool bInDisplayCollections)
+	{
+		DisplayCollections = bInDisplayCollections;
+	}
+
+	/** Gets whether we are allowed to display the collection folders or not*/
+	bool GetDisplayCollections() const
+	{
+		return DisplayCollections;
+	}
+
 	/**
 	 * Returns an event delegate that is executed when a setting has changed.
 	 *
@@ -108,6 +120,9 @@ private:
 	/** Temporary override for the DisplayPluginFolders setting */
 	bool OverrideDisplayPluginFolders;
 
+	UPROPERTY(config)
+	bool DisplayCollections;
+
 	// Holds an event delegate that is executed when a setting has changed.
 	static FSettingChangedEvent SettingChangedEvent;
 };
