[INTEGRATE] Change 2427068 by [REDACTED_PERSON][REDACTED_HANDLE] on 2015/01/31 11:48:03

	Further adjustments to the content browser source view splitter size

diff --git a/Engine/Source/Editor/ContentBrowser/Private/SContentBrowser.cpp b/Engine/Source/Editor/ContentBrowser/Private/SContentBrowser.cpp
--- a/Engine/Source/Editor/ContentBrowser/Private/SContentBrowser.cpp
+++ b/Engine/Source/Editor/ContentBrowser/Private/SContentBrowser.cpp
@@ -364,7 +364,7 @@ void SContentBrowser::Construct( const FArguments& InArgs, const FName& InInstan
 
 			// Sources View
 			+ SSplitter::Slot()
-			.Value(0.1f)
+			.Value(0.15f)
 			[
 				SNew(SVerticalBox)
 				.Visibility( this, &SContentBrowser::GetSourcesViewVisibility )
