[INTEGRATE] Change 2429903 by [REDACTED_PERSON][REDACTED_HANDLE] on 2015/02/03 10:00:38

	Adjusted content browser default thumbnail size

diff --git a/Engine/Source/Editor/ContentBrowser/Private/SContentBrowser.cpp b/Engine/Source/Editor/ContentBrowser/Private/SContentBrowser.cpp
--- a/Engine/Source/Editor/ContentBrowser/Private/SContentBrowser.cpp
+++ b/Engine/Source/Editor/ContentBrowser/Private/SContentBrowser.cpp
@@ -529,7 +529,7 @@ void SContentBrowser::Construct( const FArguments& InArgs, const FName& InInstan
 					.Padding( 0 )
 					[
 						SAssignNew(AssetViewPtr, SAssetView)
-						.ThumbnailScale( 0.0f )
+						.ThumbnailScale( 0.075f )
 						.OnPathSelected(this, &SContentBrowser::FolderEntered)
 						.OnAssetSelected(this, &SContentBrowser::OnAssetSelectionChanged)
 						.OnAssetsActivated(this, &SContentBrowser::OnAssetsActivated)
