[INTEGRATE] Change 2430090 by [REDACTED_PERSON][REDACTED_HANDLE] on 2015/02/03 12:23:43

	Made "Add Feature or Content Pack..." always available from the "Add New" menu

	It was previously only available when you had an asset path selected, but as this option has nothing to do with the currently selected path(s) it is now available regardless of whether you have an asset or class path selected.

diff --git a/Engine/Source/Editor/ContentBrowser/Private/NewAssetOrClassContextMenu.cpp b/Engine/Source/Editor/ContentBrowser/Private/NewAssetOrClassContextMenu.cpp
--- a/Engine/Source/Editor/ContentBrowser/Private/NewAssetOrClassContextMenu.cpp
+++ b/Engine/Source/Editor/ContentBrowser/Private/NewAssetOrClassContextMenu.cpp
@@ -221,10 +221,7 @@ void FNewAssetOrClassContextMenu::MakeContextMenu(
 				LOCTEXT( "GetContentText", "Add Feature or Content Pack..." ),
 				LOCTEXT( "GetContentTooltip", "Add features and content packs to the project." ),
 				FSlateIcon( FEditorStyle::GetStyleSetName(), "ContentBrowser.AddContent" ),
-				FUIAction(
-					FExecuteAction::CreateStatic( &FNewAssetOrClassContextMenu::ExecuteGetContent, InOnGetContentRequested ),
-					CanExecuteAssetActionsDelegate
-					)
+				FUIAction( FExecuteAction::CreateStatic( &FNewAssetOrClassContextMenu::ExecuteGetContent, InOnGetContentRequested ) )
 				);
 		}
 		MenuBuilder.EndSection();
diff --git a/Engine/Source/Editor/ContentBrowser/Private/SContentBrowser.cpp b/Engine/Source/Editor/ContentBrowser/Private/SContentBrowser.cpp
--- a/Engine/Source/Editor/ContentBrowser/Private/SContentBrowser.cpp
+++ b/Engine/Source/Editor/ContentBrowser/Private/SContentBrowser.cpp
@@ -1380,10 +1380,16 @@ TSharedRef<SWidget> SContentBrowser::MakeAddNewContextMenu(bool bShowGetContent,
 	int32 NumAssetPaths, NumClassPaths;
 	ContentBrowserUtils::CountPathTypes(SourcesData.PackagePaths, NumAssetPaths, NumClassPaths);
 
+	// New feature packs don't depend on the current paths, so we always add this item if it was requested
+	FNewAssetOrClassContextMenu::FOnGetContentRequested OnGetContentRequested;
+	if(bShowGetContent)
+	{
+		OnGetContentRequested = FNewAssetOrClassContextMenu::FOnGetContentRequested::CreateSP(this, &SContentBrowser::OnAddContentRequested);
+	}
+
 	// Only the asset items if we have an asset path selected
 	FNewAssetOrClassContextMenu::FOnNewAssetRequested OnNewAssetRequested;
 	FNewAssetOrClassContextMenu::FOnImportAssetRequested OnImportAssetRequested;
-	FNewAssetOrClassContextMenu::FOnGetContentRequested OnGetContentRequested;
 	if(NumAssetPaths > 0)
 	{
 		OnNewAssetRequested = FNewAssetOrClassContextMenu::FOnNewAssetRequested::CreateSP(this, &SContentBrowser::NewAssetRequested);
@@ -1391,10 +1397,6 @@ TSharedRef<SWidget> SContentBrowser::MakeAddNewContextMenu(bool bShowGetContent,
 		{
 			OnImportAssetRequested = FNewAssetOrClassContextMenu::FOnImportAssetRequested::CreateSP(this, &SContentBrowser::ImportAsset);
 		}
-		if(bShowGetContent)
-		{
-			OnGetContentRequested = FNewAssetOrClassContextMenu::FOnGetContentRequested::CreateSP(this, &SContentBrowser::OnAddContentRequested);
-		}
 	}
 
 	// This menu always lets you create classes, but uses your default project source folder if the selected path is invalid for creating classes
