[INTEGRATE] Change 2432062 by [REDACTED_PERSON][REDACTED_HANDLE] on 2015/02/04 13:56:41

	Improved and unified the folder drag and drop validation in the Content Browser

	UE-9025 - User gets green check mark when trying to move class into content folder

	The Content Browser folder drag and drop validation had not been updated to handle class folders. This meant that it was possible to crash the editor by dragging class folders around, and also meant that you were given misinformation when trying to copy a class around.

	I've unified all the duplicated logic from SPathView, SAssetView, SAssetTreeItem, and SAssetViewItem into utility functions inside a new DragDropHandler in the Content Browser. This class correctly validates that you can only interact with assets and asset paths, but doesn't outright prevent you being able to include classes and class paths in your drag operations (it just warns you about them, and the strips them out of the drop lists).

diff --git a/Engine/Source/Editor/ContentBrowser/Private/AssetViewWidgets.cpp b/Engine/Source/Editor/ContentBrowser/Private/AssetViewWidgets.cpp
--- a/Engine/Source/Editor/ContentBrowser/Private/AssetViewWidgets.cpp
+++ b/Engine/Source/Editor/ContentBrowser/Private/AssetViewWidgets.cpp
@@ -9,6 +9,7 @@
 #include "SThumbnailEditModeTools.h"
 #include "DragAndDrop/AssetDragDropOp.h"
 #include "DragAndDrop/AssetPathDragDropOp.h"
+#include "DragDropHandler.h"
 #include "BreakIterator.h"
 #include "SInlineEditableTextBlock.h"
 
@@ -186,39 +187,18 @@ TSharedPtr<IToolTip> SAssetViewItem::GetToolTip()
 	return ShouldAllowToolTip.Get() ? SCompoundWidget::GetToolTip() : NULL;
 }
 
+bool SAssetViewItem::ValidateDragDrop( const FGeometry& MyGeometry, const FDragDropEvent& DragDropEvent ) const
+{
+	return IsFolder() && DragDropHandler::ValidateDragDropOnAssetFolder(MyGeometry, DragDropEvent, StaticCastSharedPtr<FAssetViewFolder>(AssetItem)->FolderPath);
+}
+
 void SAssetViewItem::OnDragEnter( const FGeometry& MyGeometry, const FDragDropEvent& DragDropEvent )
 {
 	bDraggedOver = false;
 
-	if(IsFolder())
+	if (ValidateDragDrop(MyGeometry, DragDropEvent))
 	{
-		TSharedPtr<FDragDropOperation> Operation = DragDropEvent.GetOperation();
-		if (Operation->IsOfType<FAssetDragDropOp>())
-		{
-			TArray< FAssetData > AssetDatas = AssetUtil::ExtractAssetDataFromDrag( DragDropEvent );
-
-			if ( AssetDatas.Num() > 0 )
-			{
-				TSharedPtr< FAssetDragDropOp > DragDropOp = StaticCastSharedPtr< FAssetDragDropOp >( DragDropEvent.GetOperation() );	
-				DragDropOp->SetToolTip( LOCTEXT( "OnDragAssetsOverFolder", "Move or Copy Asset(s)" ), FEditorStyle::GetBrush(TEXT("Graph.ConnectorFeedback.OK")) );
-			}
-			bDraggedOver = true;
-		}
-		else if (Operation->IsOfType<FAssetPathDragDropOp>())
-		{
-			TSharedPtr<FAssetPathDragDropOp> DragDropOp = StaticCastSharedPtr<FAssetPathDragDropOp>( DragDropEvent.GetOperation() );
-			bool bCanDrop = DragDropOp->PathNames.Num() > 0;
-			if ( DragDropOp->PathNames.Contains(StaticCastSharedPtr<FAssetViewFolder>(AssetItem)->FolderPath) )
-			{
-				// You can't drop a folder onto itself
-				bCanDrop = false;
-			}
-
-			if(bCanDrop)
-			{
-				bDraggedOver = true;
-			}
-		}
+		bDraggedOver = true;
 	}
 }
 	
@@ -226,11 +206,19 @@ void SAssetViewItem::OnDragLeave( const FDragDropEvent& DragDropEvent )
 {
 	if(IsFolder())
 	{
-		TSharedPtr< FAssetDragDropOp > DragDropOp = DragDropEvent.GetOperationAs< FAssetDragDropOp >();
-		if(DragDropOp.IsValid())
+		TSharedPtr<FDragDropOperation> Operation = DragDropEvent.GetOperation();
+		if (Operation.IsValid())
 		{
-			DragDropOp->ResetToDefaultToolTip();
+			Operation->SetCursorOverride(TOptional<EMouseCursor::Type>());
+
+			if (Operation->IsOfType<FAssetDragDropOp>())
+			{
+				TSharedPtr<FAssetDragDropOp> DragDropOp = StaticCastSharedPtr<FAssetDragDropOp>(Operation);
+				DragDropOp->ResetToDefaultToolTip();
+			}
 		}
+
+		bDraggedOver = false;
 	}
 
 	bDraggedOver = false;
@@ -238,44 +226,12 @@ void SAssetViewItem::OnDragLeave( const FDragDropEvent& DragDropEvent )
 
 FReply SAssetViewItem::OnDragOver( const FGeometry& MyGeometry, const FDragDropEvent& DragDropEvent )
 {
-	if(IsFolder())
-	{
-		TSharedPtr<FDragDropOperation> Operation = DragDropEvent.GetOperation();
-		if (!Operation.IsValid())
-		{
-			return FReply::Unhandled();
-		}
-
-		if (Operation->IsOfType<FExternalDragOperation>())
-		{
-			TSharedPtr<FExternalDragOperation> DragDropOp = StaticCastSharedPtr<FExternalDragOperation>(Operation);
-			if ( DragDropOp->HasFiles() )
-			{
-				bDraggedOver = true;
-				return FReply::Handled();
-			}
-		}
-		else if (Operation->IsOfType<FAssetDragDropOp>())
-		{
-			bDraggedOver = true;
-			return FReply::Handled();
-		}
-		else if (Operation->IsOfType<FAssetPathDragDropOp>())
-		{
-			TSharedPtr<FAssetPathDragDropOp> DragDropOp = StaticCastSharedPtr<FAssetPathDragDropOp>(Operation);
-			bool bCanDrop = DragDropOp->PathNames.Num() > 0;
-			if ( DragDropOp->PathNames.Contains(StaticCastSharedPtr<FAssetViewFolder>(AssetItem)->FolderPath) )
-			{
-				// You can't drop a folder onto itself
-				bCanDrop = false;
-			}
+	bDraggedOver = false;
 
-			if(bCanDrop)
-			{
-				bDraggedOver = true;
-			}
-			return FReply::Handled();
-		}
+	if (ValidateDragDrop(MyGeometry, DragDropEvent))
+	{
+		bDraggedOver = true;
+		return FReply::Handled();
 	}
 
 	return FReply::Unhandled();
@@ -285,7 +241,7 @@ FReply SAssetViewItem::OnDrop( const FGeometry& MyGeometry, const FDragDropEvent
 {
 	bDraggedOver = false;
 
-	if(IsFolder())
+	if (ValidateDragDrop(MyGeometry, DragDropEvent))
 	{
 		check(AssetItem->GetType() == EAssetItemType::Folder);
 
@@ -298,39 +254,19 @@ FReply SAssetViewItem::OnDrop( const FGeometry& MyGeometry, const FDragDropEvent
 		if (Operation->IsOfType<FExternalDragOperation>())
 		{
 			TSharedPtr<FExternalDragOperation> DragDropOp = StaticCastSharedPtr<FExternalDragOperation>(Operation);
-
-			if ( DragDropOp->HasFiles() )
-			{
-				OnFilesDragDropped.ExecuteIfBound(DragDropOp->GetFiles(), StaticCastSharedPtr<FAssetViewFolder>(AssetItem)->FolderPath);
-			}
-
+			OnFilesDragDropped.ExecuteIfBound(DragDropOp->GetFiles(), StaticCastSharedPtr<FAssetViewFolder>(AssetItem)->FolderPath);
 			return FReply::Handled();
 		}
 		else if (Operation->IsOfType<FAssetPathDragDropOp>())
 		{
 			TSharedPtr<FAssetPathDragDropOp> DragDropOp = StaticCastSharedPtr<FAssetPathDragDropOp>(Operation);
-
-			bool bCanDrop = DragDropOp->PathNames.Num() > 0;
-
-			if ( DragDropOp->PathNames.Contains(StaticCastSharedPtr<FAssetViewFolder>(AssetItem)->FolderPath) )
-			{
-				// You can't drop a folder onto itself
-				bCanDrop = false;
-			}
-
-			if ( bCanDrop )
-			{
-				OnPathsDragDropped.ExecuteIfBound(DragDropOp->PathNames, StaticCastSharedPtr<FAssetViewFolder>(AssetItem)->FolderPath);
-			}
-
+			OnPathsDragDropped.ExecuteIfBound(DragDropOp->PathNames, StaticCastSharedPtr<FAssetViewFolder>(AssetItem)->FolderPath);
 			return FReply::Handled();
 		}
 		else if (Operation->IsOfType<FAssetDragDropOp>())
 		{
 			TSharedPtr<FAssetDragDropOp> DragDropOp = StaticCastSharedPtr<FAssetDragDropOp>(Operation);
-
 			OnAssetsDragDropped.ExecuteIfBound(DragDropOp->AssetData, StaticCastSharedPtr<FAssetViewFolder>(AssetItem)->FolderPath);
-
 			return FReply::Handled();
 		}
 	}
diff --git a/Engine/Source/Editor/ContentBrowser/Private/AssetViewWidgets.h b/Engine/Source/Editor/ContentBrowser/Private/AssetViewWidgets.h
--- a/Engine/Source/Editor/ContentBrowser/Private/AssetViewWidgets.h
+++ b/Engine/Source/Editor/ContentBrowser/Private/AssetViewWidgets.h
@@ -125,6 +125,9 @@ public:
 	virtual void OnToolTipClosing() override;
 
 protected:
+	/** Used by OnDragEnter, OnDragOver, and OnDrop to check and update the validity of the drag operation */
+	bool ValidateDragDrop( const FGeometry& MyGeometry, const FDragDropEvent& DragDropEvent ) const;
+
 	/** Handles starting a name change */
 	virtual void HandleBeginNameChange( const FText& OriginalText );
 
diff --git a/Engine/Source/Editor/ContentBrowser/Private/DragDropHandler.cpp b/Engine/Source/Editor/ContentBrowser/Private/DragDropHandler.cpp
new file mode 100644
--- /dev/null
+++ b/Engine/Source/Editor/ContentBrowser/Private/DragDropHandler.cpp
@@ -0,0 +1,200 @@
+// Copyright 1998-2015 Epic Games, Inc. All Rights Reserved.
+
+#include "ContentBrowserPCH.h"
+#include "DragDropHandler.h"
+
+#include "DragAndDrop/AssetDragDropOp.h"
+#include "DragAndDrop/AssetPathDragDropOp.h"
+#include "ContentBrowserUtils.h"
+
+#define LOCTEXT_NAMESPACE "ContentBrowser"
+
+bool DragDropHandler::ValidateDragDropOnAssetFolder(const FGeometry& MyGeometry, const FDragDropEvent& DragDropEvent, const FString& TargetPath)
+{
+	TSharedPtr<FDragDropOperation> Operation = DragDropEvent.GetOperation();
+	if (!Operation.IsValid())
+	{
+		return false;
+	}
+
+	bool bIsValidDrag = false;
+	TOptional<EMouseCursor::Type> NewDragCursor;
+
+	const bool bIsAssetPath = !ContentBrowserUtils::IsClassPath(TargetPath);
+
+	if (Operation->IsOfType<FAssetDragDropOp>())
+	{
+		TSharedPtr<FAssetDragDropOp> DragDropOp = StaticCastSharedPtr<FAssetDragDropOp>(Operation);
+
+		if (bIsAssetPath)
+		{
+			int32 NumAssetItems, NumClassItems;
+			ContentBrowserUtils::CountItemTypes(DragDropOp->AssetData, NumAssetItems, NumClassItems);
+
+			if (NumAssetItems > 0)
+			{
+				const FText MoveOrCopyText = (NumAssetItems > 1)
+					? FText::Format(LOCTEXT("OnDragAssetsOverFolder_MultipleAssets", "Move or copy '{0}' and {1} other(s)"), FText::FromName(DragDropOp->AssetData[0].AssetName), FText::AsNumber(NumAssetItems - 1))
+					: FText::Format(LOCTEXT("OnDragAssetsOverFolder_SingularAsset", "Move or copy '{0}'"), FText::FromName(DragDropOp->AssetData[0].AssetName));
+
+				if (NumClassItems > 0)
+				{
+					DragDropOp->SetToolTip(FText::Format(LOCTEXT("OnDragAssetsOverFolder_AssetsAndClasses", "{0}\n\n{1} C++ class(es) will be ignored as they cannot be moved or copied"), MoveOrCopyText, FText::AsNumber(NumClassItems)), FEditorStyle::GetBrush(TEXT("Graph.ConnectorFeedback.OKWarn")));
+					bIsValidDrag = true;
+				}
+				else
+				{
+					DragDropOp->SetToolTip(MoveOrCopyText, FEditorStyle::GetBrush(TEXT("Graph.ConnectorFeedback.OK")));
+					bIsValidDrag = true;
+				}
+			}
+			else if (NumClassItems > 0)
+			{
+				DragDropOp->SetToolTip(LOCTEXT("OnDragAssetsOverFolder_OnlyClasses", "C++ classes cannot be moved or copied"), FEditorStyle::GetBrush(TEXT("Graph.ConnectorFeedback.Error")));
+			}
+		}
+		else
+		{
+			DragDropOp->SetToolTip(FText::Format(LOCTEXT("OnDragAssetsOverFolder_InvalidFolder", "'{0}' is not a valid place to drop assets"), FText::FromString(TargetPath)), FEditorStyle::GetBrush(TEXT("Graph.ConnectorFeedback.Error")));
+		}
+	}
+	else if (Operation->IsOfType<FAssetPathDragDropOp>())
+	{
+		TSharedPtr<FAssetPathDragDropOp> DragDropOp = StaticCastSharedPtr<FAssetPathDragDropOp>(Operation);
+
+		if (DragDropOp->PathNames.Num() == 1 && DragDropOp->PathNames[0] == TargetPath)
+		{
+			DragDropOp->SetToolTip(LOCTEXT("OnDragFoldersOverFolder_CannotSelfDrop", "Cannot move or copy a folder onto itself"), FEditorStyle::GetBrush(TEXT("Graph.ConnectorFeedback.Error")));
+		}
+		else if (bIsAssetPath)
+		{
+			int32 NumAssetPaths, NumClassPaths;
+			ContentBrowserUtils::CountPathTypes(DragDropOp->PathNames, NumAssetPaths, NumClassPaths);
+
+			if (NumAssetPaths > 0)
+			{
+				const FText MoveOrCopyText = (NumAssetPaths > 1)
+					? FText::Format(LOCTEXT("OnDragFoldersOverFolder_MultipleFolders", "Move or copy '{0}' and {1} other(s)"), FText::FromString(DragDropOp->PathNames[0]), FText::AsNumber(NumAssetPaths - 1))
+					: FText::Format(LOCTEXT("OnDragFoldersOverFolder_SingularFolder", "Move or copy '{0}'"), FText::FromString(DragDropOp->PathNames[0]));
+
+				if (NumClassPaths > 0)
+				{
+					DragDropOp->SetToolTip(FText::Format(LOCTEXT("OnDragFoldersOverFolder_AssetsAndClasses", "{0}\n\n{1} C++ class folder(s) will be ignored as they cannot be moved or copied"), MoveOrCopyText, FText::AsNumber(NumClassPaths)), FEditorStyle::GetBrush(TEXT("Graph.ConnectorFeedback.OKWarn")));
+					bIsValidDrag = true;
+				}
+				else
+				{
+					DragDropOp->SetToolTip(MoveOrCopyText, FEditorStyle::GetBrush(TEXT("Graph.ConnectorFeedback.OK")));
+					bIsValidDrag = true;
+				}
+			}
+			else if (NumClassPaths > 0)
+			{
+				DragDropOp->SetToolTip(LOCTEXT("OnDragFoldersOverFolder_OnlyClasses", "C++ class folders cannot be moved or copied"), FEditorStyle::GetBrush(TEXT("Graph.ConnectorFeedback.Error")));
+			}
+		}
+		else
+		{
+			DragDropOp->SetToolTip(FText::Format(LOCTEXT("OnDragFoldersOverFolder_InvalidFolder", "'{0}' is not a valid place to drop folders"), FText::FromString(TargetPath)), FEditorStyle::GetBrush(TEXT("Graph.ConnectorFeedback.Error")));
+		}
+	}
+	else if (Operation->IsOfType<FExternalDragOperation>())
+	{
+		TSharedPtr<FExternalDragOperation> DragDropOp = StaticCastSharedPtr<FExternalDragOperation>(Operation);
+		bIsValidDrag = DragDropOp->HasFiles() && bIsAssetPath;
+	}
+
+	// Set the default slashed circle if this drag is invalid and a drag operation hasn't set NewDragCursor to something custom
+	if (!bIsValidDrag && !NewDragCursor.IsSet())
+	{
+		NewDragCursor = EMouseCursor::SlashedCircle;
+	}
+	Operation->SetCursorOverride(NewDragCursor);
+
+	return bIsValidDrag;
+}
+
+void DragDropHandler::HandleAssetsDroppedOnAssetFolder(const TSharedRef<SWidget>& ParentWidget, const TArray<FAssetData>& AssetList, const FString& TargetPath, const FText& TargetDisplayName, FExecuteCopyOrMoveAssets CopyActionHandler, FExecuteCopyOrMoveAssets MoveActionHandler)
+{
+	TArray<FAssetData> FinalAssetList;
+	FinalAssetList.Reserve(AssetList.Num());
+
+	// Remove any classes from the asset list
+	for (const FAssetData& AssetData : AssetList)
+	{
+		if (AssetData.AssetClass != NAME_Class)
+		{
+			FinalAssetList.Add(AssetData);
+		}
+	}
+
+	FMenuBuilder MenuBuilder(/*bInShouldCloseWindowAfterMenuSelection=*/true, nullptr);
+	const FText MoveCopyHeaderString = FText::Format(LOCTEXT("AssetViewDropMenuHeading", "Move/Copy to {0}"), TargetDisplayName);
+	MenuBuilder.BeginSection("PathAssetMoveCopy", MoveCopyHeaderString);
+	{
+		MenuBuilder.AddMenuEntry(
+			LOCTEXT("DragDropCopy", "Copy Here"),
+			LOCTEXT("DragDropCopyTooltip", "Creates a copy of all dragged files in this folder."),
+			FSlateIcon(),
+			FUIAction(FExecuteAction::CreateLambda([=]() { CopyActionHandler.ExecuteIfBound(FinalAssetList, TargetPath); }))
+			);
+
+		MenuBuilder.AddMenuEntry(
+			LOCTEXT("DragDropMove", "Move Here"),
+			LOCTEXT("DragDropMoveTooltip", "Moves all dragged files to this folder."),
+			FSlateIcon(),
+			FUIAction(FExecuteAction::CreateLambda([=]() { MoveActionHandler.ExecuteIfBound(FinalAssetList, TargetPath); }))
+			);
+	}
+	MenuBuilder.EndSection();
+
+	FSlateApplication::Get().PushMenu(
+		ParentWidget,
+		MenuBuilder.MakeWidget(),
+		FSlateApplication::Get().GetCursorPos(),
+		FPopupTransitionEffect(FPopupTransitionEffect::ContextMenu)
+		);
+}
+
+void DragDropHandler::HandleFoldersDroppedOnAssetFolder(const TSharedRef<SWidget>& ParentWidget, const TArray<FString>& PathNames, const FString& TargetPath, const FText& TargetDisplayName, FExecuteCopyOrMoveFolders CopyActionHandler, FExecuteCopyOrMoveFolders MoveActionHandler)
+{
+	TArray<FString> FinalPathNames;
+	FinalPathNames.Reserve(PathNames.Num());
+
+	// Remove any class paths from the list
+	for (const FString& PathName : PathNames)
+	{
+		if (!ContentBrowserUtils::IsClassPath(PathName))
+		{
+			FinalPathNames.Add(PathName);
+		}
+	}
+
+	FMenuBuilder MenuBuilder(/*bInShouldCloseWindowAfterMenuSelection=*/true, nullptr);
+	MenuBuilder.BeginSection("PathFolderMoveCopy", FText::Format(LOCTEXT("AssetViewDropMenuHeading", "Move/Copy to {0}"), TargetDisplayName));
+	{
+		MenuBuilder.AddMenuEntry(
+			LOCTEXT("DragDropCopyFolder", "Copy Folder Here"),
+			LOCTEXT("DragDropCopyFolderTooltip", "Creates a copy of all assets in the dragged folders to this folder, preserving folder structure."),
+			FSlateIcon(),
+			FUIAction(FExecuteAction::CreateLambda([=]() { MoveActionHandler.ExecuteIfBound(FinalPathNames, TargetPath); }))
+			);
+
+		MenuBuilder.AddMenuEntry(
+			LOCTEXT("DragDropMoveFolder", "Move Folder Here"),
+			LOCTEXT("DragDropMoveFolderTooltip", "Moves all assets in the dragged folders to this folder, preserving folder structure."),
+			FSlateIcon(),
+			FUIAction(FExecuteAction::CreateLambda([=]() { MoveActionHandler.ExecuteIfBound(FinalPathNames, TargetPath); }))
+			);
+	}
+	MenuBuilder.EndSection();
+
+	FSlateApplication::Get().PushMenu(
+		ParentWidget,
+		MenuBuilder.MakeWidget(),
+		FSlateApplication::Get().GetCursorPos(),
+		FPopupTransitionEffect(FPopupTransitionEffect::ContextMenu)
+		);
+}
+
+#undef LOCTEXT_NAMESPACE
diff --git a/Engine/Source/Editor/ContentBrowser/Private/DragDropHandler.h b/Engine/Source/Editor/ContentBrowser/Private/DragDropHandler.h
new file mode 100644
--- /dev/null
+++ b/Engine/Source/Editor/ContentBrowser/Private/DragDropHandler.h
@@ -0,0 +1,19 @@
+// Copyright 1998-2015 Epic Games, Inc. All Rights Reserved.
+
+#pragma once
+
+/** Common Content Browser drag-drop handler logic */
+namespace DragDropHandler
+{
+	DECLARE_DELEGATE_TwoParams(FExecuteCopyOrMoveAssets, TArray<FAssetData> /*AssetList*/, FString /*TargetPath*/);
+	DECLARE_DELEGATE_TwoParams(FExecuteCopyOrMoveFolders, TArray<FString> /*PathNames*/, FString /*TargetPath*/);
+
+	/** Used by OnDragEnter, OnDragOver, and OnDrop to check and update the validity of a drag-drop operation on an asset folder in the Content Browser */
+	bool ValidateDragDropOnAssetFolder(const FGeometry& MyGeometry, const FDragDropEvent& DragDropEvent, const FString& TargetPath);
+
+	/** Handle assets being dropped onto an asset folder in the Content Browser - this drop should have been externally validated by ValidateDragDropOnAssetFolder */
+	void HandleAssetsDroppedOnAssetFolder(const TSharedRef<SWidget>& ParentWidget, const TArray<FAssetData>& AssetList, const FString& TargetPath, const FText& TargetDisplayName, FExecuteCopyOrMoveAssets CopyActionHandler, FExecuteCopyOrMoveAssets MoveActionHandler);
+
+	/** Handle folders being dropped onto an asset folder in the Content Browser - this drop should have been externally validated by ValidateDragDropOnAssetFolder */
+	void HandleFoldersDroppedOnAssetFolder(const TSharedRef<SWidget>& ParentWidget, const TArray<FString>& PathNames, const FString& TargetPath, const FText& TargetDisplayName, FExecuteCopyOrMoveFolders CopyActionHandler, FExecuteCopyOrMoveFolders MoveActionHandler);
+}
diff --git a/Engine/Source/Editor/ContentBrowser/Private/SAssetView.cpp b/Engine/Source/Editor/ContentBrowser/Private/SAssetView.cpp
--- a/Engine/Source/Editor/ContentBrowser/Private/SAssetView.cpp
+++ b/Engine/Source/Editor/ContentBrowser/Private/SAssetView.cpp
@@ -6,6 +6,7 @@
 #include "AssetViewTypes.h"
 #include "DragAndDrop/AssetDragDropOp.h"
 #include "DragAndDrop/AssetPathDragDropOp.h"
+#include "DragDropHandler.h"
 #include "AssetThumbnail.h"
 #include "AssetViewWidgets.h"
 #include "FileHelpers.h"
@@ -3651,78 +3652,25 @@ bool SAssetView::HasSingleCollectionSource() const
 
 void SAssetView::OnAssetsDragDropped(const TArray<FAssetData>& AssetList, const FString& DestinationPath)
 {
-	// Do not display the menu if any of the assets are classes as they cannot be moved or copied
-	for( int32 AssetIndex = 0; AssetIndex < AssetList.Num(); AssetIndex++ )
-	{
-		const FAssetData& Asset = AssetList[AssetIndex];
-		if ( Asset.AssetClass == "Class" )
-		{
-			const FText MessageText = LOCTEXT("AssetTreeDropClassError", "The selection contains one or more 'Class' type assets, these cannot be moved or copied.");
-			FMessageDialog::Open(EAppMsgType::Ok, MessageText);
-			return;
-		}
-	}
-
-	FMenuBuilder MenuBuilder(/*bInShouldCloseWindowAfterMenuSelection=*/true, NULL);
-	const FText MoveCopyHeaderString = FText::Format( LOCTEXT("AssetViewDropMenuHeading", "Move/Copy to {0}"), FText::FromString( DestinationPath ) );
-	MenuBuilder.BeginSection("PathAssetMoveCopy", MoveCopyHeaderString);
-	{
-		MenuBuilder.AddMenuEntry(
-			LOCTEXT("DragDropCopy", "Copy Here"),
-			LOCTEXT("DragDropCopyTooltip", "Creates a copy of all dragged files in this folder."),
-			FSlateIcon(),
-			FUIAction(
-			FExecuteAction::CreateSP( this, &SAssetView::ExecuteDropCopy, AssetList, DestinationPath ),
-			FCanExecuteAction()
-			)
-			);
-
-		MenuBuilder.AddMenuEntry(
-			LOCTEXT("DragDropMove", "Move Here"),
-			LOCTEXT("DragDropMoveTooltip", "Moves all dragged files to this folder."),
-			FSlateIcon(),
-			FUIAction(
-			FExecuteAction::CreateSP( this, &SAssetView::ExecuteDropMove, AssetList, DestinationPath ),
-			FCanExecuteAction()
-			)
-			);
-	}
-	MenuBuilder.EndSection();
-
-	TWeakPtr< SWindow > ContextMenuWindow = FSlateApplication::Get().PushMenu(
-		SharedThis( this ),
-		MenuBuilder.MakeWidget(),
-		FSlateApplication::Get().GetCursorPos(),
-		FPopupTransitionEffect( FPopupTransitionEffect::ContextMenu )
+	DragDropHandler::HandleAssetsDroppedOnAssetFolder(
+		SharedThis(this), 
+		AssetList, 
+		DestinationPath, 
+		FText::FromString(FPaths::GetCleanFilename(DestinationPath)), 
+		DragDropHandler::FExecuteCopyOrMoveAssets::CreateSP(this, &SAssetView::ExecuteDropCopy),
+		DragDropHandler::FExecuteCopyOrMoveAssets::CreateSP(this, &SAssetView::ExecuteDropMove)
 		);
 }
 
 void SAssetView::OnPathsDragDropped(const TArray<FString>& PathNames, const FString& DestinationPath)
 {
-	FMenuBuilder MenuBuilder(/*bInShouldCloseWindowAfterMenuSelection=*/true, NULL);
-	MenuBuilder.BeginSection("PathFolderMoveCopy", FText::Format(LOCTEXT("AssetViewDropMenuHeading", "Move/Copy to {0}"), FText::FromString(DestinationPath)));
-	{
-		MenuBuilder.AddMenuEntry(
-			LOCTEXT("DragDropCopyFolder", "Copy Folder Here"),
-			LOCTEXT("DragDropCopyFolderTooltip", "Creates a copy of all assets in the dragged folders to this folder, preserving folder structure."),
-			FSlateIcon(),
-			FUIAction( FExecuteAction::CreateSP( this, &SAssetView::ExecuteDropCopyFolder, PathNames, DestinationPath ) )
-			);
-
-		MenuBuilder.AddMenuEntry(
-			LOCTEXT("DragDropMoveFolder", "Move Folder Here"),
-			LOCTEXT("DragDropMoveFolderTooltip", "Moves all assets in the dragged folders to this folder, preserving folder structure."),
-			FSlateIcon(),
-			FUIAction( FExecuteAction::CreateSP( this, &SAssetView::ExecuteDropMoveFolder, PathNames, DestinationPath ) )
-			);
-	}
-	MenuBuilder.EndSection();
-
-	TWeakPtr< SWindow > ContextMenuWindow = FSlateApplication::Get().PushMenu(
-		SharedThis( this ),
-		MenuBuilder.MakeWidget(),
-		FSlateApplication::Get().GetCursorPos(),
-		FPopupTransitionEffect( FPopupTransitionEffect::ContextMenu )
+	DragDropHandler::HandleFoldersDroppedOnAssetFolder(
+		SharedThis(this), 
+		PathNames, 
+		DestinationPath, 
+		FText::FromString(FPaths::GetCleanFilename(DestinationPath)), 
+		DragDropHandler::FExecuteCopyOrMoveFolders::CreateSP(this, &SAssetView::ExecuteDropCopyFolder),
+		DragDropHandler::FExecuteCopyOrMoveFolders::CreateSP(this, &SAssetView::ExecuteDropMoveFolder)
 		);
 }
 
diff --git a/Engine/Source/Editor/ContentBrowser/Private/SPathView.cpp b/Engine/Source/Editor/ContentBrowser/Private/SPathView.cpp
--- a/Engine/Source/Editor/ContentBrowser/Private/SPathView.cpp
+++ b/Engine/Source/Editor/ContentBrowser/Private/SPathView.cpp
@@ -3,6 +3,7 @@
 #include "ContentBrowserPCH.h"
 
 #include "DragAndDrop/AssetPathDragDropOp.h"
+#include "DragDropHandler.h"
 
 #include "PathViewTypes.h"
 #include "ObjectTools.h"
@@ -1251,79 +1252,25 @@ void SPathView::RemoveFolderItem(const TSharedPtr< FTreeItem >& TreeItem)
 
 void SPathView::TreeAssetsDropped(const TArray<FAssetData>& AssetList, const TSharedPtr<FTreeItem>& TreeItem)
 {
-	// Do not display the menu if any of the assets are classes as they cannot be moved or copied
-	for( int32 AssetIndex = 0; AssetIndex < AssetList.Num(); AssetIndex++ )
-	{
-		const FAssetData& Asset = AssetList[AssetIndex];
-		if ( Asset.AssetClass == "Class" )
-		{
-			const FText MessageText = LOCTEXT("AssetTreeDropClassError", "The selection contains one or more 'Class' type assets, these cannot be moved or copied.");
-			FMessageDialog::Open(EAppMsgType::Ok, MessageText);
-			return;
-		}
-	}
-
-	FMenuBuilder MenuBuilder(/*bInShouldCloseWindowAfterMenuSelection=*/true, NULL);
-	const FText MoveCopyHeaderString = FText::Format( LOCTEXT("AssetTreeDropMenuHeading", "Move/Copy to {0}"), TreeItem->DisplayName );
-	MenuBuilder.BeginSection("PathAssetMoveCopy", MoveCopyHeaderString);
-	{
-		MenuBuilder.AddMenuEntry(
-			LOCTEXT("DragDropCopy", "Copy Here"),
-			LOCTEXT("DragDropCopyTooltip", "Creates a copy of all dragged files in this folder."),
-			FSlateIcon(),
-			FUIAction(
-			FExecuteAction::CreateSP( this, &SPathView::ExecuteTreeDropCopy, AssetList, TreeItem ),
-			FCanExecuteAction()
-			)
-			);
-
-		MenuBuilder.AddMenuEntry(
-			LOCTEXT("DragDropMove", "Move Here"),
-			LOCTEXT("DragDropMoveTooltip", "Moves all dragged files to this folder."),
-			FSlateIcon(),
-			FUIAction(
-			FExecuteAction::CreateSP( this, &SPathView::ExecuteTreeDropMove, AssetList, TreeItem ),
-			FCanExecuteAction()
-			)
-			);
-	}
-	MenuBuilder.EndSection();
-
-	TWeakPtr< SWindow > ContextMenuWindow = FSlateApplication::Get().PushMenu(
-		SharedThis( this ),
-		MenuBuilder.MakeWidget(),
-		FSlateApplication::Get().GetCursorPos(),
-		FPopupTransitionEffect( FPopupTransitionEffect::ContextMenu )
+	DragDropHandler::HandleAssetsDroppedOnAssetFolder(
+		SharedThis(this), 
+		AssetList, 
+		TreeItem->FolderPath, 
+		TreeItem->DisplayName, 
+		DragDropHandler::FExecuteCopyOrMoveAssets::CreateSP(this, &SPathView::ExecuteTreeDropCopy),
+		DragDropHandler::FExecuteCopyOrMoveAssets::CreateSP(this, &SPathView::ExecuteTreeDropMove)
 		);
 }
 
 void SPathView::TreeFoldersDropped(const TArray<FString>& PathNames, const TSharedPtr<FTreeItem>& TreeItem)
 {
-	FMenuBuilder MenuBuilder(/*bInShouldCloseWindowAfterMenuSelection=*/true, NULL);
-	const FText MoveCopyHeaderString = FText::Format( LOCTEXT("AssetTreeDropMenuHeading", "Move/Copy to {0}"), TreeItem->DisplayName );
-	MenuBuilder.BeginSection("PathFolderMoveCopy", MoveCopyHeaderString);
-	{
-		MenuBuilder.AddMenuEntry(
-			LOCTEXT("DragDropCopyFolder", "Copy Folder Here"),
-			LOCTEXT("DragDropCopyFolderTooltip", "Creates a copy of all assets in the dragged folders to this folder, preserving folder structure."),
-			FSlateIcon(),
-			FUIAction( FExecuteAction::CreateSP( this, &SPathView::ExecuteTreeDropCopyFolder, PathNames, TreeItem ) )
-			);
-
-		MenuBuilder.AddMenuEntry(
-			LOCTEXT("DragDropMoveFolder", "Move Folder Here"),
-			LOCTEXT("DragDropMoveFolderTooltip", "Moves all assets in the dragged folders to this folder, preserving folder structure."),
-			FSlateIcon(),
-			FUIAction( FExecuteAction::CreateSP( this, &SPathView::ExecuteTreeDropMoveFolder, PathNames, TreeItem ) )
-			);
-	}
-	MenuBuilder.EndSection();
-
-	TWeakPtr< SWindow > ContextMenuWindow = FSlateApplication::Get().PushMenu(
-		SharedThis( this ),
-		MenuBuilder.MakeWidget(),
-		FSlateApplication::Get().GetCursorPos(),
-		FPopupTransitionEffect( FPopupTransitionEffect::ContextMenu )
+	DragDropHandler::HandleFoldersDroppedOnAssetFolder(
+		SharedThis(this), 
+		PathNames, 
+		TreeItem->FolderPath, 
+		TreeItem->DisplayName, 
+		DragDropHandler::FExecuteCopyOrMoveFolders::CreateSP(this, &SPathView::ExecuteTreeDropCopyFolder),
+		DragDropHandler::FExecuteCopyOrMoveFolders::CreateSP(this, &SPathView::ExecuteTreeDropMoveFolder)
 		);
 }
 
@@ -1343,62 +1290,70 @@ bool SPathView::IsTreeItemSelected(TSharedPtr<FTreeItem> TreeItem) const
 	return TreeViewPtr->IsItemSelected(TreeItem);
 }
 
-void SPathView::ExecuteTreeDropCopy(TArray<FAssetData> AssetList, TSharedPtr<FTreeItem> TreeItem)
+void SPathView::ExecuteTreeDropCopy(TArray<FAssetData> AssetList, FString DestinationPath)
 {
 	TArray<UObject*> DroppedObjects;
 	ContentBrowserUtils::GetObjectsInAssetData(AssetList, DroppedObjects);
 
-	ContentBrowserUtils::CopyAssets(DroppedObjects, TreeItem->FolderPath);
+	ContentBrowserUtils::CopyAssets(DroppedObjects, DestinationPath);
 }
 
-void SPathView::ExecuteTreeDropMove(TArray<FAssetData> AssetList, TSharedPtr<FTreeItem> TreeItem)
+void SPathView::ExecuteTreeDropMove(TArray<FAssetData> AssetList, FString DestinationPath)
 {
 	TArray<UObject*> DroppedObjects;
 	ContentBrowserUtils::GetObjectsInAssetData(AssetList, DroppedObjects);
 
-	ContentBrowserUtils::MoveAssets(DroppedObjects, TreeItem->FolderPath);
+	ContentBrowserUtils::MoveAssets(DroppedObjects, DestinationPath);
 }
 
-void SPathView::ExecuteTreeDropCopyFolder(TArray<FString> PathNames, TSharedPtr<FTreeItem> TreeItem)
+void SPathView::ExecuteTreeDropCopyFolder(TArray<FString> PathNames, FString DestinationPath)
 {
-	ContentBrowserUtils::CopyFolders(PathNames, TreeItem->FolderPath);
-
-	TreeViewPtr->SetItemExpansion(TreeItem, true);
+	ContentBrowserUtils::CopyFolders(PathNames, DestinationPath);
 
-	// Select all the new folders
-	TreeViewPtr->ClearSelection();
-	for ( auto PathIt = PathNames.CreateConstIterator(); PathIt; ++PathIt )
+	TSharedPtr<FTreeItem> RootItem = FindItemRecursive(DestinationPath);
+	if (RootItem.IsValid())
 	{
-		const FString SubFolderName = FPackageName::GetLongPackageAssetName(*PathIt);
-		const FString NewPath = TreeItem->FolderPath + TEXT("/") + SubFolderName;
-		
-		TSharedPtr<FTreeItem> Item = FindItemRecursive(NewPath);
-		if ( Item.IsValid() )
+		TreeViewPtr->SetItemExpansion(RootItem, true);
+
+		// Select all the new folders
+		TreeViewPtr->ClearSelection();
+		for ( auto PathIt = PathNames.CreateConstIterator(); PathIt; ++PathIt )
 		{
-			TreeViewPtr->SetItemSelection(Item, true);
-			TreeViewPtr->RequestScrollIntoView(Item);
+			const FString SubFolderName = FPackageName::GetLongPackageAssetName(*PathIt);
+			const FString NewPath = DestinationPath + TEXT("/") + SubFolderName;
+		
+			TSharedPtr<FTreeItem> Item = FindItemRecursive(NewPath);
+			if ( Item.IsValid() )
+			{
+				TreeViewPtr->SetItemSelection(Item, true);
+				TreeViewPtr->RequestScrollIntoView(Item);
+			}
 		}
 	}
 }
 
-void SPathView::ExecuteTreeDropMoveFolder(TArray<FString> PathNames, TSharedPtr<FTreeItem> TreeItem)
+void SPathView::ExecuteTreeDropMoveFolder(TArray<FString> PathNames, FString DestinationPath)
 {
-	ContentBrowserUtils::MoveFolders(PathNames, TreeItem->FolderPath);
+	ContentBrowserUtils::MoveFolders(PathNames, DestinationPath);
 
-	TreeViewPtr->SetItemExpansion(TreeItem, true);
-	
-	// Select all the new folders
-	TreeViewPtr->ClearSelection();
-	for ( auto PathIt = PathNames.CreateConstIterator(); PathIt; ++PathIt )
+	TSharedPtr<FTreeItem> RootItem = FindItemRecursive(DestinationPath);
+	if (RootItem.IsValid())
 	{
-		const FString SubFolderName = FPackageName::GetLongPackageAssetName(*PathIt);
-		const FString NewPath = TreeItem->FolderPath + TEXT("/") + SubFolderName;
-
-		TSharedPtr<FTreeItem> Item = FindItemRecursive(NewPath);
-		if ( Item.IsValid() )
+		TreeViewPtr->SetItemExpansion(RootItem, true);
+	
+		// Select all the new folders
+		TreeViewPtr->ClearSelection();
+		for ( auto PathIt = PathNames.CreateConstIterator(); PathIt; ++PathIt )
 		{
-			TreeViewPtr->SetItemSelection(Item, true);
-			TreeViewPtr->RequestScrollIntoView(Item);
+			const FString SubFolderName = FPackageName::GetLongPackageAssetName(*PathIt);
+			const FString NewPath = DestinationPath + TEXT("/") + SubFolderName;
+
+			TSharedPtr<FTreeItem> Item = FindItemRecursive(NewPath);
+			if ( Item.IsValid() )
+			{
+				TreeViewPtr->SetItemSelection(Item, true);
+				TreeViewPtr->RequestScrollIntoView(Item);
+			}
 		}
 	}
 }
diff --git a/Engine/Source/Editor/ContentBrowser/Private/SPathView.h b/Engine/Source/Editor/ContentBrowser/Private/SPathView.h
--- a/Engine/Source/Editor/ContentBrowser/Private/SPathView.h
+++ b/Engine/Source/Editor/ContentBrowser/Private/SPathView.h
@@ -198,16 +198,16 @@ private:
 	void GetDroppedObjects(const TArray<FAssetData>& AssetList, TArray<UObject*>& OutDroppedObjects);
 
 	/** Handler for the user selecting to copy assets to the specified folder */
-	void ExecuteTreeDropCopy(TArray<FAssetData> AssetList, TSharedPtr<FTreeItem> TreeItem);
+	void ExecuteTreeDropCopy(TArray<FAssetData> AssetList, FString DestinationPath);
 
 	/** Handler for the user selecting to move assets to the specified folder */
-	void ExecuteTreeDropMove(TArray<FAssetData> AssetList, TSharedPtr<FTreeItem> TreeItem);
+	void ExecuteTreeDropMove(TArray<FAssetData> AssetList, FString DestinationPath);
 
 	/** Handler for the user selecting to copy folders to the specified folder */
-	void ExecuteTreeDropCopyFolder(TArray<FString> PathNames, TSharedPtr<FTreeItem> TreeItem);
+	void ExecuteTreeDropCopyFolder(TArray<FString> PathNames, FString DestinationPath);
 
 	/** Handler for the user selecting to move folders to the specified folder */
-	void ExecuteTreeDropMoveFolder(TArray<FString> PathNames, TSharedPtr<FTreeItem> TreeItem);
+	void ExecuteTreeDropMoveFolder(TArray<FString> PathNames, FString DestinationPath);
 
 	/** Handles updating the content browser when an asset path is added to the asset registry */
 	void OnAssetRegistryPathAdded(const FString& Path);
diff --git a/Engine/Source/Editor/ContentBrowser/Private/SourcesViewWidgets.cpp b/Engine/Source/Editor/ContentBrowser/Private/SourcesViewWidgets.cpp
--- a/Engine/Source/Editor/ContentBrowser/Private/SourcesViewWidgets.cpp
+++ b/Engine/Source/Editor/ContentBrowser/Private/SourcesViewWidgets.cpp
@@ -7,6 +7,7 @@
 
 #include "DragAndDrop/AssetDragDropOp.h"
 #include "DragAndDrop/AssetPathDragDropOp.h"
+#include "DragDropHandler.h"
 #include "ContentBrowserUtils.h"
 #include "CollectionViewUtils.h"
 #include "SInlineEditableTextBlock.h"
@@ -96,112 +97,47 @@ SAssetTreeItem::~SAssetTreeItem()
 	}
 }
 
-void SAssetTreeItem::OnDragEnter( const FGeometry& MyGeometry, const FDragDropEvent& DragDropEvent )
+bool SAssetTreeItem::ValidateDragDrop( const FGeometry& MyGeometry, const FDragDropEvent& DragDropEvent ) const
 {
-	if ( IsValidAssetPath() )
-	{
-		TSharedPtr<FDragDropOperation> Operation = DragDropEvent.GetOperation();
-		if (!Operation.IsValid())
-		{
-			return;
-		}
-
-		if (Operation->IsOfType<FAssetDragDropOp>())
-		{
-			TSharedPtr<FAssetDragDropOp> DragDropOp = StaticCastSharedPtr<FAssetDragDropOp>(Operation);
-			bool bCanDrop = DragDropOp->AssetData.Num() > 0;
-
-			if (bCanDrop)
-			{
-				DragDropOp->SetToolTip( LOCTEXT( "OnDragAssetsOverFolder", "Move or Copy Asset(s)" ), FEditorStyle::GetBrush(TEXT("Graph.ConnectorFeedback.OK")) );
-				bDraggedOver = true;
-			}
-		}
-		else if (Operation->IsOfType<FAssetPathDragDropOp>())
-		{
-			TSharedPtr<FAssetPathDragDropOp> DragDropOp = StaticCastSharedPtr<FAssetPathDragDropOp>(Operation);
-			bool bCanDrop = DragDropOp->PathNames.Num() > 0;
-
-			if ( DragDropOp->PathNames.Num() == 1 && DragDropOp->PathNames[0] == TreeItem.Pin()->FolderPath )
-			{
-				// You can't drop a single folder onto itself
-				bCanDrop = false;
-			}
-
-			if (bCanDrop)
-			{
-				bDraggedOver = true;
-			}
-		}
-		else if (Operation->IsOfType<FExternalDragOperation>())
-		{
-			TSharedPtr<FExternalDragOperation> DragDropOp = StaticCastSharedPtr<FExternalDragOperation>(Operation);
-			bool bCanDrop = DragDropOp->HasFiles();
-
-			if (bCanDrop)
-			{
-				bDraggedOver = true;
-			}
-		}
-	}
+	TSharedPtr<FTreeItem> TreeItemPinned = TreeItem.Pin();
+	return TreeItemPinned.IsValid() && DragDropHandler::ValidateDragDropOnAssetFolder(MyGeometry, DragDropEvent, TreeItemPinned->FolderPath);
 }
 
-void SAssetTreeItem::OnDragLeave( const FDragDropEvent& DragDropEvent )
+void SAssetTreeItem::OnDragEnter( const FGeometry& MyGeometry, const FDragDropEvent& DragDropEvent )
 {
-	TSharedPtr<FAssetDragDropOp> DragDropOp = DragDropEvent.GetOperationAs<FAssetDragDropOp>();
-	if ( DragDropOp.IsValid() )
+	bDraggedOver = false;
+
+	if (ValidateDragDrop(MyGeometry, DragDropEvent))
 	{
-		DragDropOp->ResetToDefaultToolTip();
+		bDraggedOver = true;
 	}
-	bDraggedOver = false;
 }
 
-FReply SAssetTreeItem::OnDragOver( const FGeometry& MyGeometry, const FDragDropEvent& DragDropEvent )
+void SAssetTreeItem::OnDragLeave( const FDragDropEvent& DragDropEvent )
 {
-	if ( IsValidAssetPath() )
+	TSharedPtr<FDragDropOperation> Operation = DragDropEvent.GetOperation();
+	if (Operation.IsValid())
 	{
-		TSharedPtr<FDragDropOperation> Operation = DragDropEvent.GetOperation();
-		if (!Operation.IsValid())
-		{
-			return FReply::Unhandled();
-		}
+		Operation->SetCursorOverride(TOptional<EMouseCursor::Type>());
 
 		if (Operation->IsOfType<FAssetDragDropOp>())
 		{
-			TSharedPtr<FAssetDragDropOp> DragDropOp = StaticCastSharedPtr<FAssetDragDropOp>( DragDropEvent.GetOperation() );
-			bool bCanDrop = DragDropOp->AssetData.Num() > 0;
-
-			if (bCanDrop)
-			{
-				return FReply::Handled();
-			}
+			TSharedPtr<FAssetDragDropOp> DragDropOp = StaticCastSharedPtr<FAssetDragDropOp>(Operation);
+			DragDropOp->ResetToDefaultToolTip();
 		}
-		else if (Operation->IsOfType<FAssetPathDragDropOp>())
-		{
-			TSharedPtr<FAssetPathDragDropOp> DragDropOp = StaticCastSharedPtr<FAssetPathDragDropOp>( DragDropEvent.GetOperation() );
-			bool bCanDrop = DragDropOp->PathNames.Num() > 0;
+	}
 
-			if ( DragDropOp->PathNames.Num() == 1 && DragDropOp->PathNames[0] == TreeItem.Pin()->FolderPath )
-			{
-				// You can't drop a single folder onto itself
-				bCanDrop = false;
-			}
+	bDraggedOver = false;
+}
 
-			if (bCanDrop)
-			{
-				return FReply::Handled();
-			}
-		}
-		else if (Operation->IsOfType<FExternalDragOperation>())
-		{
-			TSharedPtr<FExternalDragOperation> DragDropOp = StaticCastSharedPtr<FExternalDragOperation>( DragDropEvent.GetOperation() );
-			bool bCanDrop = DragDropOp->HasFiles();
+FReply SAssetTreeItem::OnDragOver( const FGeometry& MyGeometry, const FDragDropEvent& DragDropEvent )
+{
+	bDraggedOver = false;
 
-			if (bCanDrop)
-			{
-				return FReply::Handled();
-			}
-		}
+	if (ValidateDragDrop(MyGeometry, DragDropEvent))
+	{
+		bDraggedOver = true;
+		return FReply::Handled();
 	}
 
 	return FReply::Unhandled();
@@ -211,7 +147,7 @@ FReply SAssetTreeItem::OnDrop( const FGeometry& MyGeometry, const FDragDropEvent
 {
 	bDraggedOver = false;
 
-	if ( IsValidAssetPath() )
+	if (ValidateDragDrop(MyGeometry, DragDropEvent))
 	{
 		TSharedPtr<FDragDropOperation> Operation = DragDropEvent.GetOperation();
 		if (!Operation.IsValid())
@@ -222,41 +158,19 @@ FReply SAssetTreeItem::OnDrop( const FGeometry& MyGeometry, const FDragDropEvent
 		if (Operation->IsOfType<FAssetDragDropOp>())
 		{
 			TSharedPtr<FAssetDragDropOp> DragDropOp = StaticCastSharedPtr<FAssetDragDropOp>( DragDropEvent.GetOperation() );
-
 			OnAssetsDragDropped.ExecuteIfBound(DragDropOp->AssetData, TreeItem.Pin());
-
 			return FReply::Handled();
 		}
 		else if (Operation->IsOfType<FAssetPathDragDropOp>())
 		{
 			TSharedPtr<FAssetPathDragDropOp> DragDropOp = StaticCastSharedPtr<FAssetPathDragDropOp>( DragDropEvent.GetOperation() );
-
-			bool bCanDrop = DragDropOp->PathNames.Num() > 0;
-
-			if ( DragDropOp->PathNames.Num() == 1 && DragDropOp->PathNames[0] == TreeItem.Pin()->FolderPath )
-			{
-				// You can't drop a single folder onto itself
-				bCanDrop = false;
-			}
-
-			if ( bCanDrop )
-			{
-				OnPathsDragDropped.ExecuteIfBound(DragDropOp->PathNames, TreeItem.Pin());
-			}
-
+			OnPathsDragDropped.ExecuteIfBound(DragDropOp->PathNames, TreeItem.Pin());
 			return FReply::Handled();
 		}
 		else if (Operation->IsOfType<FExternalDragOperation>())
 		{
 			TSharedPtr<FExternalDragOperation> DragDropOp = StaticCastSharedPtr<FExternalDragOperation>( DragDropEvent.GetOperation() );
-
-			bool bCanDrop = DragDropOp->HasFiles();
-
-			if ( bCanDrop )
-			{
-				OnFilesDragDropped.ExecuteIfBound(DragDropOp->GetFiles(), TreeItem.Pin());
-			}
-
+			OnFilesDragDropped.ExecuteIfBound(DragDropOp->GetFiles(), TreeItem.Pin());
 			return FReply::Handled();
 		}
 	}
diff --git a/Engine/Source/Editor/ContentBrowser/Private/SourcesViewWidgets.h b/Engine/Source/Editor/ContentBrowser/Private/SourcesViewWidgets.h
--- a/Engine/Source/Editor/ContentBrowser/Private/SourcesViewWidgets.h
+++ b/Engine/Source/Editor/ContentBrowser/Private/SourcesViewWidgets.h
@@ -57,6 +57,9 @@ public:
 	virtual void Tick( const FGeometry& AllottedGeometry, const double InCurrentTime, const float InDeltaTime ) override;
 
 private:
+	/** Used by OnDragEnter, OnDragOver, and OnDrop to check and update the validity of the drag operation */
+	bool ValidateDragDrop( const FGeometry& MyGeometry, const FDragDropEvent& DragDropEvent ) const;
+
 	/** Handles verifying name changes */
 	bool VerifyNameChanged(const FText& InName, FText& OutError) const;
 
diff --git a/Engine/Source/Editor/UnrealEd/Private/DragAndDrop/AssetDragDropOp.cpp b/Engine/Source/Editor/UnrealEd/Private/DragAndDrop/AssetDragDropOp.cpp
--- a/Engine/Source/Editor/UnrealEd/Private/DragAndDrop/AssetDragDropOp.cpp
+++ b/Engine/Source/Editor/UnrealEd/Private/DragAndDrop/AssetDragDropOp.cpp
@@ -120,6 +120,7 @@ TSharedPtr<SWidget> FAssetDragDropOp::GetDefaultDecorator() const
 					+ SHorizontalBox::Slot()
 					.AutoWidth()
 					.Padding(3.0f)
+					.VAlign(VAlign_Center)
 					[
 						SNew(SImage) 
 						.Image(this, &FAssetDragDropOp::GetIcon)
diff --git a/Engine/Source/Editor/UnrealEd/Public/DragAndDrop/AssetPathDragDropOp.h b/Engine/Source/Editor/UnrealEd/Public/DragAndDrop/AssetPathDragDropOp.h
--- a/Engine/Source/Editor/UnrealEd/Public/DragAndDrop/AssetPathDragDropOp.h
+++ b/Engine/Source/Editor/UnrealEd/Public/DragAndDrop/AssetPathDragDropOp.h
@@ -2,10 +2,12 @@
 
 #pragma once
 
-class FAssetPathDragDropOp : public FDragDropOperation
+#include "DecoratedDragDropOp.h"
+
+class FAssetPathDragDropOp : public FDecoratedDragDropOp
 {
 public:
-	DRAG_DROP_OPERATOR_TYPE(FAssetPathDragDropOp, FDragDropOperation)
+	DRAG_DROP_OPERATOR_TYPE(FAssetPathDragDropOp, FDecoratedDragDropOp)
 
 	/** Data for the asset this item represents */
 	TArray<FString> PathNames;
@@ -22,36 +24,66 @@ public:
 	}
 	
 public:
-	virtual TSharedPtr<SWidget> GetDefaultDecorator() const override
+	FText GetDecoratorText() const
 	{
-		FString Text = PathNames.Num() > 0 ? PathNames[0] : TEXT("");
-
-		if ( PathNames.Num() > 1 )
+		if ( CurrentHoverText.IsEmpty() )
 		{
-			Text += TEXT(" ");
-			Text += FString::Printf(*NSLOCTEXT("ContentBrowser", "FolderDescription", "and %d other(s)").ToString(), PathNames.Num() - 1);
+			FString Text = PathNames.Num() > 0 ? PathNames[0] : TEXT("");
+
+			if ( PathNames.Num() > 1 )
+			{
+				Text += TEXT(" ");
+				Text += FString::Printf(*NSLOCTEXT("ContentBrowser", "FolderDescription", "and %d other(s)").ToString(), PathNames.Num() - 1);
+			}
+
+			return FText::FromString(Text);
 		}
+		
+		return CurrentHoverText;
+	}
 
-		return
+	virtual TSharedPtr<SWidget> GetDefaultDecorator() const override
+	{
+		return 
 			SNew(SBorder)
-			.BorderImage( FEditorStyle::GetBrush("Menu.Background") )
+			.BorderImage(FEditorStyle::GetBrush("ContentBrowser.AssetDragDropTooltipBackground"))
 			.Content()
 			[
 				SNew(SHorizontalBox)
-				+SHorizontalBox::Slot()
+
+				// Left slot is folder icon
+				+ SHorizontalBox::Slot()
 				.AutoWidth()
-				.HAlign(HAlign_Center)
 				.VAlign(VAlign_Center)
 				[
 					SNew(SImage)
-					.Image( FEditorStyle::GetBrush("ContentBrowser.AssetTreeFolderClosed") )
+					.Image(FEditorStyle::GetBrush("ContentBrowser.AssetTreeFolderClosed"))
 				]
-				+SHorizontalBox::Slot()
+
+				// Right slot is for tooltip
+				+ SHorizontalBox::Slot()
 				.AutoWidth()
 				.VAlign(VAlign_Center)
 				[
-					SNew(STextBlock) 
-					.Text( FText::FromString(Text) )
+					SNew(SHorizontalBox)
+
+					+ SHorizontalBox::Slot()
+					.AutoWidth()
+					.Padding(3.0f)
+					.VAlign(VAlign_Center)
+					[
+						SNew(SImage) 
+						.Image(this, &FAssetPathDragDropOp::GetIcon)
+					]
+
+					+SHorizontalBox::Slot()
+					.AutoWidth()
+					.Padding(0,0,3,0)
+					.VAlign(VAlign_Center)
+					[
+						SNew(STextBlock) 
+						.Text(this, &FAssetPathDragDropOp::GetDecoratorText)
+					]
 				]
 			];
 	}
diff --git a/Engine/Source/Editor/UnrealEd/Public/DragAndDrop/DecoratedDragDropOp.h b/Engine/Source/Editor/UnrealEd/Public/DragAndDrop/DecoratedDragDropOp.h
--- a/Engine/Source/Editor/UnrealEd/Public/DragAndDrop/DecoratedDragDropOp.h
+++ b/Engine/Source/Editor/UnrealEd/Public/DragAndDrop/DecoratedDragDropOp.h
@@ -48,6 +48,7 @@ public:
 				+ SHorizontalBox::Slot()
 				.AutoWidth()
 				.Padding( 0.0f, 0.0f, 3.0f, 0.0f )
+				.VAlign( VAlign_Center )
 				[
 					SNew( SImage )
 					.Image( this, &FDecoratedDragDropOp::GetIcon )
