[INTEGRATE] Change 2433543 by [REDACTED_PERSON][REDACTED_HANDLE] on 2015/02/05 09:35:01

	Stopped C++ classes sometimes being shown as dirty in the Content Browser

	UE-9111 - All C++ Classes in Engine module appear "unsaved" in Content Browser

diff --git a/Engine/Source/Editor/ContentBrowser/Private/AssetViewWidgets.cpp b/Engine/Source/Editor/ContentBrowser/Private/AssetViewWidgets.cpp
--- a/Engine/Source/Editor/ContentBrowser/Private/AssetViewWidgets.cpp
+++ b/Engine/Source/Editor/ContentBrowser/Private/AssetViewWidgets.cpp
@@ -347,7 +347,7 @@ const FSlateBrush* SAssetViewItem::GetSCCStateImage() const
 
 void SAssetViewItem::HandleSourceControlStateChanged()
 {
-	if ( ISourceControlModule::Get().IsEnabled() && AssetItem.IsValid() && (AssetItem->GetType() == EAssetItemType::Normal) && !AssetItem->IsTemporaryItem() && !CachedPackageFileName.IsEmpty() )
+	if ( ISourceControlModule::Get().IsEnabled() && AssetItem.IsValid() && (AssetItem->GetType() == EAssetItemType::Normal) && !AssetItem->IsTemporaryItem() && !FPackageName::IsScriptPackage(CachedPackageName) )
 	{
 		FSourceControlStatePtr SourceControlState = ISourceControlModule::Get().GetProvider().GetState(CachedPackageFileName, EStateCacheUsage::Use);
 		if(SourceControlState.IsValid())
@@ -720,9 +720,14 @@ void SAssetViewItem::AddToToolTipInfoBox(const TSharedRef<SVerticalBox>& InfoBox
 void SAssetViewItem::UpdatePackageDirtyState()
 {
 	bool bNewIsDirty = false;
-	if ( AssetPackage.IsValid() )
+
+	// Only update the dirty state for non-temporary asset items that aren't a built in script
+	if ( AssetItem.IsValid() && !AssetItem->IsTemporaryItem() && AssetItem->GetType() != EAssetItemType::Folder && !FPackageName::IsScriptPackage(CachedPackageName) )
 	{
-		bNewIsDirty = AssetPackage->IsDirty();
+		if ( AssetPackage.IsValid() )
+		{
+			bNewIsDirty = AssetPackage->IsDirty();
+		}
 	}
 
 	if ( bNewIsDirty != bPackageDirty )
@@ -743,17 +748,11 @@ void SAssetViewItem::UpdateSourceControlState(float InDeltaTime)
 
 	if ( !bSourceControlStateRequested && SourceControlStateDelay > 1.0f && AssetItem.IsValid() )
 	{
-		if ( AssetItem.IsValid() && AssetItem->GetType() != EAssetItemType::Folder )
+		// Only update the SCC state for non-temporary asset items that aren't a built in script
+		if ( AssetItem.IsValid() && !AssetItem->IsTemporaryItem() && AssetItem->GetType() != EAssetItemType::Folder && !FPackageName::IsScriptPackage(CachedPackageName) )
 		{
-			if ( !AssetItem->IsTemporaryItem() )
-			{
-				// dont query status for built-in types
-				if(!FPackageName::IsScriptPackage(CachedPackageName))
-				{
-				// Request the most recent SCC state for this asset
-				ISourceControlModule::Get().QueueStatusUpdate(CachedPackageFileName);
-			}
-		}
+			// Request the most recent SCC state for this asset
+			ISourceControlModule::Get().QueueStatusUpdate(CachedPackageFileName);
 		}
 
 		bSourceControlStateRequested = true;
