[INTEGRATE] Change 2435465 by [REDACTED_PERSON][REDACTED_HANDLE] on 2015/02/06 10:42:44

	Editor - Adjusting the minimum allowed thumbnail size in the content browser's asset view.

diff --git a/Engine/Source/Editor/ContentBrowser/Private/SAssetView.cpp b/Engine/Source/Editor/ContentBrowser/Private/SAssetView.cpp
--- a/Engine/Source/Editor/ContentBrowser/Private/SAssetView.cpp
+++ b/Engine/Source/Editor/ContentBrowser/Private/SAssetView.cpp
@@ -134,7 +134,7 @@ void SAssetView::Construct( const FArguments& InArgs )
 		ThumbnailScaleSliderValue = FMath::Clamp<float>(ThumbnailScaleSliderValue.Get(), 0.0f, 1.0f);
 	}
 
-	MinThumbnailScale = 0.4f * ThumbnailScaleRangeScalar;
+	MinThumbnailScale = 0.2f * ThumbnailScaleRangeScalar;
 	MaxThumbnailScale = 2.0f * ThumbnailScaleRangeScalar;
 
 	bCanShowClasses = InArgs._CanShowClasses;
