Additional file which got lost as part of CL 2446774.

diff --git a/Engine/Source/Editor/ContentBrowser/Private/PathContextMenu.cpp b/Engine/Source/Editor/ContentBrowser/Private/PathContextMenu.cpp
--- a/Engine/Source/Editor/ContentBrowser/Private/PathContextMenu.cpp
+++ b/Engine/Source/Editor/ContentBrowser/Private/PathContextMenu.cpp
@@ -827,7 +827,14 @@ void FPathContextMenu::ExecuteSCCCheckIn()
 	const bool bShouldProceed = ( UserResponse == FEditorFileUtils::EPromptReturnCode::PR_Success || UserResponse == FEditorFileUtils::EPromptReturnCode::PR_Declined );
 	if ( bShouldProceed )
 	{
-		FSourceControlWindows::PromptForCheckin(PackageNames);
+		TArray<FString> PendingDeletePaths;
+		for (const auto& Path : SelectedPaths)
+		{
+			PendingDeletePaths.Add(FPaths::ConvertRelativePathToFull(FPackageName::LongPackageNameToFilename(Path + TEXT("/"))));
+		}
+
+		const bool bUseSourceControlStateCache = false;
+		FSourceControlWindows::PromptForCheckin(bUseSourceControlStateCache, PackageNames, PendingDeletePaths);
 	}
 	else
 	{
