[INTEGRATE] Change 2440213 by [REDACTED_PERSON][REDACTED_HANDLE] on 2015/02/10 14:49:45

	Tweaked the starting size of thumbnails in the content browser

diff --git a/Engine/Source/Editor/ContentBrowser/Private/SContentBrowser.cpp b/Engine/Source/Editor/ContentBrowser/Private/SContentBrowser.cpp
--- a/Engine/Source/Editor/ContentBrowser/Private/SContentBrowser.cpp
+++ b/Engine/Source/Editor/ContentBrowser/Private/SContentBrowser.cpp
@@ -545,7 +545,7 @@ void SContentBrowser::Construct( const FArguments& InArgs, const FName& InInstan
 					.Padding( 0 )
 					[
 						SAssignNew(AssetViewPtr, SAssetView)
-						.ThumbnailScale( 0.075f )
+						.ThumbnailScale( 0.18 )
 						.OnPathSelected(this, &SContentBrowser::FolderEntered)
 						.OnAssetSelected(this, &SContentBrowser::OnAssetSelectionChanged)
 						.OnAssetsActivated(this, &SContentBrowser::OnAssetsActivated)
