[INTEGRATE] Change 2442108 by [REDACTED_PERSON][REDACTED_HANDLE] on 2015/02/11 14:40:03

	UE-9651 - Editor - Change the setting name for the thumbnail size so that everyone gets the newer, better default value.

diff --git a/Engine/Source/Editor/ContentBrowser/Private/SAssetView.cpp b/Engine/Source/Editor/ContentBrowser/Private/SAssetView.cpp
--- a/Engine/Source/Editor/ContentBrowser/Private/SAssetView.cpp
+++ b/Engine/Source/Editor/ContentBrowser/Private/SAssetView.cpp
@@ -666,16 +666,26 @@ void SAssetView::RequestAddNewAssetsNextFrame()
 	LastProcessAddsTime = FPlatformTime::Seconds() - TimeBetweenAddingNewAssets;
 }
 
+FString SAssetView::GetThumbnailScaleSettingPath(const FString& SettingsString) const
+{
+	return SettingsString + TEXT(".ThumbnailSizeScale");
+}
+
+FString SAssetView::GetCurrentViewTypeSettingPath(const FString& SettingsString) const
+{
+	return SettingsString + TEXT(".CurrentViewType");
+}
+
 void SAssetView::SaveSettings(const FString& IniFilename, const FString& IniSection, const FString& SettingsString) const
 {
-	GConfig->SetFloat(*IniSection, *(SettingsString + TEXT(".ThumbnailScale")), ThumbnailScaleSliderValue.Get(), IniFilename);
-	GConfig->SetInt(*IniSection, *(SettingsString + TEXT(".CurrentViewType")), CurrentViewType, IniFilename);
+	GConfig->SetFloat(*IniSection, *GetThumbnailScaleSettingPath(SettingsString), ThumbnailScaleSliderValue.Get(), IniFilename);
+	GConfig->SetInt(*IniSection, *GetCurrentViewTypeSettingPath(SettingsString), CurrentViewType, IniFilename);
 }
 
 void SAssetView::LoadSettings(const FString& IniFilename, const FString& IniSection, const FString& SettingsString)
 {
 	float Scale = 0.f;
-	if ( GConfig->GetFloat(*IniSection, *(SettingsString + TEXT(".ThumbnailScale")), Scale, IniFilename) )
+	if ( GConfig->GetFloat(*IniSection, *GetThumbnailScaleSettingPath(SettingsString), Scale, IniFilename) )
 	{
 		// Clamp value to normal range and update state
 		Scale = FMath::Clamp<float>(Scale, 0.f, 1.f);
@@ -683,7 +693,7 @@ void SAssetView::LoadSettings(const FString& IniFilename, const FString& IniSect
 	}
 
 	int32 ViewType = EAssetViewType::Tile;
-	if ( GConfig->GetInt(*IniSection, *(SettingsString + TEXT(".CurrentViewType")), ViewType, IniFilename) )
+	if ( GConfig->GetInt(*IniSection, *GetCurrentViewTypeSettingPath(SettingsString), ViewType, IniFilename) )
 	{
 		// Clamp value to normal range and update state
 		if ( ViewType < 0 || ViewType >= EAssetViewType::MAX)
diff --git a/Engine/Source/Editor/ContentBrowser/Private/SAssetView.h b/Engine/Source/Editor/ContentBrowser/Private/SAssetView.h
--- a/Engine/Source/Editor/ContentBrowser/Private/SAssetView.h
+++ b/Engine/Source/Editor/ContentBrowser/Private/SAssetView.h
@@ -250,6 +250,12 @@ public:
 
 private:
 
+	/** @return the thumbnail scale setting path to use when looking up the setting in an ini. */
+	FString GetThumbnailScaleSettingPath(const FString& SettingsString) const;
+
+	/** @return the view type setting path to use when looking up the setting in an ini. */
+	FString GetCurrentViewTypeSettingPath(const FString& SettingsString) const;
+
 	/** Calculates a new filler scale used to adjust the thumbnails to fill empty space. */
 	void CalculateFillScale( const FGeometry& AllottedGeometry );
 
