[INTEGRATE] Change 2442516 by [REDACTED_PERSON][REDACTED_HANDLE] on 2015/02/11 17:52:51

	Disable the ability to duplicate the current world due to instability with open blueprints

diff --git a/Engine/Source/Editor/ContentBrowser/Private/AssetContextMenu.cpp b/Engine/Source/Editor/ContentBrowser/Private/AssetContextMenu.cpp
--- a/Engine/Source/Editor/ContentBrowser/Private/AssetContextMenu.cpp
+++ b/Engine/Source/Editor/ContentBrowser/Private/AssetContextMenu.cpp
@@ -48,6 +48,7 @@ FAssetContextMenu::FAssetContextMenu(const TWeakPtr<SAssetView>& InAssetView)
 	, bCanExecuteSCCHistory(false)
 	, bCanExecuteSCCRevert(false)
 	, bCanExecuteSCCSync(false)
+	, bCurrentWorldSelected(false)
 {
 	
 }
@@ -1643,7 +1644,7 @@ bool FAssetContextMenu::CanExecutePropertyMatrix() const
 
 bool FAssetContextMenu::CanExecuteDuplicate() const
 {
-	return bAtLeastOneNonRedirectorSelected && !bAtLeastOneClassSelected;
+	return bAtLeastOneNonRedirectorSelected && !bAtLeastOneClassSelected && !bCurrentWorldSelected;
 }
 
 bool FAssetContextMenu::CanExecuteRename() const
@@ -1846,7 +1847,9 @@ void FAssetContextMenu::CacheCanExecuteVars()
 	bCanExecuteSCCHistory = false;
 	bCanExecuteSCCRevert = false;
 	bCanExecuteSCCSync = false;
+	bCurrentWorldSelected = false;
 
+	static const FName WorldClass("World");
 	for (auto AssetIt = SelectedAssets.CreateConstIterator(); AssetIt; ++AssetIt)
 	{
 		const FAssetData& AssetData = *AssetIt;
@@ -1855,6 +1858,12 @@ void FAssetContextMenu::CacheCanExecuteVars()
 			continue;
 		}
 
+		// Hack: 4.7 current world is selected and cannot be duplicated due to blueprint trashed components issues.  
+		if( AssetData.AssetClass == WorldClass && AssetData.IsAssetLoaded() )
+		{
+			bCurrentWorldSelected = (GWorld == AssetData.GetAsset());
+		}
+
 		if ( !bAtLeastOneNonRedirectorSelected && AssetData.AssetClass != UObjectRedirector::StaticClass()->GetFName() )
 		{
 			bAtLeastOneNonRedirectorSelected = true;
diff --git a/Engine/Source/Editor/ContentBrowser/Private/AssetContextMenu.h b/Engine/Source/Editor/ContentBrowser/Private/AssetContextMenu.h
--- a/Engine/Source/Editor/ContentBrowser/Private/AssetContextMenu.h
+++ b/Engine/Source/Editor/ContentBrowser/Private/AssetContextMenu.h
@@ -331,7 +331,7 @@ private:
 	bool bCanExecuteSCCHistory;
 	bool bCanExecuteSCCRevert;
 	bool bCanExecuteSCCSync;
-
+	bool bCurrentWorldSelected;
 	/** */
 	int32 ChunkIDSelected;
 };
