[INTEGRATE] Change 2444699 by [REDACTED_PERSON][REDACTED_HANDLE] on 2015/02/13 09:53:29

	UE-9723 - Reworking the buttons on the content browser to be cleaner and one of them green.

diff --git a/Engine/Source/Editor/ContentBrowser/Private/SContentBrowser.cpp b/Engine/Source/Editor/ContentBrowser/Private/SContentBrowser.cpp
--- a/Engine/Source/Editor/ContentBrowser/Private/SContentBrowser.cpp
+++ b/Engine/Source/Editor/ContentBrowser/Private/SContentBrowser.cpp
@@ -106,39 +106,54 @@ void SContentBrowser::Construct( const FArguments& InArgs, const FName& InInstan
 						.AutoWidth()
 						.VAlign( VAlign_Center )
 						.HAlign( HAlign_Left )
-						.Padding(0,0,4,0)
 						[
 							SNew( SComboButton )
 							.ComboButtonStyle( FEditorStyle::Get(), "ToolbarComboButton" )
+							.ButtonStyle(FEditorStyle::Get(), "FlatButton.Success")
 							.ForegroundColor(FLinearColor::White)
-							.ContentPadding(0)
+							.ContentPadding(FMargin(6, 2))
 							.OnGetMenuContent_Lambda( [this]{ return MakeAddNewContextMenu( true, false ); } )
 							.ToolTipText( this, &SContentBrowser::GetAddNewToolTipText )
 							.IsEnabled( this, &SContentBrowser::IsAddNewEnabled )
 							.AddMetaData<FTagMetaData>(FTagMetaData(TEXT("ContentBrowserNewAsset")))
+							.HasDownArrow(false)
 							.ButtonContent()
 							[
 								SNew( SHorizontalBox )
 
 								// New Icon
 								+ SHorizontalBox::Slot()
+								.VAlign(VAlign_Center)
 								.AutoWidth()
-								.VAlign( VAlign_Center )
 								[
-									SNew( SImage )
-									.Image( FEditorStyle::GetBrush( "ContentBrowser.NewAsset" ) )
+									SNew(STextBlock)
+									.TextStyle(FEditorStyle::Get(), "ContentBrowser.TopBar.Font")
+									.Font(FEditorStyle::Get().GetFontStyle("FontAwesome.11"))
+									.Text(FString(TEXT("\xf15b")) /*fa-file*/)
 								]
 
 								// New Text
 								+ SHorizontalBox::Slot()
 								.AutoWidth()
 								.VAlign(VAlign_Center)
-								.Padding(0,0,2,0)
+								.Padding(4, 0, 0, 0)
 								[
 									SNew( STextBlock )
 									.TextStyle( FEditorStyle::Get(), "ContentBrowser.TopBar.Font" )
 									.Text( LOCTEXT( "NewButton", "Add New" ) )
 								]
+
+								// Down Arrow
+								+ SHorizontalBox::Slot()
+								.VAlign(VAlign_Center)
+								.AutoWidth()
+								.Padding(4, 0, 0, 0)
+								[
+									SNew(STextBlock)
+									.TextStyle(FEditorStyle::Get(), "ContentBrowser.TopBar.Font")
+									.Font(FEditorStyle::Get().GetFontStyle("FontAwesome.10"))
+									.Text(FString(TEXT("\xf0d7")) /*fa-caret-down*/)
+								]
 							]
 						]
 
@@ -147,32 +162,34 @@ void SContentBrowser::Construct( const FArguments& InArgs, const FName& InInstan
 						.AutoWidth()
 						.VAlign( VAlign_Center )
 						.HAlign( HAlign_Left )
-						.Padding(0,0,10,0)
+						.Padding(6, 0)
 						[
 							SNew( SButton )
-							.ButtonStyle( FEditorStyle::Get(), "ToggleButton" )
+							.ButtonStyle(FEditorStyle::Get(), "FlatButton")
 							.ToolTipText( this, &SContentBrowser::GetImportTooltipText )
 							.IsEnabled( this, &SContentBrowser::IsImportEnabled )
 							.OnClicked( this, &SContentBrowser::HandleImportClicked )
-							.ContentPadding( 0 )
+							.ContentPadding(FMargin(6, 2))
 							.AddMetaData<FTagMetaData>(FTagMetaData(TEXT("ContentBrowserImportAsset")))
 							[
 								SNew( SHorizontalBox )
 
 								// Import Icon
 								+ SHorizontalBox::Slot()
+								.VAlign(VAlign_Center)
 								.AutoWidth()
-								.VAlign( VAlign_Center )
 								[
-									SNew( SImage )
-									.Image( FEditorStyle::GetBrush( "ContentBrowser.ImportPackage" ) )
+									SNew(STextBlock)
+									.TextStyle(FEditorStyle::Get(), "ContentBrowser.TopBar.Font")
+									.Font(FEditorStyle::Get().GetFontStyle("FontAwesome.11"))
+									.Text(FString(TEXT("\xf019")) /*fa-download*/)
 								]
 
 								// Import Text
 								+ SHorizontalBox::Slot()
 								.AutoWidth()
 								.VAlign(VAlign_Center)
-								.Padding(0,0,2,0)
+								.Padding(4, 0, 0, 0)
 								[
 									SNew( STextBlock )
 									.TextStyle( FEditorStyle::Get(), "ContentBrowser.TopBar.Font" )
@@ -188,9 +205,9 @@ void SContentBrowser::Construct( const FArguments& InArgs, const FName& InInstan
 						.HAlign(HAlign_Left)
 						[
 							SNew( SButton )
-							.ButtonStyle( FEditorStyle::Get(), "ToggleButton" )
+							.ButtonStyle(FEditorStyle::Get(), "FlatButton")
 							.ToolTipText( LOCTEXT( "SaveDirtyPackagesTooltip", "Save all modified assets." ) )
-							.ContentPadding( 0 )
+							.ContentPadding(FMargin(6, 2))
 							.OnClicked( this, &SContentBrowser::OnSaveClicked )
 							.AddMetaData<FTagMetaData>(FTagMetaData(TEXT("ContentBrowserSaveDirtyPackages")))
 							[
@@ -198,18 +215,20 @@ void SContentBrowser::Construct( const FArguments& InArgs, const FName& InInstan
 
 								// Save All Icon
 								+ SHorizontalBox::Slot()
+								.VAlign(VAlign_Center)
 								.AutoWidth()
-								.VAlign( VAlign_Center )
 								[
-									SNew( SImage )
-									.Image( FEditorStyle::GetBrush( "ContentBrowser.SaveDirtyPackages" ) )
+									SNew(STextBlock)
+									.TextStyle(FEditorStyle::Get(), "ContentBrowser.TopBar.Font")
+									.Font(FEditorStyle::Get().GetFontStyle("FontAwesome.11"))
+									.Text(FString(TEXT("\xf0c7")) /*fa-floppy-o*/)
 								]
 
 								// Save All Text
 								+ SHorizontalBox::Slot()
 								.AutoWidth()
 								.VAlign(VAlign_Center)
-								.Padding(0,0,2,0)
+								.Padding(4, 0, 0, 0)
 								[
 									SNew( STextBlock )
 									.TextStyle( FEditorStyle::Get(), "ContentBrowser.TopBar.Font" )
@@ -233,7 +252,6 @@ void SContentBrowser::Construct( const FArguments& InArgs, const FName& InInstan
 					// History Back Button
 					+SHorizontalBox::Slot()
 					.AutoWidth()
-					//.Padding(0, 1)
 					[
 						SNew(SVerticalBox)
 
@@ -242,7 +260,7 @@ void SContentBrowser::Construct( const FArguments& InArgs, const FName& InInstan
 						[
 							SNew(SButton)
 							.VAlign(EVerticalAlignment::VAlign_Center)
-							.ButtonStyle( FEditorStyle::Get(), "ToggleButton" )
+							.ButtonStyle(FEditorStyle::Get(), "FlatButton")
 							.ForegroundColor( FEditorStyle::GetSlateColor(DefaultForegroundName) )
 							.ToolTipText( this, &SContentBrowser::GetHistoryBackTooltip )
 							.ContentPadding( FMargin(1, 0) )
@@ -250,8 +268,10 @@ void SContentBrowser::Construct( const FArguments& InArgs, const FName& InInstan
 							.IsEnabled(this, &SContentBrowser::IsBackEnabled)
 							.AddMetaData<FTagMetaData>(FTagMetaData(TEXT("ContentBrowserHistoryBack")))
 							[
-								SNew(SImage)
-								.Image(FEditorStyle::GetBrush("ContentBrowser.HistoryBack"))
+								SNew(STextBlock)
+								.TextStyle(FEditorStyle::Get(), "ContentBrowser.TopBar.Font")
+								.Font(FEditorStyle::Get().GetFontStyle("FontAwesome.11"))
+								.Text(FString(TEXT("\xf060")) /*fa-arrow-left*/)
 							]
 						]
 					]
@@ -259,7 +279,6 @@ void SContentBrowser::Construct( const FArguments& InArgs, const FName& InInstan
 					// History Forward Button
 					+ SHorizontalBox::Slot()
 					.AutoWidth()
-					//.Padding(0, 1)
 					[
 						SNew(SVerticalBox)
 
@@ -268,7 +287,7 @@ void SContentBrowser::Construct( const FArguments& InArgs, const FName& InInstan
 						[
 							SNew(SButton)
 							.VAlign(EVerticalAlignment::VAlign_Center)
-							.ButtonStyle( FEditorStyle::Get(), "ToggleButton" )
+							.ButtonStyle(FEditorStyle::Get(), "FlatButton")
 							.ForegroundColor( FEditorStyle::GetSlateColor(DefaultForegroundName) )
 							.ToolTipText( this, &SContentBrowser::GetHistoryForwardTooltip )
 							.ContentPadding( FMargin(1, 0) )
@@ -276,8 +295,10 @@ void SContentBrowser::Construct( const FArguments& InArgs, const FName& InInstan
 							.IsEnabled(this, &SContentBrowser::IsForwardEnabled)
 							.AddMetaData<FTagMetaData>(FTagMetaData(TEXT("ContentBrowserHistoryForward")))
 							[
-								SNew(SImage)
-								.Image(FEditorStyle::GetBrush("ContentBrowser.HistoryForward"))
+								SNew(STextBlock)
+								.TextStyle(FEditorStyle::Get(), "ContentBrowser.TopBar.Font")
+								.Font(FEditorStyle::Get().GetFontStyle("FontAwesome.11"))
+								.Text(FString(TEXT("\xf061")) /*fa-arrow-right*/)
 							]
 						]
 					]
@@ -297,16 +318,19 @@ void SContentBrowser::Construct( const FArguments& InArgs, const FName& InInstan
 					.VAlign( VAlign_Fill )
 					[
 						SAssignNew( PathPickerButton, SComboButton )
-						.ComboButtonStyle( FEditorStyle::Get(), "ToolbarComboButton" )
+						.ButtonStyle(FEditorStyle::Get(), "FlatButton")
 						.ForegroundColor(FLinearColor::White)
 						.ToolTipText( LOCTEXT( "PathPickerTooltip", "Choose a path" ) )
 						.OnGetMenuContent( this, &SContentBrowser::GetPathPickerContent )
 						.HasDownArrow( false )
 						.AddMetaData<FTagMetaData>(FTagMetaData(TEXT("ContentBrowserPathPicker")))
+						.ContentPadding(FMargin(3, 3))
 						.ButtonContent()
 						[
-							SNew( SImage )
-							.Image( FEditorStyle::GetBrush( "ContentBrowser.Sources" ) )
+							SNew(STextBlock)
+							.TextStyle(FEditorStyle::Get(), "ContentBrowser.TopBar.Font")
+							.Font(FEditorStyle::Get().GetFontStyle("FontAwesome.11"))
+							.Text(FString(TEXT("\xf07c")) /*fa-folder-open*/)
 						]
 					]
 
@@ -318,7 +342,8 @@ void SContentBrowser::Construct( const FArguments& InArgs, const FName& InInstan
 					.Padding(FMargin(0))
 					[
 						SAssignNew(PathBreadcrumbTrail, SBreadcrumbTrail<FString>)
-						.ButtonContentPadding(FMargin(3.0f, 3.0f))
+						.ButtonContentPadding(FMargin(2, 2))
+						.ButtonStyle(FEditorStyle::Get(), "FlatButton")
 						.DelimiterImage(FEditorStyle::GetBrush("ContentBrowser.PathDelimiter"))
 						.TextStyle(FEditorStyle::Get(), "ContentBrowser.PathText")
 						.ShowLeadingDelimiter(false)
@@ -340,7 +365,7 @@ void SContentBrowser::Construct( const FArguments& InArgs, const FName& InInstan
 						[
 							SNew(SButton)
 							.VAlign(EVerticalAlignment::VAlign_Center)
-							.ButtonStyle( FEditorStyle::Get(), "NoBorder" )
+							.ButtonStyle(FEditorStyle::Get(), "FlatButton")
 							.ToolTipText( LOCTEXT("LockToggleTooltip", "Toggle lock. If locked, this browser will ignore Find in Content Browser requests.") )
 							.ContentPadding( FMargin(1, 0) )
 							.OnClicked(this, &SContentBrowser::ToggleLockClicked)
