Fixed some new code that was providing text to Slate as FString rather than FText

UETOOL-213 - Minimize Slate FString -> FText conversion (remove SLATE_TEXT_ATTRIBUTE)

diff --git a/Engine/Source/Editor/ContentBrowser/Private/SContentBrowser.cpp b/Engine/Source/Editor/ContentBrowser/Private/SContentBrowser.cpp
--- a/Engine/Source/Editor/ContentBrowser/Private/SContentBrowser.cpp
+++ b/Engine/Source/Editor/ContentBrowser/Private/SContentBrowser.cpp
@@ -129,7 +129,7 @@ void SContentBrowser::Construct( const FArguments& InArgs, const FName& InInstan
 									SNew(STextBlock)
 									.TextStyle(FEditorStyle::Get(), "ContentBrowser.TopBar.Font")
 									.Font(FEditorStyle::Get().GetFontStyle("FontAwesome.11"))
-									.Text(FString(TEXT("\xf15b")) /*fa-file*/)
+									.Text(FText::FromString(FString(TEXT("\xf15b"))) /*fa-file*/)
 								]
 
 								// New Text
@@ -152,7 +152,7 @@ void SContentBrowser::Construct( const FArguments& InArgs, const FName& InInstan
 									SNew(STextBlock)
 									.TextStyle(FEditorStyle::Get(), "ContentBrowser.TopBar.Font")
 									.Font(FEditorStyle::Get().GetFontStyle("FontAwesome.10"))
-									.Text(FString(TEXT("\xf0d7")) /*fa-caret-down*/)
+									.Text(FText::FromString(FString(TEXT("\xf0d7"))) /*fa-caret-down*/)
 								]
 							]
 						]
@@ -182,7 +182,7 @@ void SContentBrowser::Construct( const FArguments& InArgs, const FName& InInstan
 									SNew(STextBlock)
 									.TextStyle(FEditorStyle::Get(), "ContentBrowser.TopBar.Font")
 									.Font(FEditorStyle::Get().GetFontStyle("FontAwesome.11"))
-									.Text(FString(TEXT("\xf019")) /*fa-download*/)
+									.Text(FText::FromString(FString(TEXT("\xf019"))) /*fa-download*/)
 								]
 
 								// Import Text
@@ -221,7 +221,7 @@ void SContentBrowser::Construct( const FArguments& InArgs, const FName& InInstan
 									SNew(STextBlock)
 									.TextStyle(FEditorStyle::Get(), "ContentBrowser.TopBar.Font")
 									.Font(FEditorStyle::Get().GetFontStyle("FontAwesome.11"))
-									.Text(FString(TEXT("\xf0c7")) /*fa-floppy-o*/)
+									.Text(FText::FromString(FString(TEXT("\xf0c7"))) /*fa-floppy-o*/)
 								]
 
 								// Save All Text
@@ -271,7 +271,7 @@ void SContentBrowser::Construct( const FArguments& InArgs, const FName& InInstan
 								SNew(STextBlock)
 								.TextStyle(FEditorStyle::Get(), "ContentBrowser.TopBar.Font")
 								.Font(FEditorStyle::Get().GetFontStyle("FontAwesome.11"))
-								.Text(FString(TEXT("\xf060")) /*fa-arrow-left*/)
+								.Text(FText::FromString(FString(TEXT("\xf060"))) /*fa-arrow-left*/)
 							]
 						]
 					]
@@ -298,7 +298,7 @@ void SContentBrowser::Construct( const FArguments& InArgs, const FName& InInstan
 								SNew(STextBlock)
 								.TextStyle(FEditorStyle::Get(), "ContentBrowser.TopBar.Font")
 								.Font(FEditorStyle::Get().GetFontStyle("FontAwesome.11"))
-								.Text(FString(TEXT("\xf061")) /*fa-arrow-right*/)
+								.Text(FText::FromString(FString(TEXT("\xf061"))) /*fa-arrow-right*/)
 							]
 						]
 					]
@@ -330,7 +330,7 @@ void SContentBrowser::Construct( const FArguments& InArgs, const FName& InInstan
 							SNew(STextBlock)
 							.TextStyle(FEditorStyle::Get(), "ContentBrowser.TopBar.Font")
 							.Font(FEditorStyle::Get().GetFontStyle("FontAwesome.11"))
-							.Text(FString(TEXT("\xf07c")) /*fa-folder-open*/)
+							.Text(FText::FromString(FString(TEXT("\xf07c"))) /*fa-folder-open*/)
 						]
 					]
 
@@ -523,7 +523,7 @@ void SContentBrowser::Construct( const FArguments& InArgs, const FName& InInstan
 									SNew(STextBlock)
 									.TextStyle(FEditorStyle::Get(), "ContentBrowser.Filters.Text")
 									.Font(FEditorStyle::Get().GetFontStyle("FontAwesome.9"))
-									.Text(FString(TEXT("\xf0b0")) /*fa-filter*/)
+									.Text(FText::FromString(FString(TEXT("\xf0b0"))) /*fa-filter*/)
 								]
 
 								+ SHorizontalBox::Slot()
diff --git a/Engine/Source/Editor/GameProjectGeneration/Private/SNewProjectWizard.cpp b/Engine/Source/Editor/GameProjectGeneration/Private/SNewProjectWizard.cpp
--- a/Engine/Source/Editor/GameProjectGeneration/Private/SNewProjectWizard.cpp
+++ b/Engine/Source/Editor/GameProjectGeneration/Private/SNewProjectWizard.cpp
@@ -338,7 +338,7 @@ void SNewProjectWizard::Construct( const FArguments& InArgs )
 		StartContentCombo = SNew(SDecoratedEnumCombo<int32>, MoveTemp(StarterContentInfo))
 			.SelectedEnum(this, &SNewProjectWizard::GetCopyStarterContentIndex)
 			.OnEnumChanged(this, &SNewProjectWizard::OnSetCopyStarterContent)
-			.ToolTipText(LOCTEXT("CopyStarterContent_ToolTip", "Enable to include an additional content pack containing simple placeable meshes with basic materials and textures.\nYou can opt out of including this to create a project that only has the bare essentials for the selected project template.").ToString());
+			.ToolTipText(LOCTEXT("CopyStarterContent_ToolTip", "Enable to include an additional content pack containing simple placeable meshes with basic materials and textures.\nYou can opt out of including this to create a project that only has the bare essentials for the selected project template."));
 	}
 
 	const float UniformPadding = 16.f;
diff --git a/Engine/Source/Editor/Kismet/Private/SSCSEditor.cpp b/Engine/Source/Editor/Kismet/Private/SSCSEditor.cpp
--- a/Engine/Source/Editor/Kismet/Private/SSCSEditor.cpp
+++ b/Engine/Source/Editor/Kismet/Private/SSCSEditor.cpp
@@ -3182,7 +3182,7 @@ void SSCSEditor::Construct( const FArguments& InArgs )
 							SNew(STextBlock)
 							.TextStyle(FEditorStyle::Get(), "ContentBrowser.TopBar.Font")
 							.Font( FEditorStyle::Get().GetFontStyle( "FontAwesome.10" ) )
-							.Text(FString(TEXT("\xf085")) /*fa-cogs*/)
+							.Text(FText::FromString(FString(TEXT("\xf085"))) /*fa-cogs*/)
 						]
 
 						+ SHorizontalBox::Slot()
@@ -3222,7 +3222,7 @@ void SSCSEditor::Construct( const FArguments& InArgs )
 							SNew(STextBlock)
 							.TextStyle(FEditorStyle::Get(), "ContentBrowser.TopBar.Font")
 							.Font(FEditorStyle::Get().GetFontStyle("FontAwesome.10"))
-							.Text(FString(TEXT("\xf085")) /*fa-cogs*/)
+							.Text(FText::FromString(FString(TEXT("\xf085"))) /*fa-cogs*/)
 						]
 						
 						+ SHorizontalBox::Slot()
diff --git a/Engine/Source/Editor/SourceControlWindows/Private/SSourceControlSubmit.cpp b/Engine/Source/Editor/SourceControlWindows/Private/SSourceControlSubmit.cpp
--- a/Engine/Source/Editor/SourceControlWindows/Private/SSourceControlSubmit.cpp
+++ b/Engine/Source/Editor/SourceControlWindows/Private/SSourceControlSubmit.cpp
@@ -31,7 +31,7 @@ FSubmitItem::FSubmitItem(const FSourceControlStateRef& InItem)
 	: Item(InItem)
 {
 	CheckBoxState = ECheckBoxState::Checked;
-	DisplayName = FPackageName::FilenameToLongPackageName(Item->GetFilename());
+	DisplayName = FText::FromString(FPackageName::FilenameToLongPackageName(Item->GetFilename()));
 }
 
 
@@ -404,12 +404,12 @@ void SSourceControlSubmitWidget::SortTree()
 		if (SortMode == EColumnSortMode::Ascending)
 		{
 			ListViewItems.Sort([](const TSharedPtr<FSubmitItem>& A, const TSharedPtr<FSubmitItem>& B) {
-				return A->GetDisplayName() < B->GetDisplayName(); });
+				return A->GetDisplayName().ToString() < B->GetDisplayName().ToString(); });
 		}
 		else if (SortMode == EColumnSortMode::Descending)
 		{
 			ListViewItems.Sort([](const TSharedPtr<FSubmitItem>& A, const TSharedPtr<FSubmitItem>& B) {
-				return A->GetDisplayName() >= B->GetDisplayName(); });
+				return A->GetDisplayName().ToString() >= B->GetDisplayName().ToString(); });
 		}
 	}
 	else if (SortByColumn == SSourceControlSubmitWidgetDefs::ColumnID_IconLabel)
diff --git a/Engine/Source/Editor/SourceControlWindows/Private/SSourceControlSubmit.h b/Engine/Source/Editor/SourceControlWindows/Private/SSourceControlSubmit.h
--- a/Engine/Source/Editor/SourceControlWindows/Private/SSourceControlSubmit.h
+++ b/Engine/Source/Editor/SourceControlWindows/Private/SSourceControlSubmit.h
@@ -36,7 +36,7 @@ public:
 	FString GetFilename() const { return Item->GetFilename(); }
 
 	/** Returns the name of the item as displayed in the widget */
-	FString GetDisplayName() const { return DisplayName; }
+	FText GetDisplayName() const { return DisplayName; }
 
 	/** Returns the name of the icon to be used in the list item widget */
 	FName GetIconName() const { return Item->GetSmallIconName(); }
@@ -67,7 +67,7 @@ private:
 	ECheckBoxState CheckBoxState;
 
 	/** Cached name to display in the listview */
-	FString DisplayName;
+	FText DisplayName;
 };
 
 class SSourceControlSubmitWidget : public SCompoundWidget
diff --git a/Engine/Source/Editor/UMGEditor/Private/TabFactory/AnimationTabSummoner.cpp b/Engine/Source/Editor/UMGEditor/Private/TabFactory/AnimationTabSummoner.cpp
--- a/Engine/Source/Editor/UMGEditor/Private/TabFactory/AnimationTabSummoner.cpp
+++ b/Engine/Source/Editor/UMGEditor/Private/TabFactory/AnimationTabSummoner.cpp
@@ -223,7 +223,7 @@ public:
 								SNew(STextBlock)
 								.TextStyle(FEditorStyle::Get(), "NormalText.Important")
 								.Font(FEditorStyle::Get().GetFontStyle("FontAwesome.10"))
-								.Text(FString(TEXT("\xf067")) /*fa-plus*/)
+								.Text(FText::FromString(FString(TEXT("\xf067"))) /*fa-plus*/)
 							]
 
 							+ SHorizontalBox::Slot()
diff --git a/Engine/Source/Editor/UnrealEd/Private/SComponentClassCombo.cpp b/Engine/Source/Editor/UnrealEd/Private/SComponentClassCombo.cpp
--- a/Engine/Source/Editor/UnrealEd/Private/SComponentClassCombo.cpp
+++ b/Engine/Source/Editor/UnrealEd/Private/SComponentClassCombo.cpp
@@ -55,7 +55,7 @@ void SComponentClassCombo::Construct(const FArguments& InArgs)
 			SNew(STextBlock)
 			.TextStyle(FEditorStyle::Get(), "ContentBrowser.TopBar.Font")
 			.Font(FEditorStyle::Get().GetFontStyle("FontAwesome.10"))
-			.Text(FString(TEXT("\xf067")) /*fa-plus*/)
+			.Text(FText::FromString(FString(TEXT("\xf067"))) /*fa-plus*/)
 		]
 		+ SHorizontalBox::Slot()
 		.VAlign(VAlign_Center)
