[UE-9398] Restored ability to duplicate the current level.

change summary:
- reverted CL# 2442516 per notes in JIRA, as the crash does not occur on main (may have been fixed)

diff --git a/Engine/Source/Editor/ContentBrowser/Private/AssetContextMenu.cpp b/Engine/Source/Editor/ContentBrowser/Private/AssetContextMenu.cpp
--- a/Engine/Source/Editor/ContentBrowser/Private/AssetContextMenu.cpp
+++ b/Engine/Source/Editor/ContentBrowser/Private/AssetContextMenu.cpp
@@ -48,7 +48,6 @@ FAssetContextMenu::FAssetContextMenu(const TWeakPtr<SAssetView>& InAssetView)
 	, bCanExecuteSCCHistory(false)
 	, bCanExecuteSCCRevert(false)
 	, bCanExecuteSCCSync(false)
-	, bCurrentWorldSelected(false)
 {
 	
 }
@@ -1644,7 +1643,7 @@ bool FAssetContextMenu::CanExecutePropertyMatrix() const
 
 bool FAssetContextMenu::CanExecuteDuplicate() const
 {
-	return bAtLeastOneNonRedirectorSelected && !bAtLeastOneClassSelected && !bCurrentWorldSelected;
+	return bAtLeastOneNonRedirectorSelected && !bAtLeastOneClassSelected;
 }
 
 bool FAssetContextMenu::CanExecuteRename() const
@@ -1847,9 +1846,7 @@ void FAssetContextMenu::CacheCanExecuteVars()
 	bCanExecuteSCCHistory = false;
 	bCanExecuteSCCRevert = false;
 	bCanExecuteSCCSync = false;
-	bCurrentWorldSelected = false;
 
-	static const FName WorldClass("World");
 	for (auto AssetIt = SelectedAssets.CreateConstIterator(); AssetIt; ++AssetIt)
 	{
 		const FAssetData& AssetData = *AssetIt;
@@ -1858,12 +1855,6 @@ void FAssetContextMenu::CacheCanExecuteVars()
 			continue;
 		}
 
-		// Hack: 4.7 current world is selected and cannot be duplicated due to blueprint trashed components issues.  
-		if( AssetData.AssetClass == WorldClass && AssetData.IsAssetLoaded() )
-		{
-			bCurrentWorldSelected = (GWorld == AssetData.GetAsset());
-		}
-
 		if ( !bAtLeastOneNonRedirectorSelected && AssetData.AssetClass != UObjectRedirector::StaticClass()->GetFName() )
 		{
 			bAtLeastOneNonRedirectorSelected = true;
diff --git a/Engine/Source/Editor/ContentBrowser/Private/AssetContextMenu.h b/Engine/Source/Editor/ContentBrowser/Private/AssetContextMenu.h
--- a/Engine/Source/Editor/ContentBrowser/Private/AssetContextMenu.h
+++ b/Engine/Source/Editor/ContentBrowser/Private/AssetContextMenu.h
@@ -331,7 +331,6 @@ private:
 	bool bCanExecuteSCCHistory;
 	bool bCanExecuteSCCRevert;
 	bool bCanExecuteSCCSync;
-	bool bCurrentWorldSelected;
 	/** */
 	int32 ChunkIDSelected;
 };
