"Show Collections" in the Content Browser now remembers its state.

diff --git a/Engine/Config/BaseEditorUserSettings.ini b/Engine/Config/BaseEditorUserSettings.ini
--- a/Engine/Config/BaseEditorUserSettings.ini
+++ b/Engine/Config/BaseEditorUserSettings.ini
@@ -111,6 +111,8 @@ DisplayFolders=True
 DisplayEngineFolder=False
 ; Whether to display the developers folder in the path or assets view of the content browser
 DisplayDevelopersFolder=False
+; Whether to display the collections view in the Content Browser
+DisplayCollections=False
 
 [/Script/UnrealEd.DestructableMeshEditorSettings]
 ; Persona viewport default settings
diff --git a/Engine/Source/Editor/ContentBrowser/Private/SContentBrowser.cpp b/Engine/Source/Editor/ContentBrowser/Private/SContentBrowser.cpp
--- a/Engine/Source/Editor/ContentBrowser/Private/SContentBrowser.cpp
+++ b/Engine/Source/Editor/ContentBrowser/Private/SContentBrowser.cpp
@@ -686,7 +686,8 @@ void SContentBrowser::ExtendAssetViewMenu( FMenuBuilder& MenuBuilder )
 
 void SContentBrowser::ToggleShowCollections()
 {
-	return GetMutableDefault<UContentBrowserSettings>()->SetDisplayCollections(!GetDefault<UContentBrowserSettings>()->GetDisplayCollections());
+	GetMutableDefault<UContentBrowserSettings>()->SetDisplayCollections(!GetDefault<UContentBrowserSettings>()->GetDisplayCollections());
+	GetMutableDefault<UContentBrowserSettings>()->PostEditChange();
 }
 
 bool SContentBrowser::IsShowingCollections() const
