Safety fix for unreproducible crash opening content browser view options menu

diff --git a/Engine/Source/Editor/ContentBrowser/Private/SContentBrowser.cpp b/Engine/Source/Editor/ContentBrowser/Private/SContentBrowser.cpp
--- a/Engine/Source/Editor/ContentBrowser/Private/SContentBrowser.cpp
+++ b/Engine/Source/Editor/ContentBrowser/Private/SContentBrowser.cpp
@@ -626,13 +626,18 @@ void SContentBrowser::Construct( const FArguments& InArgs, const FName& InInstan
 	FAssetRegistryModule& AssetRegistryModule = FModuleManager::LoadModuleChecked<FAssetRegistryModule>(TEXT("AssetRegistry"));
 	AssetRegistryModule.Get().OnPathRemoved().AddSP(this, &SContentBrowser::HandlePathRemoved);
 
+	const TWeakPtr<SContentBrowser> WeakThis = SharedThis(this);
+
 	FContentBrowserModule& ContentBrowserModule = FModuleManager::GetModuleChecked<FContentBrowserModule>( TEXT("ContentBrowser") );
 	ContentBrowserModule.GetAllAssetViewViewMenuExtenders().Add( 
-		FContentBrowserMenuExtender::CreateLambda( [&]()
+		FContentBrowserMenuExtender::CreateLambda( [=]()
 		{
-			TSharedRef<FExtender> Extender = MakeShareable( new FExtender );
+			TSharedRef<FExtender> Extender = MakeShareable(new FExtender);
+			if( WeakThis.IsValid() )
+			{
+				Extender->AddMenuExtension(FName("Folders"), EExtensionHook::After, nullptr, FMenuExtensionDelegate::CreateSP(WeakThis.Pin().ToSharedRef(), &SContentBrowser::ExtendAssetViewMenu));
+			}
 
-			Extender->AddMenuExtension( FName("Folders"), EExtensionHook::After, nullptr, FMenuExtensionDelegate::CreateSP( this, &SContentBrowser::ExtendAssetViewMenu ) );
 			return Extender;
 		})
 	);
