UE-11208
Content browser search is broken in some cases

diff --git a/Engine/Source/Editor/ContentBrowser/Private/FrontendFilters.cpp b/Engine/Source/Editor/ContentBrowser/Private/FrontendFilters.cpp
--- a/Engine/Source/Editor/ContentBrowser/Private/FrontendFilters.cpp
+++ b/Engine/Source/Editor/ContentBrowser/Private/FrontendFilters.cpp
@@ -10,19 +10,20 @@
 // FFrontendFilter_Text
 /////////////////////////////////////////
 
-void AssetDataToExportTextString(AssetFilterType Asset, OUT TArray< FString >& Array)
+void AssetDataToClassAndNameStrings(AssetFilterType Asset, OUT TArray< FString >& Array)
 {
-	Array.Add(Asset.GetExportTextName());
+	Array.Add(Asset.AssetClass.ToString());
+	Array.Add(Asset.AssetName.ToString());
 }
 
-void AssetDataToObjectPathString(AssetFilterType Asset, OUT TArray< FString >& Array)
+void AssetDataToNameString(AssetFilterType Asset, OUT TArray< FString >& Array)
 {
-	Array.Add(Asset.ObjectPath.ToString());
+	Array.Add(Asset.AssetName.ToString());
 }
 
 FFrontendFilter_Text::FFrontendFilter_Text()
 	: FFrontendFilter(nullptr)
-	, TextFilter( TTextFilter<AssetFilterType>::FItemToStringArray::CreateStatic( &AssetDataToExportTextString ) )
+	, TextFilter(TTextFilter<AssetFilterType>::FItemToStringArray::CreateStatic( &AssetDataToClassAndNameStrings ) )
 {
 	SetIncludeClassName(true);
 }
@@ -54,11 +55,11 @@ void FFrontendFilter_Text::SetIncludeClassName(bool bIncludeClassName)
 
 	if ( bIncludeClassName )
 	{
-		TextFilter = TTextFilter<AssetFilterType>::FItemToStringArray::CreateStatic( &AssetDataToExportTextString );
+		TextFilter = TTextFilter<AssetFilterType>::FItemToStringArray::CreateStatic( &AssetDataToClassAndNameStrings );
 	}
 	else
 	{
-		TextFilter = TTextFilter<AssetFilterType>::FItemToStringArray::CreateStatic( &AssetDataToObjectPathString );
+		TextFilter = TTextFilter<AssetFilterType>::FItemToStringArray::CreateStatic( &AssetDataToNameString );
 	}
 
 	// Apply the existing text before we re-assign the delegate since we want the text preservation to be opaque.
