Fixed a bug where the Content Browser could be left looking at Engine or Plugin classes after that view had been disabled via the view options

diff --git a/Engine/Source/Editor/ContentBrowser/Private/SPathView.cpp b/Engine/Source/Editor/ContentBrowser/Private/SPathView.cpp
--- a/Engine/Source/Editor/ContentBrowser/Private/SPathView.cpp
+++ b/Engine/Source/Editor/ContentBrowser/Private/SPathView.cpp
@@ -1410,7 +1410,12 @@ void SPathView::HandleSettingChanged(FName PropertyName)
 		if (!bDisplayDev || !bDisplayEngine || !bDisplayPlugins)
 		{
 			const FString OldSelectedPath = GetSelectedPath();
-			if ((!bDisplayDev && ContentBrowserUtils::IsDevelopersFolder(OldSelectedPath)) || (!bDisplayEngine && ContentBrowserUtils::IsEngineFolder(OldSelectedPath)) || (!bDisplayPlugins && ContentBrowserUtils::IsPluginFolder(OldSelectedPath)))
+			const ContentBrowserUtils::ECBFolderCategory OldFolderCategory = ContentBrowserUtils::GetFolderCategory(OldSelectedPath);
+
+			if ((!bDisplayDev && OldFolderCategory == ContentBrowserUtils::ECBFolderCategory::DeveloperContent) || 
+				(!bDisplayEngine && (OldFolderCategory == ContentBrowserUtils::ECBFolderCategory::EngineContent || OldFolderCategory == ContentBrowserUtils::ECBFolderCategory::EngineClasses)) || 
+				(!bDisplayPlugins && (OldFolderCategory == ContentBrowserUtils::ECBFolderCategory::PluginContent || OldFolderCategory == ContentBrowserUtils::ECBFolderCategory::PluginClasses))
+				)
 			{
 				// Set the folder back to the root, and refresh the contents
 				TSharedPtr<FTreeItem> GameRoot = FindItemRecursive(TEXT("/Game"));
@@ -1432,7 +1437,12 @@ void SPathView::HandleSettingChanged(FName PropertyName)
 		if (bDisplayDev || bDisplayEngine || bDisplayPlugins)
 		{
 			const FString NewSelectedPath = GetSelectedPath();
-			if ((bDisplayDev && ContentBrowserUtils::IsDevelopersFolder(NewSelectedPath)) || (bDisplayEngine && ContentBrowserUtils::IsEngineFolder(NewSelectedPath)) || (bDisplayPlugins && ContentBrowserUtils::IsPluginFolder(NewSelectedPath)))
+			const ContentBrowserUtils::ECBFolderCategory NewFolderCategory = ContentBrowserUtils::GetFolderCategory(NewSelectedPath);
+
+			if ((bDisplayDev && NewFolderCategory == ContentBrowserUtils::ECBFolderCategory::DeveloperContent) || 
+				(bDisplayEngine && (NewFolderCategory == ContentBrowserUtils::ECBFolderCategory::EngineContent || NewFolderCategory == ContentBrowserUtils::ECBFolderCategory::EngineClasses)) || 
+				(bDisplayPlugins && (NewFolderCategory == ContentBrowserUtils::ECBFolderCategory::PluginContent || NewFolderCategory == ContentBrowserUtils::ECBFolderCategory::PluginClasses))
+				)
 			{
 				// Refresh the contents
 				OnPathSelected.ExecuteIfBound(NewSelectedPath);
