fix for UE-11637 Copying Blueprints folder within the Content Browser is not working properly

diff --git a/Engine/Source/Editor/ContentBrowser/Private/DragDropHandler.cpp b/Engine/Source/Editor/ContentBrowser/Private/DragDropHandler.cpp
--- a/Engine/Source/Editor/ContentBrowser/Private/DragDropHandler.cpp
+++ b/Engine/Source/Editor/ContentBrowser/Private/DragDropHandler.cpp
@@ -177,7 +177,7 @@ void DragDropHandler::HandleFoldersDroppedOnAssetFolder(const TSharedRef<SWidget
 			LOCTEXT("DragDropCopyFolder", "Copy Folder Here"),
 			LOCTEXT("DragDropCopyFolderTooltip", "Creates a copy of all assets in the dragged folders to this folder, preserving folder structure."),
 			FSlateIcon(),
-			FUIAction(FExecuteAction::CreateLambda([=]() { MoveActionHandler.ExecuteIfBound(FinalPathNames, TargetPath); }))
+			FUIAction(FExecuteAction::CreateLambda([=]() { CopyActionHandler.ExecuteIfBound(FinalPathNames, TargetPath); }))
 			);
 
 		MenuBuilder.AddMenuEntry(
