Fixed issue in which folder could not be moved into another folder in the content browser whose name begins similarly e.g. "ab" into "abcd"

diff --git a/Engine/Source/Editor/ContentBrowser/Private/ContentBrowserUtils.cpp b/Engine/Source/Editor/ContentBrowser/Private/ContentBrowserUtils.cpp
--- a/Engine/Source/Editor/ContentBrowser/Private/ContentBrowserUtils.cpp
+++ b/Engine/Source/Editor/ContentBrowser/Private/ContentBrowserUtils.cpp
@@ -733,13 +733,14 @@ void ContentBrowserUtils::CopyFolders(const TArray<FString>& InSourcePathNames,
 void ContentBrowserUtils::MoveFolders(const TArray<FString>& InSourcePathNames, const FString& DestPath)
 {
 	TMap<FString, TArray<UObject*> > SourcePathToLoadedAssets;
+	FString DestPathWithTrailingSlash = DestPath / "";
 
 	// Do not allow parent directories to be moved to themselves or children.
 	TArray<FString> SourcePathNames = InSourcePathNames;
 	TArray<FString> SourcePathNamesToRemove;
 	for (auto SourcePathIt = SourcePathNames.CreateConstIterator(); SourcePathIt; ++SourcePathIt)
 	{
-		if (DestPath.StartsWith(*SourcePathIt))
+		if(DestPathWithTrailingSlash.StartsWith(*SourcePathIt / ""))
 		{
 			SourcePathNamesToRemove.Add(*SourcePathIt);
 		}
@@ -761,7 +762,7 @@ void ContentBrowserUtils::MoveFolders(const TArray<FString>& InSourcePathNames,
 		// Put dragged folders in a sub-folder under the destination path
 		const FString SourcePath = PathIt.Key();
 		const FString SubFolderName = FPackageName::GetLongPackageAssetName(SourcePath);
-		const FString Destination = DestPath + TEXT("/") + SubFolderName;
+		const FString Destination = DestPathWithTrailingSlash + SubFolderName;
 
 		// Add the new path to notify sources views
 		AssetRegistryModule.Get().AddPath(Destination);
