UE-11898 Unable to duplicated asset in Content Browser with Ctrl+W

diff --git a/Engine/Source/Editor/ContentBrowser/Private/AssetContextMenu.cpp b/Engine/Source/Editor/ContentBrowser/Private/AssetContextMenu.cpp
--- a/Engine/Source/Editor/ContentBrowser/Private/AssetContextMenu.cpp
+++ b/Engine/Source/Editor/ContentBrowser/Private/AssetContextMenu.cpp
@@ -1643,7 +1643,27 @@ bool FAssetContextMenu::CanExecutePropertyMatrix() const
 
 bool FAssetContextMenu::CanExecuteDuplicate() const
 {
-	return bAtLeastOneNonRedirectorSelected && !bAtLeastOneClassSelected;
+	const TArray< FAssetData > AssetViewSelectedAssets = AssetView.Pin()->GetSelectedAssets();
+	uint32 NumNonRedirectors = 0;
+	for(auto& AssetData : AssetViewSelectedAssets)
+	{
+		if(!AssetData.IsValid())
+		{
+			continue;
+		}
+
+		if(AssetData.AssetClass == NAME_Class)
+		{
+			return false;
+		}
+
+		if(AssetData.AssetClass != UObjectRedirector::StaticClass()->GetFName())
+		{
+			++NumNonRedirectors;
+		}
+	}
+
+	return (NumNonRedirectors > 0);
 }
 
 bool FAssetContextMenu::CanExecuteRename() const
