Renaming an asset no longer disables the filters in the content browser.

diff --git a/Engine/Source/Editor/ContentBrowser/Private/SContentBrowser.cpp b/Engine/Source/Editor/ContentBrowser/Private/SContentBrowser.cpp
--- a/Engine/Source/Editor/ContentBrowser/Private/SContentBrowser.cpp
+++ b/Engine/Source/Editor/ContentBrowser/Private/SContentBrowser.cpp
@@ -760,7 +760,7 @@ void SContentBrowser::ImportAsset( const FString& InPath )
 	}
 }
 
-void SContentBrowser::SyncToAssets( const TArray<FAssetData>& AssetDataList, const bool bAllowImplicitSync )
+void SContentBrowser::SyncToAssets( const TArray<FAssetData>& AssetDataList, const bool bAllowImplicitSync, const bool bDisableFiltersThatHideAssets )
 {
 	// Check to see if any of the assets require certain folders to be visible
 	const UContentBrowserSettings* tmp = GetDefault<UContentBrowserSettings>();
@@ -815,8 +815,11 @@ void SContentBrowser::SyncToAssets( const TArray<FAssetData>& AssetDataList, con
 		}
 	}
 
-	// Disable the filter categories
-	FilterListPtr->DisableFiltersThatHideAssets(AssetDataList);
+	if ( bDisableFiltersThatHideAssets )
+	{
+		// Disable the filter categories
+		FilterListPtr->DisableFiltersThatHideAssets(AssetDataList);
+	}
 
 	// Disable the filter search (reset the filter, then clear the search text)
 	// Note: we have to remove the filter immediately, we can't wait for OnSearchBoxChanged to hit
@@ -1988,7 +1991,9 @@ void SContentBrowser::OnAssetRenameCommitted(const TArray<FAssetData>& Assets)
 	// After a rename is committed we allow an implicit sync so as not to
 	// disorientate the user if they are looking at a parent folder
 
-	SyncToAssets(Assets, /*bAllowImplicitSync*/true);
+	const bool bAllowImplicitSync = true;
+	const bool bDisableFiltersThatHideAssets = false;
+	SyncToAssets(Assets, bAllowImplicitSync, bDisableFiltersThatHideAssets);
 }
 
 void SContentBrowser::OnFindInAssetTreeRequested(const TArray<FAssetData>& AssetsToFind)
diff --git a/Engine/Source/Editor/ContentBrowser/Private/SContentBrowser.h b/Engine/Source/Editor/ContentBrowser/Private/SContentBrowser.h
--- a/Engine/Source/Editor/ContentBrowser/Private/SContentBrowser.h
+++ b/Engine/Source/Editor/ContentBrowser/Private/SContentBrowser.h
@@ -37,7 +37,7 @@ public:
 	 *	@param bAllowImplicitSync	- true to allow the view to sync to parent folders if they are already selected,
 	 *								  false to force the view to select the explicit Parent folders of each asset 
 	 */
-	void SyncToAssets( const TArray<FAssetData>& AssetDataList, const bool bAllowImplicitSync = false );
+	void SyncToAssets( const TArray<FAssetData>& AssetDataList, const bool bAllowImplicitSync = false, const bool bDisableFiltersThatHideAssets = true );
 
 	/** Sets this content browser as the primary browser. The primary browser is the target for asset syncs and contributes to the global selection set. */
 	void SetIsPrimaryContentBrowser (bool NewIsPrimary);
