Editor: Minor stylistic improvements in the content browser "Asset Actions" submenu
- Fixed inconsistent icon usage in the 'Find' section
- Renamed "Select Containing Folder" to "Show in Sources Panel" to match how the documentation talks about it and by analogy with Show in [Explorer|Finder]
- Renamed "Property Matrix..." to "Bulk Edit via Property Matrix..." to use a more familiar grounding term

diff --git a/Engine/Source/Editor/ContentBrowser/Private/AssetContextMenu.cpp b/Engine/Source/Editor/ContentBrowser/Private/AssetContextMenu.cpp
--- a/Engine/Source/Editor/ContentBrowser/Private/AssetContextMenu.cpp
+++ b/Engine/Source/Editor/ContentBrowser/Private/AssetContextMenu.cpp
@@ -306,7 +306,7 @@ void FAssetContextMenu::MakeAssetActionsSubMenu(FMenuBuilder& MenuBuilder)
 		MenuBuilder.AddMenuEntry(
 			LOCTEXT("FindAssetInWorld", "Select Actors Using This Asset"),
 			LOCTEXT("FindAssetInWorldTooltip", "Selects all actors referencing this asset."),
-			FSlateIcon(FEditorStyle::GetStyleSetName(), "ContentBrowser.AssetActions.FindAssetInWorld"),
+			FSlateIcon(),
 			FUIAction(
 				FExecuteAction::CreateSP(this, &FAssetContextMenu::ExecuteFindAssetInWorld),
 				FCanExecuteAction::CreateSP(this, &FAssetContextMenu::CanExecuteFindAssetInWorld)
@@ -315,9 +315,9 @@ void FAssetContextMenu::MakeAssetActionsSubMenu(FMenuBuilder& MenuBuilder)
 
 		// Open Containing Folder
 		MenuBuilder.AddMenuEntry(
-			LOCTEXT("SyncToAssetTree", "Select Containing Folder"),
-			LOCTEXT("SyncToAssetTreeTooltip", "Selects the folder in the path view that contains this asset."),
-			FSlateIcon(FEditorStyle::GetStyleSetName(), "MainFrame.OpenProject"),
+			LOCTEXT("SyncToAssetTree", "Show in Sources Panel"),
+			LOCTEXT("SyncToAssetTreeTooltip", "Selects the folder that contains this asset in the Content Browser Sources Panel."),
+			FSlateIcon(),
 			FUIAction(
 				FExecuteAction::CreateSP(this, &FAssetContextMenu::ExecuteSyncToAssetTree),
 				FCanExecuteAction::CreateSP(this, &FAssetContextMenu::CanExecuteSyncToAssetTree)
@@ -387,7 +387,7 @@ void FAssetContextMenu::MakeAssetActionsSubMenu(FMenuBuilder& MenuBuilder)
 
 		// Property Matrix
 		MenuBuilder.AddMenuEntry(
-			LOCTEXT("PropertyMatrix", "Property Matrix..."),
+			LOCTEXT("PropertyMatrix", "Bulk Edit via Property Matrix..."),
 			LOCTEXT("PropertyMatrixTooltip", "Opens the property matrix editor for the selected assets."),
 			FSlateIcon(),
 			FUIAction(
