Content Browser: Added F prefix to FAssetFilterType and FAssetFilterCollectionType, and removed a redundant declaration

diff --git a/Engine/Source/Editor/ContentBrowser/Private/FrontendFilters.cpp b/Engine/Source/Editor/ContentBrowser/Private/FrontendFilters.cpp
--- a/Engine/Source/Editor/ContentBrowser/Private/FrontendFilters.cpp
+++ b/Engine/Source/Editor/ContentBrowser/Private/FrontendFilters.cpp
@@ -11,20 +11,20 @@
 // FFrontendFilter_Text
 /////////////////////////////////////////
 
-void AssetDataToClassAndNameStrings(AssetFilterType Asset, OUT TArray< FString >& Array)
+void AssetDataToClassAndNameStrings(FAssetFilterType Asset, OUT TArray< FString >& Array)
 {
 	Array.Add(Asset.AssetClass.ToString());
 	Array.Add(Asset.AssetName.ToString());
 }
 
-void AssetDataToNameString(AssetFilterType Asset, OUT TArray< FString >& Array)
+void AssetDataToNameString(FAssetFilterType Asset, OUT TArray< FString >& Array)
 {
 	Array.Add(Asset.AssetName.ToString());
 }
 
 FFrontendFilter_Text::FFrontendFilter_Text()
 	: FFrontendFilter(nullptr)
-	, TextFilter(TTextFilter<AssetFilterType>::FItemToStringArray::CreateStatic( &AssetDataToClassAndNameStrings ) )
+	, TextFilter(TTextFilter<FAssetFilterType>::FItemToStringArray::CreateStatic(&AssetDataToClassAndNameStrings))
 {
 	SetIncludeClassName(true);
 }
@@ -34,7 +34,7 @@ FFrontendFilter_Text::~FFrontendFilter_Text()
 	TextFilter.OnChanged().RemoveAll(this);
 }
 
-bool FFrontendFilter_Text::PassesFilter(AssetFilterType InItem) const
+bool FFrontendFilter_Text::PassesFilter(FAssetFilterType InItem) const
 {
 	return TextFilter.PassesFilter(InItem);
 }
@@ -56,11 +56,11 @@ void FFrontendFilter_Text::SetIncludeClassName(bool bIncludeClassName)
 
 	if ( bIncludeClassName )
 	{
-		TextFilter = TTextFilter<AssetFilterType>::FItemToStringArray::CreateStatic( &AssetDataToClassAndNameStrings );
+		TextFilter = TTextFilter<FAssetFilterType>::FItemToStringArray::CreateStatic(&AssetDataToClassAndNameStrings);
 	}
 	else
 	{
-		TextFilter = TTextFilter<AssetFilterType>::FItemToStringArray::CreateStatic( &AssetDataToNameString );
+		TextFilter = TTextFilter<FAssetFilterType>::FItemToStringArray::CreateStatic(&AssetDataToNameString);
 	}
 
 	// Apply the existing text before we re-assign the delegate since we want the text preservation to be opaque.
@@ -93,7 +93,7 @@ void FFrontendFilter_CheckedOut::ActiveStateChanged(bool bActive)
 	}
 }
 
-bool FFrontendFilter_CheckedOut::PassesFilter( AssetFilterType InItem ) const
+bool FFrontendFilter_CheckedOut::PassesFilter(FAssetFilterType InItem) const
 {
 	FSourceControlStatePtr SourceControlState = ISourceControlModule::Get().GetProvider().GetState(SourceControlHelpers::PackageFilename(InItem.PackageName.ToString()), EStateCacheUsage::Use);
 	return SourceControlState.IsValid() && (SourceControlState->IsCheckedOut() || SourceControlState->IsAdded());
@@ -137,7 +137,7 @@ void FFrontendFilter_Modified::ActiveStateChanged(bool bActive)
 	bIsCurrentlyActive = bActive;
 }
 
-bool FFrontendFilter_Modified::PassesFilter( AssetFilterType InItem ) const
+bool FFrontendFilter_Modified::PassesFilter(FAssetFilterType InItem) const
 {
 	UPackage* Package = FindPackage(NULL, *InItem.PackageName.ToString());
 
@@ -161,7 +161,7 @@ void FFrontendFilter_Modified::OnPackageDirtyStateUpdated(UPackage* Package)
 // FFrontendFilter_ReplicatedBlueprint
 /////////////////////////////////////////
 
-bool FFrontendFilter_ReplicatedBlueprint::PassesFilter( AssetFilterType InItem ) const
+bool FFrontendFilter_ReplicatedBlueprint::PassesFilter(FAssetFilterType InItem) const
 {
 	FString TagValue = InItem.TagsAndValues.FindRef("NumReplicatedProperties");
 	return !TagValue.IsEmpty() && (FCString::Atoi(*TagValue) > 0);
@@ -196,7 +196,7 @@ void FFrontendFilter_ShowOtherDevelopers::SetCurrentFilter(const FARFilter& InFi
 	}
 }
 
-bool FFrontendFilter_ShowOtherDevelopers::PassesFilter( AssetFilterType InItem ) const
+bool FFrontendFilter_ShowOtherDevelopers::PassesFilter(FAssetFilterType InItem) const
 {
 	// Pass all assets if other developer assets are allowed
 	if ( !bShowOtherDeveloperAssets )
@@ -252,7 +252,7 @@ void FFrontendFilter_ShowRedirectors::SetCurrentFilter(const FARFilter& InFilter
 	bAreRedirectorsInBaseFilter = InFilter.ClassNames.Contains(RedirectorClassName);
 }
 
-bool FFrontendFilter_ShowRedirectors::PassesFilter( AssetFilterType InItem ) const
+bool FFrontendFilter_ShowRedirectors::PassesFilter(FAssetFilterType InItem) const
 {
 	// Never hide redirectors if they are explicitly searched for
 	if ( !bAreRedirectorsInBaseFilter )
@@ -301,7 +301,7 @@ void FFrontendFilter_InUseByLoadedLevels::OnAssetPostRename(const TArray<FAssetR
 	ObjectTools::TagInUseObjects(ObjectTools::SO_LoadedLevels);
 }
 
-bool FFrontendFilter_InUseByLoadedLevels::PassesFilter( AssetFilterType InItem ) const
+bool FFrontendFilter_InUseByLoadedLevels::PassesFilter(FAssetFilterType InItem) const
 {
 	bool bObjectInUse = false;
 	
diff --git a/Engine/Source/Editor/ContentBrowser/Private/FrontendFilters.h b/Engine/Source/Editor/ContentBrowser/Private/FrontendFilters.h
--- a/Engine/Source/Editor/ContentBrowser/Private/FrontendFilters.h
+++ b/Engine/Source/Editor/ContentBrowser/Private/FrontendFilters.h
@@ -19,7 +19,7 @@ public:
 	virtual FText GetToolTipText() const override { return LOCTEXT("FrontendFilter_TextTooltip", "Show only assets that match the input text"); }
 
 	// IFilter implementation
-	virtual bool PassesFilter(AssetFilterType InItem) const override;
+	virtual bool PassesFilter(FAssetFilterType InItem) const override;
 
 public:
 	/** Returns the unsanitized and unsplit filter terms */
@@ -36,7 +36,7 @@ private:
 	/** Handler for the internal text filter */
 	void HandleOnChangedEvent();
 
-	TTextFilter<AssetFilterType> TextFilter;
+	TTextFilter<FAssetFilterType> TextFilter;
 };
 
 /** A filter that displays only checked out assets */
@@ -52,7 +52,7 @@ public:
 	virtual void ActiveStateChanged(bool bActive) override;
 
 	// IFilter implementation
-	virtual bool PassesFilter( AssetFilterType InItem ) const override;
+	virtual bool PassesFilter(FAssetFilterType InItem) const override;
 
 private:
 	
@@ -77,7 +77,7 @@ public:
 	virtual void ActiveStateChanged(bool bActive) override;
 
 	// IFilter implementation
-	virtual bool PassesFilter( AssetFilterType InItem ) const override;
+	virtual bool PassesFilter(FAssetFilterType InItem) const override;
 
 private:
 
@@ -99,7 +99,7 @@ public:
 	virtual FText GetToolTipText() const override { return LOCTEXT("FFrontendFilter_ReplicatedBlueprintToolTip", "Show only blueprints with replicated properties."); }
 
 	// IFilter implementation
-	virtual bool PassesFilter( AssetFilterType InItem ) const override;
+	virtual bool PassesFilter(FAssetFilterType InItem) const override;
 };
 
 /** An inverse filter that allows display of content in developer folders that are not the current user's */
@@ -117,7 +117,7 @@ public:
 	virtual void SetCurrentFilter(const FARFilter& InFilter) override;
 
 	// IFilter implementation
-	virtual bool PassesFilter( AssetFilterType InItem ) const override;
+	virtual bool PassesFilter(FAssetFilterType InItem) const override;
 
 public:
 	/** Sets if we should filter out assets from other developers */
@@ -148,7 +148,7 @@ public:
 	virtual void SetCurrentFilter(const FARFilter& InFilter) override;
 
 	// IFilter implementation
-	virtual bool PassesFilter( AssetFilterType InItem ) const override;
+	virtual bool PassesFilter(FAssetFilterType InItem) const override;
 
 private:
 	bool bAreRedirectorsInBaseFilter;
@@ -170,7 +170,7 @@ public:
 	virtual void ActiveStateChanged(bool bActive) override;
 
 	// IFilter implementation
-	virtual bool PassesFilter( AssetFilterType InItem ) const override;
+	virtual bool PassesFilter(FAssetFilterType InItem) const override;
 
 	/** Handler for when maps change in the editor */
 	void OnEditorMapChange( uint32 MapChangeFlags );
diff --git a/Engine/Source/Editor/ContentBrowser/Private/SAssetPicker.cpp b/Engine/Source/Editor/ContentBrowser/Private/SAssetPicker.cpp
--- a/Engine/Source/Editor/ContentBrowser/Private/SAssetPicker.cpp
+++ b/Engine/Source/Editor/ContentBrowser/Private/SAssetPicker.cpp
@@ -62,7 +62,7 @@ void SAssetPicker::Construct( const FArguments& InArgs )
 	TAttribute< FText > HighlightText;
 	EThumbnailLabel::Type ThumbnailLabel = InArgs._AssetPickerConfig.ThumbnailLabel;
 
-	FrontendFilters = MakeShareable(new AssetFilterCollectionType());
+	FrontendFilters = MakeShareable(new FAssetFilterCollectionType());
 
 	// Search box
 	if (!InArgs._AssetPickerConfig.bAutohideSearchBar)
diff --git a/Engine/Source/Editor/ContentBrowser/Private/SAssetPicker.h b/Engine/Source/Editor/ContentBrowser/Private/SAssetPicker.h
--- a/Engine/Source/Editor/ContentBrowser/Private/SAssetPicker.h
+++ b/Engine/Source/Editor/ContentBrowser/Private/SAssetPicker.h
@@ -90,7 +90,7 @@ private:
 private:
 
 	/** The list of FrontendFilters currently applied to the asset view */
-	TSharedPtr<AssetFilterCollectionType> FrontendFilters;
+	TSharedPtr<FAssetFilterCollectionType> FrontendFilters;
 
 	/** The asset view widget */
 	TSharedPtr<SAssetView> AssetViewPtr;
@@ -120,7 +120,7 @@ private:
 	bool bPendingFocusNextFrame;
 
 	/** Filters needed for filtering the assets */
-	TSharedPtr< AssetFilterCollectionType > FilterCollection;
+	TSharedPtr< FAssetFilterCollectionType > FilterCollection;
 	TSharedPtr< FFrontendFilter_Text > TextFilter;
 	TSharedPtr< FFrontendFilter_ShowOtherDevelopers > OtherDevelopersFilter;
 
diff --git a/Engine/Source/Editor/ContentBrowser/Private/SAssetView.h b/Engine/Source/Editor/ContentBrowser/Private/SAssetView.h
--- a/Engine/Source/Editor/ContentBrowser/Private/SAssetView.h
+++ b/Engine/Source/Editor/ContentBrowser/Private/SAssetView.h
@@ -82,7 +82,7 @@ public:
 		SLATE_ARGUMENT( bool, AllowThumbnailHintLabel )
 
 		/** The filter collection used to further filter down assets returned from the backend */
-		SLATE_ARGUMENT( TSharedPtr<AssetFilterCollectionType>, FrontendFilters )
+		SLATE_ARGUMENT( TSharedPtr<FAssetFilterCollectionType>, FrontendFilters )
 
 		/** The initial base sources filter */
 		SLATE_ARGUMENT( FSourcesData, InitialSourcesData )
@@ -632,7 +632,7 @@ private:
 	/** The current base source filter for the view */
 	FSourcesData SourcesData;
 	FARFilter BackendFilter;
-	TSharedPtr<AssetFilterCollectionType> FrontendFilters;
+	TSharedPtr<FAssetFilterCollectionType> FrontendFilters;
 
 	/** If true, the source items will be refreshed next frame. Very slow. */
 	bool bSlowFullListRefreshRequested;
diff --git a/Engine/Source/Editor/ContentBrowser/Private/SContentBrowser.cpp b/Engine/Source/Editor/ContentBrowser/Private/SContentBrowser.cpp
--- a/Engine/Source/Editor/ContentBrowser/Private/SContentBrowser.cpp
+++ b/Engine/Source/Editor/ContentBrowser/Private/SContentBrowser.cpp
@@ -64,7 +64,7 @@ void SContentBrowser::Construct( const FArguments& InArgs, const FName& InInstan
 	PathContextMenu->SetOnRenameFolderRequested(FPathContextMenu::FOnRenameFolderRequested::CreateSP(this, &SContentBrowser::OnRenameFolderRequested));
 	PathContextMenu->SetOnFolderDeleted(FPathContextMenu::FOnFolderDeleted::CreateSP(this, &SContentBrowser::OnOpenedFolderDeleted));
 
-	FrontendFilters = MakeShareable(new AssetFilterCollectionType());
+	FrontendFilters = MakeShareable(new FAssetFilterCollectionType());
 	TextFilter = MakeShareable( new FFrontendFilter_Text() );
 
 	static const FName DefaultForegroundName("DefaultForeground");
diff --git a/Engine/Source/Editor/ContentBrowser/Private/SContentBrowser.h b/Engine/Source/Editor/ContentBrowser/Private/SContentBrowser.h
--- a/Engine/Source/Editor/ContentBrowser/Private/SContentBrowser.h
+++ b/Engine/Source/Editor/ContentBrowser/Private/SContentBrowser.h
@@ -379,7 +379,7 @@ private:
 	bool bIsLocked;
 
 	/** The list of FrontendFilters currently applied to the asset view */
-	TSharedPtr<AssetFilterCollectionType> FrontendFilters;
+	TSharedPtr<FAssetFilterCollectionType> FrontendFilters;
 
 	/** The text filter to use on the assets */
 	TSharedPtr< FFrontendFilter_Text > TextFilter;
diff --git a/Engine/Source/Editor/ContentBrowser/Private/SFilterList.h b/Engine/Source/Editor/ContentBrowser/Private/SFilterList.h
--- a/Engine/Source/Editor/ContentBrowser/Private/SFilterList.h
+++ b/Engine/Source/Editor/ContentBrowser/Private/SFilterList.h
@@ -28,7 +28,7 @@ public:
 		SLATE_EVENT( FOnFilterChanged, OnFilterChanged )
 
 		/** The filter collection used to further filter down assets returned from the backend */
-		SLATE_ARGUMENT( TSharedPtr<AssetFilterCollectionType>, FrontendFilters)
+		SLATE_ARGUMENT( TSharedPtr<FAssetFilterCollectionType>, FrontendFilters)
 
 		/** An array of classes to filter the menu by */
 		SLATE_ARGUMENT( TArray<UClass*>, InitialClassFilters)
@@ -127,7 +127,7 @@ private:
 	TArray<TSharedRef<SFilter>> Filters;
 
 	/** The filter collection used to further filter down assets returned from the backend */
-	TSharedPtr<AssetFilterCollectionType> FrontendFilters;
+	TSharedPtr<FAssetFilterCollectionType> FrontendFilters;
 
 	/** All possible frontend filter objects */
 	TArray< TSharedRef<FFrontendFilter> > AllFrontendFilters;
diff --git a/Engine/Source/Editor/ContentBrowser/Public/FrontendFilterBase.h b/Engine/Source/Editor/ContentBrowser/Public/FrontendFilterBase.h
--- a/Engine/Source/Editor/ContentBrowser/Public/FrontendFilterBase.h
+++ b/Engine/Source/Editor/ContentBrowser/Public/FrontendFilterBase.h
@@ -2,9 +2,6 @@
 
 #pragma once
 
-typedef const FAssetData& AssetFilterType;
-typedef TFilterCollection<AssetFilterType> AssetFilterCollectionType;
-
 class FFrontendFilterCategory
 {
 public:
@@ -17,7 +14,7 @@ public:
 	const FText Tooltip;
 };
 
-class FFrontendFilter : public IFilter<AssetFilterType>
+class FFrontendFilter : public IFilter<FAssetFilterType>
 {
 public:
 	FFrontendFilter(TSharedPtr<FFrontendFilterCategory> InCategory) : FilterCategory(InCategory) {}
@@ -47,7 +44,7 @@ public:
 	virtual void ActiveStateChanged(bool bActive) { }
 
 	// IFilter implementation
-	DECLARE_DERIVED_EVENT( FFrontendFilter, IFilter<AssetFilterType>::FChangedEvent, FChangedEvent );
+	DECLARE_DERIVED_EVENT( FFrontendFilter, IFilter<FAssetFilterType>::FChangedEvent, FChangedEvent );
 	virtual FChangedEvent& OnChanged() override { return ChangedEvent; }
 
 	TSharedPtr<FFrontendFilterCategory> GetCategory() { return FilterCategory; }
diff --git a/Engine/Source/Editor/ContentBrowser/Public/IContentBrowserSingleton.h b/Engine/Source/Editor/ContentBrowser/Public/IContentBrowserSingleton.h
--- a/Engine/Source/Editor/ContentBrowser/Public/IContentBrowserSingleton.h
+++ b/Engine/Source/Editor/ContentBrowser/Public/IContentBrowserSingleton.h
@@ -1,6 +1,5 @@
 // Copyright 1998-2015 Epic Games, Inc. All Rights Reserved.
 
-
 #pragma once
 
 #include "ContentBrowserDelegates.h"
@@ -10,8 +9,8 @@
 #include "FilterCollection.h"
 #include "AssetThumbnail.h"
 
-typedef const FAssetData& AssetFilterType;
-typedef TFilterCollection<AssetFilterType> AssetFilterCollectionType;
+typedef const FAssetData& FAssetFilterType;
+typedef TFilterCollection<FAssetFilterType> FAssetFilterCollectionType;
 
 /** The view modes used in SAssetView */
 namespace EAssetViewType
diff --git a/Engine/Source/Editor/Persona/Private/SAnimationSequenceBrowser.cpp b/Engine/Source/Editor/Persona/Private/SAnimationSequenceBrowser.cpp
--- a/Engine/Source/Editor/Persona/Private/SAnimationSequenceBrowser.cpp
+++ b/Engine/Source/Editor/Persona/Private/SAnimationSequenceBrowser.cpp
@@ -31,9 +31,9 @@ public:
 	virtual FText GetToolTipText() const override { return LOCTEXT("FFrontendFilter_AdditiveAnimAssetsToolTip", "Show only animations that are additive."); }
 
 	// IFilter implementation
-	virtual bool PassesFilter( AssetFilterType InItem ) const override
+	virtual bool PassesFilter(FAssetFilterType InItem) const override
 	{
-		FString TagValue = InItem.TagsAndValues.FindRef("AdditiveAnimType");
+		const FString TagValue = InItem.TagsAndValues.FindRef(GET_MEMBER_NAME_CHECKED(UAnimSequence, AdditiveAnimType));
 		return !TagValue.IsEmpty() && !TagValue.Equals(TEXT("AAT_None"));
 	}
 };
