Editor: Fixed a typo in a comment

diff --git a/Engine/Source/Editor/ContentBrowser/Private/AssetViewSortManager.cpp b/Engine/Source/Editor/ContentBrowser/Private/AssetViewSortManager.cpp
--- a/Engine/Source/Editor/ContentBrowser/Private/AssetViewSortManager.cpp
+++ b/Engine/Source/Editor/ContentBrowser/Private/AssetViewSortManager.cpp
@@ -371,12 +371,12 @@ void FAssetViewSortManager::SortList(TArray<TSharedPtr<FAssetViewItem>>& AssetIt
 
 			if (TagType == UObject::FAssetRegistryTag::TT_Numerical)
 			{
-				// The property is a Num2er, compare using atof
+				// The property is a Number, compare using atof
 				SortMethod.Add(MakeUnique<FCompareFAssetItemByTagNumerical>(bAscending, Tag));
 			}
 			else if (TagType == UObject::FAssetRegistryTag::TT_Dimensional)
 			{
-				// The property is a series of Num2ers representing dimensions, compare by using atof for each Num2er, delimited by an "x"
+				// The property is a series of Numbers representing dimensions, compare by using atof for each Number, delimited by an "x"
 				SortMethod.Add(MakeUnique<FCompareFAssetItemByTagDimensional>(bAscending, Tag));
 			}
 			else
