Content Browser: Added a prototype configurable asset tag filter
- You can configure the tag, value, and comparison operation in the context menu for the filter bar button
- Note: Does not persist the selected tag, value, or operation across editor runs yet
- Note: Due to current assumptions about front-end filters, you can only have one tag filter active at a time
#epicfriday

diff --git a/Engine/Source/Editor/ContentBrowser/Private/FrontendFilters.cpp b/Engine/Source/Editor/ContentBrowser/Private/FrontendFilters.cpp
--- a/Engine/Source/Editor/ContentBrowser/Private/FrontendFilters.cpp
+++ b/Engine/Source/Editor/ContentBrowser/Private/FrontendFilters.cpp
@@ -167,6 +167,190 @@ bool FFrontendFilter_ReplicatedBlueprint::PassesFilter(FAssetFilterType InItem)
 	return !TagValue.IsEmpty() && (FCString::Atoi(*TagValue) > 0);
 }
 
+/////////////////////////////////////////
+// FFrontendFilter_ArbitraryComparisonOperation
+/////////////////////////////////////////
+
+#define LOCTEXT_NAMESPACE "ContentBrowser"
+
+FFrontendFilter_ArbitraryComparisonOperation::FFrontendFilter_ArbitraryComparisonOperation(TSharedPtr<FFrontendFilterCategory> InCategory)
+	: FFrontendFilter(InCategory)
+	, TagName(TEXT("TagName"))
+	, TargetTagValue(TEXT("Value"))
+	, ComparisonOp(ETextFilterComparisonOperation::NotEqual)
+{
+}
+
+FString FFrontendFilter_ArbitraryComparisonOperation::GetName() const
+{
+	return TEXT("CompareTags");
+}
+
+FText FFrontendFilter_ArbitraryComparisonOperation::GetDisplayName() const
+{
+	return FText::Format(LOCTEXT("FFrontendFilter_CompareOperation", "Compare Tags ({0} {1} {2})"),
+		FText::FromName(TagName),
+		FText::AsCultureInvariant(ConvertOperationToString(ComparisonOp)),
+		FText::AsCultureInvariant(TargetTagValue));
+}
+
+FText FFrontendFilter_ArbitraryComparisonOperation::GetToolTipText() const
+{
+	return LOCTEXT("FFrontendFilter_CompareOperation", "Compares AssetRegistrySearchable values on assets with a target value.");
+}
+
+bool FFrontendFilter_ArbitraryComparisonOperation::PassesFilter(FAssetFilterType InItem) const
+{
+	if (const FString* pTagValue = InItem.TagsAndValues.Find(TagName))
+	{
+		const FString& TestTagValue = *pTagValue;
+
+		const bool bNumericComparison = TestTagValue.IsNumeric() && TargetTagValue.IsNumeric();
+
+		int32 ComparisonSign = 0;
+		if (bNumericComparison)
+		{
+			const double Difference = FCString::Atod(*TestTagValue) - FCString::Atod(*TargetTagValue);
+			ComparisonSign = (int32)FMath::Sign(Difference);
+		}
+		else
+		{
+			ComparisonSign = TestTagValue.Compare(TargetTagValue, ESearchCase::IgnoreCase);
+		}
+
+		bool bComparisonPassed = false;
+		switch (ComparisonOp)
+		{
+		case ETextFilterComparisonOperation::Equal:
+			bComparisonPassed = ComparisonSign == 0;
+			break;
+		case ETextFilterComparisonOperation::NotEqual:
+			bComparisonPassed = ComparisonSign != 0;
+			break;
+		case ETextFilterComparisonOperation::Less:
+			bComparisonPassed = ComparisonSign < 0;
+			break;
+		case ETextFilterComparisonOperation::LessOrEqual:
+			bComparisonPassed = ComparisonSign <= 0;
+			break;
+		case ETextFilterComparisonOperation::Greater:
+			bComparisonPassed = ComparisonSign > 0;
+			break;
+		case ETextFilterComparisonOperation::GreaterOrEqual:
+			bComparisonPassed = ComparisonSign >= 0;
+			break;
+		default:
+			check(false);
+		};
+
+		return bComparisonPassed;
+	}
+	else
+	{
+		// Failed to find the tag, can't pass the filter
+		//@TODO: Maybe we should succeed here if the operation is !=
+		return false;
+	}
+}
+
+void FFrontendFilter_ArbitraryComparisonOperation::ModifyContextMenu(FMenuBuilder& MenuBuilder)
+{
+	FUIAction Action;
+
+	MenuBuilder.BeginSection(TEXT("ComparsionSection"), LOCTEXT("ComparisonSectionHeading", "AssetRegistrySearchable Comparison"));
+
+	TSharedRef<SWidget> KeyWidget =
+		SNew(SEditableTextBox)
+		.Text_Raw(this, &FFrontendFilter_ArbitraryComparisonOperation::GetKeyValueAsText)
+		.OnTextCommitted_Raw(this, &FFrontendFilter_ArbitraryComparisonOperation::OnKeyValueTextCommitted)
+		.MinDesiredWidth(100.0f);
+	TSharedRef<SWidget> ValueWidget = SNew(SEditableTextBox)
+		.Text_Raw(this, &FFrontendFilter_ArbitraryComparisonOperation::GetTargetValueAsText)
+		.OnTextCommitted_Raw(this, &FFrontendFilter_ArbitraryComparisonOperation::OnTargetValueTextCommitted)
+		.MinDesiredWidth(100.0f);
+
+	MenuBuilder.AddWidget(KeyWidget, LOCTEXT("KeyMenuDesc", "Tag"));
+	MenuBuilder.AddWidget(ValueWidget, LOCTEXT("ValueMenuDesc", "Target Value"));
+
+#define UE_SET_COMP_OP(Operation) \
+	MenuBuilder.AddMenuEntry(FText::AsCultureInvariant(ConvertOperationToString(Operation)), \
+		LOCTEXT("SwitchOpsTooltip", "Switch comparsion type"), \
+		FSlateIcon(), \
+		FUIAction(FExecuteAction::CreateRaw(this, &FFrontendFilter_ArbitraryComparisonOperation::SetComparisonOperation, Operation), FCanExecuteAction(), FIsActionChecked::CreateRaw(this, &FFrontendFilter_ArbitraryComparisonOperation::IsComparisonOperationEqualTo, Operation)), \
+		NAME_None, \
+		EUserInterfaceActionType::RadioButton);
+
+	UE_SET_COMP_OP(ETextFilterComparisonOperation::Equal);
+	UE_SET_COMP_OP(ETextFilterComparisonOperation::NotEqual);
+	UE_SET_COMP_OP(ETextFilterComparisonOperation::Less);
+	UE_SET_COMP_OP(ETextFilterComparisonOperation::LessOrEqual);
+	UE_SET_COMP_OP(ETextFilterComparisonOperation::Greater);
+	UE_SET_COMP_OP(ETextFilterComparisonOperation::GreaterOrEqual);
+#undef UE_SET_COMP_OP
+
+	MenuBuilder.EndSection();
+}
+
+void FFrontendFilter_ArbitraryComparisonOperation::SetComparisonOperation(ETextFilterComparisonOperation NewOp)
+{
+	ComparisonOp = NewOp;
+	BroadcastChangedEvent();
+}
+
+bool FFrontendFilter_ArbitraryComparisonOperation::IsComparisonOperationEqualTo(ETextFilterComparisonOperation TestOp) const
+{
+	return (ComparisonOp == TestOp);
+}
+
+FText FFrontendFilter_ArbitraryComparisonOperation::GetKeyValueAsText() const
+{
+	return FText::FromName(TagName);
+}
+
+FText FFrontendFilter_ArbitraryComparisonOperation::GetTargetValueAsText() const
+{
+	return FText::AsCultureInvariant(TargetTagValue);
+}
+
+void FFrontendFilter_ArbitraryComparisonOperation::OnKeyValueTextCommitted(const FText& InText, ETextCommit::Type InCommitType)
+{
+	if (!InText.IsEmpty())
+	{
+		TagName = *InText.ToString();
+		BroadcastChangedEvent();
+	}
+}
+
+void FFrontendFilter_ArbitraryComparisonOperation::OnTargetValueTextCommitted(const FText& InText, ETextCommit::Type InCommitType)
+{
+	TargetTagValue = InText.ToString();
+	BroadcastChangedEvent();
+}
+
+FString FFrontendFilter_ArbitraryComparisonOperation::ConvertOperationToString(ETextFilterComparisonOperation Op)
+{
+	switch (Op)
+	{
+	case ETextFilterComparisonOperation::Equal:
+		return TEXT("==");
+	case ETextFilterComparisonOperation::NotEqual:
+		return TEXT("!=");
+	case ETextFilterComparisonOperation::Less:
+		return TEXT("<");
+	case ETextFilterComparisonOperation::LessOrEqual:
+		return TEXT("<=");
+	case ETextFilterComparisonOperation::Greater:
+		return TEXT(">");
+	case ETextFilterComparisonOperation::GreaterOrEqual:
+		return TEXT(">=");
+	default:
+		check(false);
+		return TEXT("op");
+	};
+}
+
+#undef LOCTEXT_NAMESPACE
+
 /////////////////////////////////////////
 // FFrontendFilter_ShowOtherDevelopers
 /////////////////////////////////////////
diff --git a/Engine/Source/Editor/ContentBrowser/Private/FrontendFilters.h b/Engine/Source/Editor/ContentBrowser/Private/FrontendFilters.h
--- a/Engine/Source/Editor/ContentBrowser/Private/FrontendFilters.h
+++ b/Engine/Source/Editor/ContentBrowser/Private/FrontendFilters.h
@@ -102,6 +102,38 @@ public:
 	virtual bool PassesFilter(FAssetFilterType InItem) const override;
 };
 
+/** A filter that compares the value of an asset registry tag to a target value */
+class FFrontendFilter_ArbitraryComparisonOperation : public FFrontendFilter
+{
+public:
+	FFrontendFilter_ArbitraryComparisonOperation(TSharedPtr<FFrontendFilterCategory> InCategory);
+
+	// FFrontendFilter implementation
+	virtual FString GetName() const override;
+	virtual FText GetDisplayName() const override;
+	virtual FText GetToolTipText() const override;
+	virtual void ModifyContextMenu(FMenuBuilder& MenuBuilder) override;
+
+	// IFilter implementation
+	virtual bool PassesFilter(FAssetFilterType InItem) const override;
+
+protected:
+	static FString ConvertOperationToString(ETextFilterComparisonOperation Op);
+	
+	void SetComparisonOperation(ETextFilterComparisonOperation NewOp);
+	bool IsComparisonOperationEqualTo(ETextFilterComparisonOperation TestOp) const;
+
+	FText GetKeyValueAsText() const;
+	FText GetTargetValueAsText() const;
+	void OnKeyValueTextCommitted(const FText& InText, ETextCommit::Type InCommitType);
+	void OnTargetValueTextCommitted(const FText& InText, ETextCommit::Type InCommitType);
+
+public:
+	FName TagName;
+	FString TargetTagValue;
+	ETextFilterComparisonOperation ComparisonOp;
+};
+
 /** An inverse filter that allows display of content in developer folders that are not the current user's */
 class FFrontendFilter_ShowOtherDevelopers : public FFrontendFilter
 {
diff --git a/Engine/Source/Editor/ContentBrowser/Private/SFilterList.cpp b/Engine/Source/Editor/ContentBrowser/Private/SFilterList.cpp
--- a/Engine/Source/Editor/ContentBrowser/Private/SFilterList.cpp
+++ b/Engine/Source/Editor/ContentBrowser/Private/SFilterList.cpp
@@ -98,7 +98,7 @@ public:
 		FrontendFilter = InArgs._FrontendFilter;
 
 		// Get the tooltip and color of the type represented by this filter
-		FText FilterToolTip;
+		TAttribute<FText> FilterToolTip;
 		FilterColor = FLinearColor::White;
 		if ( InArgs._AssetTypeActions.IsValid() )
 		{
@@ -107,10 +107,10 @@ public:
 
 			// No tooltip for asset type filters
 		}
-		else if ( InArgs._FrontendFilter.IsValid() )
+		else if ( FrontendFilter.IsValid() )
 		{
 			FilterColor = FrontendFilter->GetColor();
-			FilterToolTip = FrontendFilter->GetToolTipText();
+			FilterToolTip = TAttribute<FText>::Create(TAttribute<FText>::FGetter::CreateSP(FrontendFilter.ToSharedRef(), &FFrontendFilter::GetToolTipText));
 		}
 
 		ChildSlot
@@ -133,7 +133,7 @@ public:
 					.ColorAndOpacity(this, &SFilter::GetFilterNameColorAndOpacity)
 					.Font(FEditorStyle::GetFontStyle("ContentBrowser.FilterNameFont"))
 					.ShadowOffset(FVector2D(1.f, 1.f))
-					.Text( GetFilterName() )
+					.Text(this, &SFilter::GetFilterName)
 				]
 			]
 		];
@@ -253,6 +253,11 @@ private:
 		}
 		MenuBuilder.EndSection();
 
+		if (FrontendFilter.IsValid())
+		{
+			FrontendFilter->ModifyContextMenu(MenuBuilder);
+		}
+
 		return MenuBuilder.MakeWidget();
 	}
 
@@ -383,6 +388,7 @@ void SFilterList::Construct( const FArguments& InArgs )
 	AllFrontendFilters.Add( MakeShareable(new FFrontendFilter_ReplicatedBlueprint(DefaultCategory)) );
 	AllFrontendFilters.Add( MakeShareable(new FFrontendFilter_ShowRedirectors(DefaultCategory)) );
 	AllFrontendFilters.Add( MakeShareable(new FFrontendFilter_InUseByLoadedLevels(DefaultCategory)) );
+	AllFrontendFilters.Add( MakeShareable(new FFrontendFilter_ArbitraryComparisonOperation(DefaultCategory)) ); //@TODO: Needs to be extensible
 
 	for(auto Iter = InArgs._ExtraFrontendFilters.CreateConstIterator(); Iter; ++Iter)
 	{
@@ -639,6 +645,8 @@ void SFilterList::SaveSettings(const FString& IniFilename, const FString& IniSec
 	GConfig->SetString(*IniSection, *(SettingsString + TEXT(".EnabledTypeFilters")), *EnabledTypeFilterString, IniFilename);
 	GConfig->SetString(*IniSection, *(SettingsString + TEXT(".ActiveFrontendFilters")), *ActiveFrontendFilterString, IniFilename);
 	GConfig->SetString(*IniSection, *(SettingsString + TEXT(".EnabledFrontendFilters")), *EnabledFrontendFilterString, IniFilename);
+
+	//@TODO: Allow filters to save information out here (and read it back in at the bottom of LoadSettings)
 }
 
 void SFilterList::LoadSettings(const FString& IniFilename, const FString& IniSection, const FString& SettingsString)
diff --git a/Engine/Source/Editor/ContentBrowser/Public/FrontendFilterBase.h b/Engine/Source/Editor/ContentBrowser/Public/FrontendFilterBase.h
--- a/Engine/Source/Editor/ContentBrowser/Public/FrontendFilterBase.h
+++ b/Engine/Source/Editor/ContentBrowser/Public/FrontendFilterBase.h
@@ -43,6 +43,9 @@ public:
 	/** Notification that the filter became active or inactive */
 	virtual void ActiveStateChanged(bool bActive) { }
 
+	/** Called when the right-click context menu is being built for this filter */
+	virtual void ModifyContextMenu(FMenuBuilder& MenuBuilder) { }
+
 	// IFilter implementation
 	DECLARE_DERIVED_EVENT( FFrontendFilter, IFilter<FAssetFilterType>::FChangedEvent, FChangedEvent );
 	virtual FChangedEvent& OnChanged() override { return ChangedEvent; }
diff --git a/Engine/Source/Runtime/Core/Public/Misc/TextFilter.h b/Engine/Source/Runtime/Core/Public/Misc/TextFilter.h
--- a/Engine/Source/Runtime/Core/Public/Misc/TextFilter.h
+++ b/Engine/Source/Runtime/Core/Public/Misc/TextFilter.h
@@ -3,6 +3,16 @@
 
 #include "IFilter.h"
 
+UENUM()
+enum class ETextFilterComparisonOperation : uint8
+{
+	Equal,
+	NotEqual,
+	Less,
+	LessOrEqual,
+	Greater,
+	GreaterOrEqual,
+};
 
 /**
  *	A generic filter specialized for text restrictions
