Content Browser: Added an extension point to make it easy to define new front-end filters and have them available in the Content Browser
- Create a subclass of UContentBrowserFrontEndFilterExtension and override AddFrontEndFilterExtensions
- Added an example usage to filter down to odd length names in QAGame

diff --git a/Engine/Source/Editor/ContentBrowser/Private/SFilterList.cpp b/Engine/Source/Editor/ContentBrowser/Private/SFilterList.cpp
--- a/Engine/Source/Editor/ContentBrowser/Private/SFilterList.cpp
+++ b/Engine/Source/Editor/ContentBrowser/Private/SFilterList.cpp
@@ -3,6 +3,7 @@
 
 #include "ContentBrowserPCH.h"
 #include "FrontendFilters.h"
+#include "ContentBrowserFrontEndFilterExtension.h"
 
 #define LOCTEXT_NAMESPACE "ContentBrowser"
 
@@ -381,20 +382,46 @@ void SFilterList::Construct( const FArguments& InArgs )
 
 	TSharedPtr<FFrontendFilterCategory> DefaultCategory = MakeShareable( new FFrontendFilterCategory(LOCTEXT("FrontendFiltersCategory", "Other Filters"), LOCTEXT("FrontendFiltersCategoryTooltip", "Filter assets by all filters in this category.")) );
 
-	// Add all frontend filters here
+	// Add all built-in frontend filters here
 	AllFrontendFilters.Add( MakeShareable(new FFrontendFilter_CheckedOut(DefaultCategory)) );
 	AllFrontendFilters.Add( MakeShareable(new FFrontendFilter_Modified(DefaultCategory)) );
 	AllFrontendFilters.Add( MakeShareable(new FFrontendFilter_ShowOtherDevelopers(DefaultCategory)) );
 	AllFrontendFilters.Add( MakeShareable(new FFrontendFilter_ReplicatedBlueprint(DefaultCategory)) );
 	AllFrontendFilters.Add( MakeShareable(new FFrontendFilter_ShowRedirectors(DefaultCategory)) );
 	AllFrontendFilters.Add( MakeShareable(new FFrontendFilter_InUseByLoadedLevels(DefaultCategory)) );
-	AllFrontendFilters.Add( MakeShareable(new FFrontendFilter_ArbitraryComparisonOperation(DefaultCategory)) ); //@TODO: Needs to be extensible
+	AllFrontendFilters.Add( MakeShareable(new FFrontendFilter_ArbitraryComparisonOperation(DefaultCategory)) );
 
-	for(auto Iter = InArgs._ExtraFrontendFilters.CreateConstIterator(); Iter; ++Iter)
+	// Add any global user-defined frontend filters
+	for (TObjectIterator<UContentBrowserFrontEndFilterExtension> ExtensionIt(RF_NoFlags); ExtensionIt; ++ExtensionIt)
+	{
+		if (UContentBrowserFrontEndFilterExtension* PotentialExtension = *ExtensionIt)
+		{
+			if (PotentialExtension->HasAnyFlags(RF_ClassDefaultObject) && !PotentialExtension->GetClass()->HasAnyCastFlag(CLASS_Deprecated | CLASS_Abstract))
+			{
+				// Grab the filters
+				TArray< TSharedRef<FFrontendFilter> > ExtendedFrontendFilters;
+				PotentialExtension->AddFrontEndFilterExtensions(DefaultCategory, ExtendedFrontendFilters);
+				AllFrontendFilters.Append(ExtendedFrontendFilters);
+
+				// Grab the categories
+				for (const TSharedRef<FFrontendFilter>& FilterRef : ExtendedFrontendFilters)
+				{
+					TSharedPtr<FFrontendFilterCategory> Category = FilterRef->GetCategory();
+					if (Category.IsValid())
+					{
+						AllFrontendFilterCategories.AddUnique(Category);
+					}
+				}
+			}
+		}
+	}
+
+	// Add in filters specific to this invocation
+	for (auto Iter = InArgs._ExtraFrontendFilters.CreateConstIterator(); Iter; ++Iter)
 	{
 		TSharedRef<FFrontendFilter> Filter = (*Iter);
 		TSharedPtr<FFrontendFilterCategory> Category = Filter->GetCategory();
-		if ( Category.IsValid() )
+		if (Category.IsValid())
 		{
 			AllFrontendFilterCategories.AddUnique( Category );
 		}
@@ -405,15 +432,17 @@ void SFilterList::Construct( const FArguments& InArgs )
 	AllFrontendFilterCategories.Add(DefaultCategory);
 
 	// Auto add all inverse filters
-	for ( auto FilterIt = AllFrontendFilters.CreateConstIterator(); FilterIt; ++FilterIt )
+	for (auto FilterIt = AllFrontendFilters.CreateConstIterator(); FilterIt; ++FilterIt)
 	{
 		SetFrontendFilterActive(*FilterIt, false);
 	}
 
+	FilterBox = SNew(SWrapBox)
+		.UseAllottedWidth(true);
+
 	ChildSlot
 	[
-		SAssignNew(FilterBox, SWrapBox)
-		.UseAllottedWidth(true)
+		FilterBox.ToSharedRef()
 	];
 }
 
diff --git a/Engine/Source/Editor/ContentBrowser/Public/ContentBrowserFrontEndFilterExtension.h b/Engine/Source/Editor/ContentBrowser/Public/ContentBrowserFrontEndFilterExtension.h
new file mode 100644
--- /dev/null
+++ b/Engine/Source/Editor/ContentBrowser/Public/ContentBrowserFrontEndFilterExtension.h
@@ -0,0 +1,16 @@
+// Copyright 1998-2015 Epic Games, Inc. All Rights Reserved.
+
+#pragma once
+
+#include "ContentBrowserFrontEndFilterExtension.generated.h"
+
+// Override this class in order to make an additional front-end filter available in the Content Browser
+UCLASS(Abstract)
+class CONTENTBROWSER_API UContentBrowserFrontEndFilterExtension : public UObject
+{
+	GENERATED_BODY()
+
+public:
+	// Override this method to add new frontend filters
+	virtual void AddFrontEndFilterExtensions(TSharedPtr<class FFrontendFilterCategory> DefaultCategory, TArray< TSharedRef<class FFrontendFilter> >& InOutFilterList) const { }
+};
