Content Browser: Allowed frontend filters to save unique per-filter state

diff --git a/Engine/Source/Editor/ContentBrowser/Private/FrontendFilters.cpp b/Engine/Source/Editor/ContentBrowser/Private/FrontendFilters.cpp
--- a/Engine/Source/Editor/ContentBrowser/Private/FrontendFilters.cpp
+++ b/Engine/Source/Editor/ContentBrowser/Private/FrontendFilters.cpp
@@ -291,6 +291,30 @@ void FFrontendFilter_ArbitraryComparisonOperation::ModifyContextMenu(FMenuBuilde
 	MenuBuilder.EndSection();
 }
 
+void FFrontendFilter_ArbitraryComparisonOperation::SaveSettings(const FString& IniFilename, const FString& IniSection, const FString& SettingsString) const
+{
+	GConfig->SetString(*IniSection, *(SettingsString + TEXT(".Key")), *TagName.ToString(), IniFilename);
+	GConfig->SetString(*IniSection, *(SettingsString + TEXT(".Value")), *TargetTagValue, IniFilename);
+	GConfig->SetString(*IniSection, *(SettingsString + TEXT(".Op")), *FString::FromInt((int32)ComparisonOp), IniFilename);
+}
+
+void FFrontendFilter_ArbitraryComparisonOperation::LoadSettings(const FString& IniFilename, const FString& IniSection, const FString& SettingsString)
+{
+	FString TagNameAsString;
+	if (GConfig->GetString(*IniSection, *(SettingsString + TEXT(".Key")), TagNameAsString, IniFilename))
+	{
+		TagName = *TagNameAsString;
+	}
+
+	GConfig->GetString(*IniSection, *(SettingsString + TEXT(".Value")), TargetTagValue, IniFilename);
+
+	int32 OpAsInteger;
+	if (GConfig->GetInt(*IniSection, *(SettingsString + TEXT(".Op")), OpAsInteger, IniFilename))
+	{
+		ComparisonOp = (ETextFilterComparisonOperation)OpAsInteger;
+	}
+}
+
 void FFrontendFilter_ArbitraryComparisonOperation::SetComparisonOperation(ETextFilterComparisonOperation NewOp)
 {
 	ComparisonOp = NewOp;
diff --git a/Engine/Source/Editor/ContentBrowser/Private/FrontendFilters.h b/Engine/Source/Editor/ContentBrowser/Private/FrontendFilters.h
--- a/Engine/Source/Editor/ContentBrowser/Private/FrontendFilters.h
+++ b/Engine/Source/Editor/ContentBrowser/Private/FrontendFilters.h
@@ -113,6 +113,8 @@ public:
 	virtual FText GetDisplayName() const override;
 	virtual FText GetToolTipText() const override;
 	virtual void ModifyContextMenu(FMenuBuilder& MenuBuilder) override;
+	virtual void SaveSettings(const FString& IniFilename, const FString& IniSection, const FString& SettingsString) const override;
+	virtual void LoadSettings(const FString& IniFilename, const FString& IniSection, const FString& SettingsString) override;
 
 	// IFilter implementation
 	virtual bool PassesFilter(FAssetFilterType InItem) const override;
diff --git a/Engine/Source/Editor/ContentBrowser/Private/SFilterList.cpp b/Engine/Source/Editor/ContentBrowser/Private/SFilterList.cpp
--- a/Engine/Source/Editor/ContentBrowser/Private/SFilterList.cpp
+++ b/Engine/Source/Editor/ContentBrowser/Private/SFilterList.cpp
@@ -650,12 +650,13 @@ void SFilterList::SaveSettings(const FString& IniFilename, const FString& IniSec
 		}
 		else if ( Filter->GetFrontendFilter().IsValid() )
 		{
+			const TSharedPtr<FFrontendFilter>& FrontendFilter = Filter->GetFrontendFilter();
 			if ( ActiveFrontendFilterString.Len() > 0 )
 			{
 				ActiveFrontendFilterString += TEXT(",");
 			}
 
-			const FString FilterName = Filter->GetFrontendFilter()->GetName();
+			const FString FilterName = FrontendFilter->GetName();
 			ActiveFrontendFilterString += FilterName;
 
 			if ( Filter->IsEnabled() )
@@ -667,6 +668,9 @@ void SFilterList::SaveSettings(const FString& IniFilename, const FString& IniSec
 
 				EnabledFrontendFilterString += FilterName;
 			}
+
+			const FString CustomSettingsString = FString::Printf(TEXT("%s.CustomSettings.%s"), *SettingsString, *FilterName);
+			FrontendFilter->SaveSettings(IniFilename, IniSection, CustomSettingsString);
 		}
 	}
 
@@ -674,8 +678,6 @@ void SFilterList::SaveSettings(const FString& IniFilename, const FString& IniSec
 	GConfig->SetString(*IniSection, *(SettingsString + TEXT(".EnabledTypeFilters")), *EnabledTypeFilterString, IniFilename);
 	GConfig->SetString(*IniSection, *(SettingsString + TEXT(".ActiveFrontendFilters")), *ActiveFrontendFilterString, IniFilename);
 	GConfig->SetString(*IniSection, *(SettingsString + TEXT(".EnabledFrontendFilters")), *EnabledFrontendFilterString, IniFilename);
-
-	//@TODO: Allow filters to save information out here (and read it back in at the bottom of LoadSettings)
 }
 
 void SFilterList::LoadSettings(const FString& IniFilename, const FString& IniSection, const FString& SettingsString)
@@ -732,12 +734,12 @@ void SFilterList::LoadSettings(const FString& IniFilename, const FString& IniSec
 		EnabledFrontendFilterString.ParseIntoArray(EnabledFrontendFilterNames, TEXT(","), /*bCullEmpty=*/true);
 
 		// For each FrontendFilter, add any that were active and enable any that were previously enabled
-		for ( auto FrontendFilterIt = AllFrontendFilters.CreateConstIterator(); FrontendFilterIt; ++FrontendFilterIt )
+		for ( auto FrontendFilterIt = AllFrontendFilters.CreateIterator(); FrontendFilterIt; ++FrontendFilterIt )
 		{
-			const TSharedRef<FFrontendFilter>& FrontendFilter = *FrontendFilterIt;
-			if ( !IsFrontendFilterInUse(FrontendFilter) )
+			TSharedRef<FFrontendFilter>& FrontendFilter = *FrontendFilterIt;
+			const FString& FilterName = FrontendFilter->GetName();
+			if (!IsFrontendFilterInUse(FrontendFilter))
 			{
-				const FString& FilterName = FrontendFilter->GetName();
 				if ( FrontendFilterNames.Contains(FilterName) )
 				{
 					TSharedRef<SFilter> NewFilter = AddFilter(FrontendFilter);
@@ -748,6 +750,9 @@ void SFilterList::LoadSettings(const FString& IniFilename, const FString& IniSec
 					}
 				}
 			}
+
+			const FString CustomSettingsString = FString::Printf(TEXT("%s.CustomSettings.%s"), *SettingsString, *FilterName);
+			FrontendFilter->LoadSettings(IniFilename, IniSection, CustomSettingsString);
 		}
 	}
 }
diff --git a/Engine/Source/Editor/ContentBrowser/Public/FrontendFilterBase.h b/Engine/Source/Editor/ContentBrowser/Public/FrontendFilterBase.h
--- a/Engine/Source/Editor/ContentBrowser/Public/FrontendFilterBase.h
+++ b/Engine/Source/Editor/ContentBrowser/Public/FrontendFilterBase.h
@@ -46,6 +46,12 @@ public:
 	/** Called when the right-click context menu is being built for this filter */
 	virtual void ModifyContextMenu(FMenuBuilder& MenuBuilder) { }
 
+	/** Called when the state of a particular Content Browser is being saved to INI */
+	virtual void SaveSettings(const FString& IniFilename, const FString& IniSection, const FString& SettingsString) const {}
+
+	/** Called when the state of a particular Content Browser is being loaded from INI */
+	virtual void LoadSettings(const FString& IniFilename, const FString& IniSection, const FString& SettingsString) {}
+
 	// IFilter implementation
 	DECLARE_DERIVED_EVENT( FFrontendFilter, IFilter<FAssetFilterType>::FChangedEvent, FChangedEvent );
 	virtual FChangedEvent& OnChanged() override { return ChangedEvent; }
