Content Browser: Changed the "content browser source was changed by the sources view" message to only be displayed in VeryVerbose mode

diff --git a/Engine/Source/Editor/ContentBrowser/Private/SContentBrowser.cpp b/Engine/Source/Editor/ContentBrowser/Private/SContentBrowser.cpp
--- a/Engine/Source/Editor/ContentBrowser/Private/SContentBrowser.cpp
+++ b/Engine/Source/Editor/ContentBrowser/Private/SContentBrowser.cpp
@@ -1053,7 +1053,7 @@ void SContentBrowser::LoadSettings(const FName& InInstanceName)
 void SContentBrowser::SourcesChanged(const TArray<FString>& SelectedPaths, const TArray<FCollectionNameType>& SelectedCollections)
 {
 	FString NewSource = SelectedPaths.Num() > 0 ? SelectedPaths[0] : (SelectedCollections.Num() > 0 ? SelectedCollections[0].Name.ToString() : TEXT("None"));
-	UE_LOG(LogContentBrowser, Verbose, TEXT("The content browser source was changed by the sources view to '%s'"), *NewSource);
+	UE_LOG(LogContentBrowser, VeryVerbose, TEXT("The content browser source was changed by the sources view to '%s'"), *NewSource);
 
 	FSourcesData SourcesData;
 	for (int32 PathIdx = 0; PathIdx < SelectedPaths.Num(); ++PathIdx)
