[UE-8485] Added metadata to folder widgets. This allows tutorials to select them as highlight-able items.

diff --git a/Engine/Source/Editor/ContentBrowser/Private/AssetViewWidgets.cpp b/Engine/Source/Editor/ContentBrowser/Private/AssetViewWidgets.cpp
--- a/Engine/Source/Editor/ContentBrowser/Private/AssetViewWidgets.cpp
+++ b/Engine/Source/Editor/ContentBrowser/Private/AssetViewWidgets.cpp
@@ -1158,7 +1158,7 @@ void SAssetTileItem::Construct( const FArguments& InArgs )
 		SNew(SBorder)
 		.BorderImage(this, &SAssetViewItem::GetBorderImage)
 		.Padding(0)
-		.AddMetaData<FTagMetaData>(FTagMetaData(AssetItem->GetType() == EAssetItemType::Normal ? StaticCastSharedPtr<FAssetViewAsset>(AssetItem)->Data.ObjectPath : NAME_None))
+		.AddMetaData<FTagMetaData>(FTagMetaData((AssetItem->GetType() == EAssetItemType::Normal) ? StaticCastSharedPtr<FAssetViewAsset>(AssetItem)->Data.ObjectPath : ((AssetItem->GetType() == EAssetItemType::Folder) ? FName(*StaticCastSharedPtr<FAssetViewFolder>(AssetItem)->FolderPath) : NAME_None)))
 		[
 			SNew(SVerticalBox)
 
@@ -1170,7 +1170,7 @@ void SAssetTileItem::Construct( const FArguments& InArgs )
 				// The remainder of the space is reserved for the name.
 				SNew(SBox)
 				.Padding(ThumbnailPadding - 4.f)
-				.WidthOverride( this, &SAssetTileItem::GetThumbnailBoxSize )
+				.WidthOverride(this, &SAssetTileItem::GetThumbnailBoxSize)
 				.HeightOverride( this, &SAssetTileItem::GetThumbnailBoxSize )
 				[
 					// Drop shadow border
