Editor: Prevented mixed icon usage (some items in a section but not all using an icon) in the content browser context menus by removing the icons

diff --git a/Engine/Source/Editor/ContentBrowser/Private/AssetContextMenu.cpp b/Engine/Source/Editor/ContentBrowser/Private/AssetContextMenu.cpp
--- a/Engine/Source/Editor/ContentBrowser/Private/AssetContextMenu.cpp
+++ b/Engine/Source/Editor/ContentBrowser/Private/AssetContextMenu.cpp
@@ -158,7 +158,7 @@ bool FAssetContextMenu::AddImportedAssetMenuOptions(FMenuBuilder& MenuBuilder)
 			MenuBuilder.AddMenuEntry(
 				LOCTEXT("Reimport", "Reimport"),
 				LOCTEXT("ReimportTooltip", "Reimport the selected asset(s) from the source file on disk."),
-				FSlateIcon(FEditorStyle::GetStyleSetName(), "ContentBrowser.AssetActions.ReimportAsset"),
+				FSlateIcon(),
 				FUIAction(
 					FExecuteAction::CreateSP(this, &FAssetContextMenu::ExecuteReimport),
 					FCanExecuteAction()
@@ -204,7 +204,7 @@ bool FAssetContextMenu::AddCommonMenuOptions(FMenuBuilder& MenuBuilder)
 	MenuBuilder.AddMenuEntry(
 		LOCTEXT("EditAsset", "Edit..."),
 		LOCTEXT("EditAssetTooltip", "Opens the selected asset(s) for edit."),
-		FSlateIcon(FEditorStyle::GetStyleSetName(), "ContentBrowser.AssetActions.Edit"),
+		FSlateIcon(),
 		FUIAction( FExecuteAction::CreateSP(this, &FAssetContextMenu::ExecuteEditAsset) )
 		);
 	
@@ -212,21 +212,21 @@ bool FAssetContextMenu::AddCommonMenuOptions(FMenuBuilder& MenuBuilder)
 	MenuBuilder.AddMenuEntry(FGenericCommands::Get().Rename, NAME_None,
 		LOCTEXT("Rename", "Rename"),
 		LOCTEXT("RenameTooltip", "Rename the selected asset."),
-		FSlateIcon(FEditorStyle::GetStyleSetName(), "ContentBrowser.AssetActions.Rename")
+		FSlateIcon()
 		);
 
 	// Duplicate
 	MenuBuilder.AddMenuEntry(FGenericCommands::Get().Duplicate, NAME_None,
 		LOCTEXT("Duplicate", "Duplicate"),
 		LOCTEXT("DuplicateTooltip", "Create a copy of the selected asset(s)."),
-		FSlateIcon(FEditorStyle::GetStyleSetName(), "ContentBrowser.AssetActions.Duplicate")
+		FSlateIcon()
 		);
 
 	// Save
 	MenuBuilder.AddMenuEntry(
 		LOCTEXT("SaveAsset", "Save"),
 		LOCTEXT("SaveAssetTooltip", "Saves the asset to file."),
-		FSlateIcon(FEditorStyle::GetStyleSetName(), "Level.SaveIcon16x"),
+		FSlateIcon(),
 		FUIAction(
 			FExecuteAction::CreateSP( this, &FAssetContextMenu::ExecuteSaveAsset ),
 			FCanExecuteAction::CreateSP( this, &FAssetContextMenu::CanExecuteSaveAsset )
@@ -237,7 +237,7 @@ bool FAssetContextMenu::AddCommonMenuOptions(FMenuBuilder& MenuBuilder)
 	MenuBuilder.AddMenuEntry(FGenericCommands::Get().Delete, NAME_None,
 		LOCTEXT("Delete", "Delete"),
 		LOCTEXT("DeleteTooltip", "Delete the selected assets."),
-		FSlateIcon(FEditorStyle::GetStyleSetName(), "ContentBrowser.AssetActions.Delete")
+		FSlateIcon()
 		);
 
 	// Asset Actions sub-menu
@@ -252,7 +252,7 @@ bool FAssetContextMenu::AddCommonMenuOptions(FMenuBuilder& MenuBuilder)
 		NAME_None,
 		EUserInterfaceActionType::Button,
 		false, 
-		FSlateIcon(FEditorStyle::GetStyleSetName(), "ContentBrowser.AssetActions")
+		FSlateIcon()
 		);
 	}
 	MenuBuilder.EndSection();
@@ -684,7 +684,7 @@ void FAssetContextMenu::FillSourceControlSubMenu(FMenuBuilder& MenuBuilder)
 		MenuBuilder.AddMenuEntry(
 			LOCTEXT("SCCCheckOut", "Check Out"),
 			LOCTEXT("SCCCheckOutTooltip", "Checks out the selected asset from source control."),
-			FSlateIcon(FEditorStyle::GetStyleSetName(), "SourceControl.Actions.CheckOut"),
+			FSlateIcon(),
 			FUIAction(
 				FExecuteAction::CreateSP( this, &FAssetContextMenu::ExecuteSCCCheckOut ),
 				FCanExecuteAction::CreateSP( this, &FAssetContextMenu::CanExecuteSCCCheckOut )
@@ -697,7 +697,7 @@ void FAssetContextMenu::FillSourceControlSubMenu(FMenuBuilder& MenuBuilder)
 		MenuBuilder.AddMenuEntry(
 			LOCTEXT("SCCOpenForAdd", "Mark For Add"),
 			LOCTEXT("SCCOpenForAddTooltip", "Adds the selected asset to source control."),
-			FSlateIcon(FEditorStyle::GetStyleSetName(), "SourceControl.Actions.Add"),
+			FSlateIcon(),
 			FUIAction(
 				FExecuteAction::CreateSP( this, &FAssetContextMenu::ExecuteSCCOpenForAdd ),
 				FCanExecuteAction::CreateSP( this, &FAssetContextMenu::CanExecuteSCCOpenForAdd )
@@ -721,7 +721,7 @@ void FAssetContextMenu::FillSourceControlSubMenu(FMenuBuilder& MenuBuilder)
 	MenuBuilder.AddMenuEntry(
 		LOCTEXT("SCCRefresh", "Refresh"),
 		LOCTEXT("SCCRefreshTooltip", "Updates the source control status of the asset."),
-		FSlateIcon(FEditorStyle::GetStyleSetName(), "SourceControl.Actions.Refresh"),
+		FSlateIcon(),
 		FUIAction(
 			FExecuteAction::CreateSP( this, &FAssetContextMenu::ExecuteSCCRefresh ),
 			FCanExecuteAction::CreateSP( this, &FAssetContextMenu::CanExecuteSCCRefresh )
diff --git a/Engine/Source/Editor/ContentBrowser/Private/PathContextMenu.cpp b/Engine/Source/Editor/ContentBrowser/Private/PathContextMenu.cpp
--- a/Engine/Source/Editor/ContentBrowser/Private/PathContextMenu.cpp
+++ b/Engine/Source/Editor/ContentBrowser/Private/PathContextMenu.cpp
@@ -126,7 +126,7 @@ void FPathContextMenu::MakePathViewContextMenu(FMenuBuilder& MenuBuilder)
 					NAME_None,
 					EUserInterfaceActionType::Button,
 					false,
-					FSlateIcon( FEditorStyle::GetStyleSetName(), "ContentBrowser.PathActions.NewAsset" )
+					FSlateIcon()
 					);
 			}
 
@@ -183,7 +183,7 @@ void FPathContextMenu::MakePathViewContextMenu(FMenuBuilder& MenuBuilder)
 					LOCTEXT("SetColorTooltip", "Sets the color this folder should appear as."),
 					FNewMenuDelegate::CreateRaw( this, &FPathContextMenu::MakeSetColorSubMenu ),
 					false,
-					FSlateIcon( FEditorStyle::GetStyleSetName(), "ContentBrowser.PathActions.SetColor" )
+					FSlateIcon()
 					);
 			}
 			else
@@ -192,7 +192,7 @@ void FPathContextMenu::MakePathViewContextMenu(FMenuBuilder& MenuBuilder)
 				MenuBuilder.AddMenuEntry(
 					LOCTEXT("SetColor", "Set Color"),
 					LOCTEXT("SetColorTooltip", "Sets the color this folder should appear as."),
-					FSlateIcon(FEditorStyle::GetStyleSetName(), "ContentBrowser.PathActions.SetColor"),
+					FSlateIcon(),
 					FUIAction( FExecuteAction::CreateSP( this, &FPathContextMenu::ExecutePickColor ) )
 					);
 			}			
@@ -208,7 +208,7 @@ void FPathContextMenu::MakePathViewContextMenu(FMenuBuilder& MenuBuilder)
 				MenuBuilder.AddMenuEntry(
 					LOCTEXT("SaveFolder", "Save All"),
 					LOCTEXT("SaveFolderTooltip", "Saves all modified assets in this folder."),
-					FSlateIcon(FEditorStyle::GetStyleSetName(), "MainFrame.SaveAll"),
+					FSlateIcon(),
 					FUIAction( FExecuteAction::CreateSP( this, &FPathContextMenu::ExecuteSaveFolder ) )
 					);
     
@@ -216,7 +216,7 @@ void FPathContextMenu::MakePathViewContextMenu(FMenuBuilder& MenuBuilder)
 				MenuBuilder.AddMenuEntry(FGenericCommands::Get().Delete, NAME_None,
 					LOCTEXT("DeleteFolder", "Delete"),
 					LOCTEXT("DeleteFolderTooltip", "Removes this folder and all assets it contains."),
-					FSlateIcon( FEditorStyle::GetStyleSetName(), "ContentBrowser.AssetActions.Delete" )
+					FSlateIcon()
 					);
 
 				// Reference Viewer
