UE4 - Do not allow Materials to be edited using property matrix (Fixes UE-3326).
Materials don't have a good pattern/paradigm for updating 'anywhere' as they rely on talking to and flushing the rendering thread.

diff --git a/Engine/Source/Editor/ContentBrowser/Private/AssetContextMenu.cpp b/Engine/Source/Editor/ContentBrowser/Private/AssetContextMenu.cpp
--- a/Engine/Source/Editor/ContentBrowser/Private/AssetContextMenu.cpp
+++ b/Engine/Source/Editor/ContentBrowser/Private/AssetContextMenu.cpp
@@ -387,15 +387,30 @@ void FAssetContextMenu::MakeAssetActionsSubMenu(FMenuBuilder& MenuBuilder)
 		}
 
 		// Property Matrix
-		MenuBuilder.AddMenuEntry(
-			LOCTEXT("PropertyMatrix", "Bulk Edit via Property Matrix..."),
-			LOCTEXT("PropertyMatrixTooltip", "Opens the property matrix editor for the selected assets."),
-			FSlateIcon(),
-			FUIAction(
-			FExecuteAction::CreateSP(this, &FAssetContextMenu::ExecutePropertyMatrix),
-			FCanExecuteAction::CreateSP(this, &FAssetContextMenu::CanExecuteProperties)
-			)
-			);
+		bool bCanUsePropertyMatrix = true;
+		static FName MaterialName(TEXT("Material"));
+		static FName MaterialInstanceConstantName(TEXT("MaterialInstanceConstant"));
+		for (auto& Asset : SelectedAssets)
+		{
+			if (Asset.AssetClass == MaterialName || Asset.AssetClass == MaterialInstanceName)
+			{
+				bCanUsePropertyMatrix = false;
+				break;
+			}
+		}
+
+		if (bCanUsePropertyMatrix)
+		{
+			MenuBuilder.AddMenuEntry(
+				LOCTEXT("PropertyMatrix", "Bulk Edit via Property Matrix..."),
+				LOCTEXT("PropertyMatrixTooltip", "Opens the property matrix editor for the selected assets."),
+				FSlateIcon(),
+				FUIAction(
+				FExecuteAction::CreateSP(this, &FAssetContextMenu::ExecutePropertyMatrix),
+				FCanExecuteAction::CreateSP(this, &FAssetContextMenu::CanExecuteProperties)
+				)
+				);
+		}
 
 		// Chunk actions
 		if (GetDefault<UEditorExperimentalSettings>()->bContextMenuChunkAssignments)
