Fixed MaterialInstanceConstantName compile error

diff --git a/Engine/Source/Editor/ContentBrowser/Private/AssetContextMenu.cpp b/Engine/Source/Editor/ContentBrowser/Private/AssetContextMenu.cpp
--- a/Engine/Source/Editor/ContentBrowser/Private/AssetContextMenu.cpp
+++ b/Engine/Source/Editor/ContentBrowser/Private/AssetContextMenu.cpp
@@ -392,7 +392,7 @@ void FAssetContextMenu::MakeAssetActionsSubMenu(FMenuBuilder& MenuBuilder)
 		static FName MaterialInstanceConstantName(TEXT("MaterialInstanceConstant"));
 		for (auto& Asset : SelectedAssets)
 		{
-			if (Asset.AssetClass == MaterialName || Asset.AssetClass == MaterialInstanceName)
+			if (Asset.AssetClass == MaterialName || Asset.AssetClass == MaterialInstanceConstantName)
 			{
 				bCanUsePropertyMatrix = false;
 				break;
