Move some asset context menu entries around

Show in Folder and Show in Explorer are now in the top level menu

diff --git a/Engine/Source/Editor/ContentBrowser/Private/AssetContextMenu.cpp b/Engine/Source/Editor/ContentBrowser/Private/AssetContextMenu.cpp
--- a/Engine/Source/Editor/ContentBrowser/Private/AssetContextMenu.cpp
+++ b/Engine/Source/Editor/ContentBrowser/Private/AssetContextMenu.cpp
@@ -100,6 +100,9 @@ TSharedRef<SWidget> FAssetContextMenu::MakeContextMenu(const TArray<FAssetData>&
 		// Add quick access to common commands.
 		AddCommonMenuOptions(MenuBuilder);
 
+		// Add quick access to view commands
+		AddExploreMenuOptions(MenuBuilder);
+
 		// Add reference options
 		AddReferenceMenuOptions(MenuBuilder);
 
@@ -201,66 +204,107 @@ bool FAssetContextMenu::AddCommonMenuOptions(FMenuBuilder& MenuBuilder)
 {
 	MenuBuilder.BeginSection("CommonAssetActions", LOCTEXT("CommonAssetActionsMenuHeading", "Common"));
 	{
-	// Edit
-	MenuBuilder.AddMenuEntry(
-		LOCTEXT("EditAsset", "Edit..."),
-		LOCTEXT("EditAssetTooltip", "Opens the selected asset(s) for edit."),
-		FSlateIcon(),
-		FUIAction( FExecuteAction::CreateSP(this, &FAssetContextMenu::ExecuteEditAsset) )
-		);
+		// Edit
+		MenuBuilder.AddMenuEntry(
+			LOCTEXT("EditAsset", "Edit..."),
+			LOCTEXT("EditAssetTooltip", "Opens the selected asset(s) for edit."),
+			FSlateIcon(),
+			FUIAction( FExecuteAction::CreateSP(this, &FAssetContextMenu::ExecuteEditAsset) )
+			);
 	
-	// Rename
-	MenuBuilder.AddMenuEntry(FGenericCommands::Get().Rename, NAME_None,
-		LOCTEXT("Rename", "Rename"),
-		LOCTEXT("RenameTooltip", "Rename the selected asset."),
-		FSlateIcon()
-		);
+		// Rename
+		MenuBuilder.AddMenuEntry(FGenericCommands::Get().Rename, NAME_None,
+			LOCTEXT("Rename", "Rename"),
+			LOCTEXT("RenameTooltip", "Rename the selected asset."),
+			FSlateIcon()
+			);
 
-	// Duplicate
-	MenuBuilder.AddMenuEntry(FGenericCommands::Get().Duplicate, NAME_None,
-		LOCTEXT("Duplicate", "Duplicate"),
-		LOCTEXT("DuplicateTooltip", "Create a copy of the selected asset(s)."),
-		FSlateIcon()
-		);
+		// Duplicate
+		MenuBuilder.AddMenuEntry(FGenericCommands::Get().Duplicate, NAME_None,
+			LOCTEXT("Duplicate", "Duplicate"),
+			LOCTEXT("DuplicateTooltip", "Create a copy of the selected asset(s)."),
+			FSlateIcon()
+			);
 
-	// Save
-	MenuBuilder.AddMenuEntry(
-		LOCTEXT("SaveAsset", "Save"),
-		LOCTEXT("SaveAssetTooltip", "Saves the asset to file."),
-		FSlateIcon(),
-		FUIAction(
-			FExecuteAction::CreateSP( this, &FAssetContextMenu::ExecuteSaveAsset ),
-			FCanExecuteAction::CreateSP( this, &FAssetContextMenu::CanExecuteSaveAsset )
-			)
-		);
+		// Save
+		MenuBuilder.AddMenuEntry(
+			LOCTEXT("SaveAsset", "Save"),
+			LOCTEXT("SaveAssetTooltip", "Saves the asset to file."),
+			FSlateIcon(),
+			FUIAction(
+				FExecuteAction::CreateSP( this, &FAssetContextMenu::ExecuteSaveAsset ),
+				FCanExecuteAction::CreateSP( this, &FAssetContextMenu::CanExecuteSaveAsset )
+				)
+			);
 
-	// Delete
-	MenuBuilder.AddMenuEntry(FGenericCommands::Get().Delete, NAME_None,
-		LOCTEXT("Delete", "Delete"),
-		LOCTEXT("DeleteTooltip", "Delete the selected assets."),
-		FSlateIcon()
-		);
+		// Delete
+		MenuBuilder.AddMenuEntry(FGenericCommands::Get().Delete, NAME_None,
+			LOCTEXT("Delete", "Delete"),
+			LOCTEXT("DeleteTooltip", "Delete the selected assets."),
+			FSlateIcon()
+			);
 
-	// Asset Actions sub-menu
-	MenuBuilder.AddSubMenu(
-		LOCTEXT("AssetActionsSubMenuLabel", "Asset Actions"),
-		LOCTEXT("AssetActionsSubMenuToolTip", "Other asset actions"),
-		FNewMenuDelegate::CreateSP(this, &FAssetContextMenu::MakeAssetActionsSubMenu),
-		FUIAction(
-			FExecuteAction(),
-			FCanExecuteAction::CreateSP( this, &FAssetContextMenu::CanExecuteAssetActions )
-			),
-		NAME_None,
-		EUserInterfaceActionType::Button,
-		false, 
-		FSlateIcon()
-		);
+			// Open Containing Folder
+			MenuBuilder.AddMenuEntry(
+				LOCTEXT("SyncToAssetTree", "Show in Sources Panel"),
+				LOCTEXT("SyncToAssetTreeTooltip", "Selects the folder that contains this asset in the Content Browser Sources Panel."),
+				FSlateIcon(),
+				FUIAction(
+					FExecuteAction::CreateSP(this, &FAssetContextMenu::ExecuteSyncToAssetTree),
+					FCanExecuteAction::CreateSP(this, &FAssetContextMenu::CanExecuteSyncToAssetTree)
+					)
+				);
+
+		// Asset Actions sub-menu
+		MenuBuilder.AddSubMenu(
+			LOCTEXT("AssetActionsSubMenuLabel", "Asset Actions"),
+			LOCTEXT("AssetActionsSubMenuToolTip", "Other asset actions"),
+			FNewMenuDelegate::CreateSP(this, &FAssetContextMenu::MakeAssetActionsSubMenu),
+			FUIAction(
+				FExecuteAction(),
+				FCanExecuteAction::CreateSP( this, &FAssetContextMenu::CanExecuteAssetActions )
+				),
+			NAME_None,
+			EUserInterfaceActionType::Button,
+			false, 
+			FSlateIcon()
+			);
 	}
 	MenuBuilder.EndSection();
 
 	return true;
 }
 
+void FAssetContextMenu::AddExploreMenuOptions(FMenuBuilder& MenuBuilder)
+{
+	MenuBuilder.BeginSection("AssetContextExploreMenuOptions", LOCTEXT("AssetContextExploreMenuOptionsHeading", "Explore"));
+	{
+		// Open Containing Folder
+		MenuBuilder.AddMenuEntry(
+			LOCTEXT("SyncToAssetTree", "Show in Folder View"),
+			LOCTEXT("SyncToAssetTreeTooltip", "Selects the folder that contains this asset in the Content Browser Sources Panel."),
+			FSlateIcon(),
+			FUIAction(
+				FExecuteAction::CreateSP(this, &FAssetContextMenu::ExecuteSyncToAssetTree),
+				FCanExecuteAction::CreateSP(this, &FAssetContextMenu::CanExecuteSyncToAssetTree)
+				)
+			);
+
+
+		// Find in Explorer
+		MenuBuilder.AddMenuEntry(
+			ContentBrowserUtils::GetExploreFolderText(),
+			LOCTEXT("FindInExplorerTooltip", "Finds this asset on disk"),
+			FSlateIcon(),
+			FUIAction(
+				FExecuteAction::CreateSP( this, &FAssetContextMenu::ExecuteFindInExplorer ),
+				FCanExecuteAction::CreateSP( this, &FAssetContextMenu::CanExecuteFindInExplorer )
+				)
+			);
+	}
+	MenuBuilder.EndSection();
+}
+
 void FAssetContextMenu::MakeAssetActionsSubMenu(FMenuBuilder& MenuBuilder)
 {	
 	// Create BP Using This
@@ -313,28 +357,6 @@ void FAssetContextMenu::MakeAssetActionsSubMenu(FMenuBuilder& MenuBuilder)
 				FCanExecuteAction::CreateSP(this, &FAssetContextMenu::CanExecuteFindAssetInWorld)
 				)
 			);
-
-		// Open Containing Folder
-		MenuBuilder.AddMenuEntry(
-			LOCTEXT("SyncToAssetTree", "Show in Sources Panel"),
-			LOCTEXT("SyncToAssetTreeTooltip", "Selects the folder that contains this asset in the Content Browser Sources Panel."),
-			FSlateIcon(),
-			FUIAction(
-				FExecuteAction::CreateSP(this, &FAssetContextMenu::ExecuteSyncToAssetTree),
-				FCanExecuteAction::CreateSP(this, &FAssetContextMenu::CanExecuteSyncToAssetTree)
-				)
-			);
-
-		// Find in Explorer
-		MenuBuilder.AddMenuEntry(
-			ContentBrowserUtils::GetExploreFolderText(),
-			LOCTEXT("FindInExplorerTooltip", "Finds this asset on disk"),
-			FSlateIcon(),
-			FUIAction(
-				FExecuteAction::CreateSP( this, &FAssetContextMenu::ExecuteFindInExplorer ),
-				FCanExecuteAction::CreateSP( this, &FAssetContextMenu::CanExecuteFindInExplorer )
-				)
-			);
 	}
 	MenuBuilder.EndSection();
 
@@ -446,7 +468,7 @@ bool FAssetContextMenu::CanExecuteAssetActions() const
 
 bool FAssetContextMenu::AddReferenceMenuOptions(FMenuBuilder& MenuBuilder)
 {
-	MenuBuilder.BeginSection("AssetContextReferences"/*, LOCTEXT("ReferencesMenuHeading", "References")*/);
+	MenuBuilder.BeginSection("AssetContextReferences", LOCTEXT("ReferencesMenuHeading", "References"));
 	{
 		MenuBuilder.AddMenuEntry(
 			LOCTEXT("CopyReference", "Copy Reference"),
diff --git a/Engine/Source/Editor/ContentBrowser/Private/AssetContextMenu.h b/Engine/Source/Editor/ContentBrowser/Private/AssetContextMenu.h
--- a/Engine/Source/Editor/ContentBrowser/Private/AssetContextMenu.h
+++ b/Engine/Source/Editor/ContentBrowser/Private/AssetContextMenu.h
@@ -86,6 +86,9 @@ private:
 	/** Adds common menu options to a menu builder. Returns true if any options were added. */
 	bool AddCommonMenuOptions(FMenuBuilder& MenuBuilder);
 
+	/** Adds explore menu options to a menu builder. */
+	void AddExploreMenuOptions(FMenuBuilder& MenuBuilder);
+
 	/** Adds Asset Actions sub-menu to a menu builder. */
 	void MakeAssetActionsSubMenu(FMenuBuilder& MenuBuilder);
 
