UE4 - Disallow material functions also on bulk edit

diff --git a/Engine/Source/Editor/ContentBrowser/Private/AssetContextMenu.cpp b/Engine/Source/Editor/ContentBrowser/Private/AssetContextMenu.cpp
--- a/Engine/Source/Editor/ContentBrowser/Private/AssetContextMenu.cpp
+++ b/Engine/Source/Editor/ContentBrowser/Private/AssetContextMenu.cpp
@@ -410,11 +410,10 @@ void FAssetContextMenu::MakeAssetActionsSubMenu(FMenuBuilder& MenuBuilder)
 
 		// Property Matrix
 		bool bCanUsePropertyMatrix = true;
-		static FName MaterialName(TEXT("Material"));
-		static FName MaterialInstanceConstantName(TEXT("MaterialInstanceConstant"));
+		// Materials can't be bulk edited currently as they require very special handling because of their dependencies with the rendering thread, and we'd have to hack the property matrix too much.
 		for (auto& Asset : SelectedAssets)
 		{
-			if (Asset.AssetClass == MaterialName || Asset.AssetClass == MaterialInstanceConstantName)
+			if (Asset.AssetClass == UMaterial::StaticClass()->GetFName() || Asset.AssetClass == UMaterialInstanceConstant::StaticClass()->GetFName() || Asset.AssetClass == UMaterialFunction::StaticClass()->GetFName())
 			{
 				bCanUsePropertyMatrix = false;
 				break;
