Content Browser: Added the ability to register new advanced categories without editing engine source
Paper2D: Defined a "Paper2D" category and moved existing asset type actions over

New methods on IAssetTools:
- RegisterAdvancedAssetCategory
- FindAdvancedAssetCategory
- GetAllAdvancedAssetCategories

Upgrade Note: This gets rid of the named extension sections in the filter menu, which should instead be extended via registering a new category.

[UE-13646]

diff --git a/Engine/Plugins/2D/Paper2D/Source/Paper2DEditor/Private/Atlasing/AtlasAssetTypeActions.cpp b/Engine/Plugins/2D/Paper2D/Source/Paper2DEditor/Private/Atlasing/AtlasAssetTypeActions.cpp
--- a/Engine/Plugins/2D/Paper2D/Source/Paper2DEditor/Private/Atlasing/AtlasAssetTypeActions.cpp
+++ b/Engine/Plugins/2D/Paper2D/Source/Paper2DEditor/Private/Atlasing/AtlasAssetTypeActions.cpp
@@ -10,6 +10,11 @@
 //////////////////////////////////////////////////////////////////////////
 // FAtlasAssetTypeActions
 
+FAtlasAssetTypeActions::FAtlasAssetTypeActions(EAssetTypeCategories::Type InAssetCategory)
+	: MyAssetCategory(InAssetCategory)
+{
+}
+
 FText FAtlasAssetTypeActions::GetName() const
 {
 	return LOCTEXT("FAtlasAssetTypeActionsName", "Sprite Atlas Group");
@@ -33,7 +38,7 @@ void FAtlasAssetTypeActions::OpenAssetEditor(const TArray<UObject*>& InObjects,
 
 uint32 FAtlasAssetTypeActions::GetCategories()
 {
-	return EAssetTypeCategories::Misc;
+	return MyAssetCategory;
 }
 
 FText FAtlasAssetTypeActions::GetAssetDescription(const FAssetData& AssetData) const
diff --git a/Engine/Plugins/2D/Paper2D/Source/Paper2DEditor/Private/Atlasing/AtlasAssetTypeActions.h b/Engine/Plugins/2D/Paper2D/Source/Paper2DEditor/Private/Atlasing/AtlasAssetTypeActions.h
--- a/Engine/Plugins/2D/Paper2D/Source/Paper2DEditor/Private/Atlasing/AtlasAssetTypeActions.h
+++ b/Engine/Plugins/2D/Paper2D/Source/Paper2DEditor/Private/Atlasing/AtlasAssetTypeActions.h
@@ -7,6 +7,8 @@
 class FAtlasAssetTypeActions : public FAssetTypeActions_Base
 {
 public:
+	FAtlasAssetTypeActions(EAssetTypeCategories::Type InAssetCategory);
+
 	// IAssetTypeActions interface
 	virtual FText GetName() const override;
 	virtual FColor GetTypeColor() const override;
@@ -15,4 +17,7 @@ public:
 	virtual uint32 GetCategories() override;
 	virtual FText GetAssetDescription(const FAssetData& AssetData) const override;
 	// End of IAssetTypeActions interface
+
+private:
+	EAssetTypeCategories::Type MyAssetCategory;
 };
\ No newline at end of file
diff --git a/Engine/Plugins/2D/Paper2D/Source/Paper2DEditor/Private/FlipbookAssetTypeActions.cpp b/Engine/Plugins/2D/Paper2D/Source/Paper2DEditor/Private/FlipbookAssetTypeActions.cpp
--- a/Engine/Plugins/2D/Paper2D/Source/Paper2DEditor/Private/FlipbookAssetTypeActions.cpp
+++ b/Engine/Plugins/2D/Paper2D/Source/Paper2DEditor/Private/FlipbookAssetTypeActions.cpp
@@ -10,6 +10,11 @@
 //////////////////////////////////////////////////////////////////////////
 // FFlipbookAssetTypeActions
 
+FFlipbookAssetTypeActions::FFlipbookAssetTypeActions(EAssetTypeCategories::Type InAssetCategory)
+	: MyAssetCategory(InAssetCategory)
+{
+}
+
 FText FFlipbookAssetTypeActions::GetName() const
 {
 	return LOCTEXT("FFlipbookAssetTypeActionsName", "Paper Flipbook");
@@ -41,7 +46,7 @@ void FFlipbookAssetTypeActions::OpenAssetEditor(const TArray<UObject*>& InObject
 
 uint32 FFlipbookAssetTypeActions::GetCategories()
 {
-	return EAssetTypeCategories::Animation;
+	return EAssetTypeCategories::Animation | MyAssetCategory;
 }
 
 //////////////////////////////////////////////////////////////////////////
diff --git a/Engine/Plugins/2D/Paper2D/Source/Paper2DEditor/Private/FlipbookAssetTypeActions.h b/Engine/Plugins/2D/Paper2D/Source/Paper2DEditor/Private/FlipbookAssetTypeActions.h
--- a/Engine/Plugins/2D/Paper2D/Source/Paper2DEditor/Private/FlipbookAssetTypeActions.h
+++ b/Engine/Plugins/2D/Paper2D/Source/Paper2DEditor/Private/FlipbookAssetTypeActions.h
@@ -7,6 +7,8 @@
 class FFlipbookAssetTypeActions : public FAssetTypeActions_Base
 {
 public:
+	FFlipbookAssetTypeActions(EAssetTypeCategories::Type InAssetCategory);
+
 	// IAssetTypeActions interface
 	virtual FText GetName() const override;
 	virtual FColor GetTypeColor() const override;
@@ -14,4 +16,7 @@ public:
 	virtual void OpenAssetEditor(const TArray<UObject*>& InObjects, TSharedPtr<class IToolkitHost> EditWithinLevelEditor = TSharedPtr<IToolkitHost>()) override;
 	virtual uint32 GetCategories() override;
 	// End of IAssetTypeActions interface
+
+private:
+	EAssetTypeCategories::Type MyAssetCategory;
 };
\ No newline at end of file
diff --git a/Engine/Plugins/2D/Paper2D/Source/Paper2DEditor/Private/Paper2DEditorModule.cpp b/Engine/Plugins/2D/Paper2D/Source/Paper2DEditor/Private/Paper2DEditorModule.cpp
--- a/Engine/Plugins/2D/Paper2D/Source/Paper2DEditor/Private/Paper2DEditorModule.cpp
+++ b/Engine/Plugins/2D/Paper2D/Source/Paper2DEditor/Private/Paper2DEditorModule.cpp
@@ -67,12 +67,18 @@ DEFINE_LOG_CATEGORY(LogPaper2DEditor);
 class FPaper2DEditor : public IPaper2DEditorModule
 {
 public:
+	FPaper2DEditor()
+		: Paper2DAssetCategoryBit(EAssetTypeCategories::Misc)
+	{
+	}
+
 	// IPaper2DEditorModule interface
 	virtual TSharedPtr<FExtensibilityManager> GetSpriteEditorMenuExtensibilityManager() override { return SpriteEditor_MenuExtensibilityManager; }
 	virtual TSharedPtr<FExtensibilityManager> GetSpriteEditorToolBarExtensibilityManager() override { return SpriteEditor_ToolBarExtensibilityManager; }
 
 	virtual TSharedPtr<FExtensibilityManager> GetFlipbookEditorMenuExtensibilityManager() override { return FlipbookEditor_MenuExtensibilityManager; }
 	virtual TSharedPtr<FExtensibilityManager> GetFlipbookEditorToolBarExtensibilityManager() override { return FlipbookEditor_ToolBarExtensibilityManager; }
+	virtual uint32 GetPaper2DAssetCategory() const override { return Paper2DAssetCategoryBit; }
 	// End of IPaper2DEditorModule
 
 private:
@@ -93,6 +99,8 @@ private:
 	FCoreUObjectDelegates::FOnObjectPropertyChanged::FDelegate OnPropertyChangedHandle;
 	FDelegateHandle OnPropertyChangedHandleDelegateHandle;
 
+	EAssetTypeCategories::Type Paper2DAssetCategoryBit;
+
 public:
 	virtual void StartupModule() override
 	{
@@ -110,11 +118,14 @@ public:
 
 		// Register asset types
 		IAssetTools& AssetTools = FModuleManager::LoadModuleChecked<FAssetToolsModule>("AssetTools").Get();
-		RegisterAssetTypeAction(AssetTools, MakeShareable(new FSpriteAssetTypeActions));
-		RegisterAssetTypeAction(AssetTools, MakeShareable(new FFlipbookAssetTypeActions));
-		RegisterAssetTypeAction(AssetTools, MakeShareable(new FTileSetAssetTypeActions));
-		RegisterAssetTypeAction(AssetTools, MakeShareable(new FTileMapAssetTypeActions));
-		RegisterAssetTypeAction(AssetTools, MakeShareable(new FAtlasAssetTypeActions));
+
+		Paper2DAssetCategoryBit = AssetTools.RegisterAdvancedAssetCategory(FName(TEXT("Paper2D")), LOCTEXT("Paper2DAssetCategory", "Paper2D"));
+
+		RegisterAssetTypeAction(AssetTools, MakeShareable(new FSpriteAssetTypeActions(Paper2DAssetCategoryBit)));
+		RegisterAssetTypeAction(AssetTools, MakeShareable(new FFlipbookAssetTypeActions(Paper2DAssetCategoryBit)));
+		RegisterAssetTypeAction(AssetTools, MakeShareable(new FTileSetAssetTypeActions(Paper2DAssetCategoryBit)));
+		RegisterAssetTypeAction(AssetTools, MakeShareable(new FTileMapAssetTypeActions(Paper2DAssetCategoryBit)));
+		RegisterAssetTypeAction(AssetTools, MakeShareable(new FAtlasAssetTypeActions(Paper2DAssetCategoryBit)));
 
 		PaperSpriteBroker = MakeShareable(new FPaperSpriteAssetBroker);
 		FComponentAssetBrokerage::RegisterBroker(PaperSpriteBroker, UPaperSpriteComponent::StaticClass(), true, true);
diff --git a/Engine/Plugins/2D/Paper2D/Source/Paper2DEditor/Private/SpriteAssetTypeActions.cpp b/Engine/Plugins/2D/Paper2D/Source/Paper2DEditor/Private/SpriteAssetTypeActions.cpp
--- a/Engine/Plugins/2D/Paper2D/Source/Paper2DEditor/Private/SpriteAssetTypeActions.cpp
+++ b/Engine/Plugins/2D/Paper2D/Source/Paper2DEditor/Private/SpriteAssetTypeActions.cpp
@@ -14,6 +14,11 @@
 //////////////////////////////////////////////////////////////////////////
 // FSpriteAssetTypeActions
 
+FSpriteAssetTypeActions::FSpriteAssetTypeActions(EAssetTypeCategories::Type InAssetCategory)
+	: MyAssetCategory(InAssetCategory)
+{
+}
+
 FText FSpriteAssetTypeActions::GetName() const
 {
 	return LOCTEXT("FSpriteAssetTypeActionsName", "Sprite");
@@ -45,7 +50,7 @@ void FSpriteAssetTypeActions::OpenAssetEditor(const TArray<UObject*>& InObjects,
 
 uint32 FSpriteAssetTypeActions::GetCategories()
 {
-	return EAssetTypeCategories::Misc;
+	return MyAssetCategory;
 }
 
 void FSpriteAssetTypeActions::GetActions(const TArray<UObject*>& InObjects, FMenuBuilder& MenuBuilder)
diff --git a/Engine/Plugins/2D/Paper2D/Source/Paper2DEditor/Private/SpriteAssetTypeActions.h b/Engine/Plugins/2D/Paper2D/Source/Paper2DEditor/Private/SpriteAssetTypeActions.h
--- a/Engine/Plugins/2D/Paper2D/Source/Paper2DEditor/Private/SpriteAssetTypeActions.h
+++ b/Engine/Plugins/2D/Paper2D/Source/Paper2DEditor/Private/SpriteAssetTypeActions.h
@@ -7,6 +7,8 @@
 class FSpriteAssetTypeActions : public FAssetTypeActions_Base
 {
 public:
+	FSpriteAssetTypeActions(EAssetTypeCategories::Type InAssetCategory);
+
 	// IAssetTypeActions interface
 	virtual FText GetName() const override;
 	virtual FColor GetTypeColor() const override;
@@ -17,6 +19,9 @@ public:
 	virtual uint32 GetCategories() override;
 	// End of IAssetTypeActions interface
 
+private:
+	EAssetTypeCategories::Type MyAssetCategory;
+
 private:
 	/* Handler for "Create Flipbook" */
 	void ExecuteCreateFlipbook(TArray<TWeakObjectPtr<UPaperSprite>> Objects);
diff --git a/Engine/Plugins/2D/Paper2D/Source/Paper2DEditor/Private/TileMapEditing/TileMapAssetTypeActions.cpp b/Engine/Plugins/2D/Paper2D/Source/Paper2DEditor/Private/TileMapEditing/TileMapAssetTypeActions.cpp
--- a/Engine/Plugins/2D/Paper2D/Source/Paper2DEditor/Private/TileMapEditing/TileMapAssetTypeActions.cpp
+++ b/Engine/Plugins/2D/Paper2D/Source/Paper2DEditor/Private/TileMapEditing/TileMapAssetTypeActions.cpp
@@ -10,6 +10,11 @@
 //////////////////////////////////////////////////////////////////////////
 // FTileMapAssetTypeActions
 
+FTileMapAssetTypeActions::FTileMapAssetTypeActions(EAssetTypeCategories::Type InAssetCategory)
+	: MyAssetCategory(InAssetCategory)
+{
+}
+
 FText FTileMapAssetTypeActions::GetName() const
 {
 	return LOCTEXT("FTileMapAssetTypeActionsName", "Tile Map");
@@ -41,7 +46,7 @@ void FTileMapAssetTypeActions::OpenAssetEditor(const TArray<UObject*>& InObjects
 
 uint32 FTileMapAssetTypeActions::GetCategories()
 {
-	return EAssetTypeCategories::Misc;
+	return MyAssetCategory;
 }
 
 //////////////////////////////////////////////////////////////////////////
diff --git a/Engine/Plugins/2D/Paper2D/Source/Paper2DEditor/Private/TileMapEditing/TileMapAssetTypeActions.h b/Engine/Plugins/2D/Paper2D/Source/Paper2DEditor/Private/TileMapEditing/TileMapAssetTypeActions.h
--- a/Engine/Plugins/2D/Paper2D/Source/Paper2DEditor/Private/TileMapEditing/TileMapAssetTypeActions.h
+++ b/Engine/Plugins/2D/Paper2D/Source/Paper2DEditor/Private/TileMapEditing/TileMapAssetTypeActions.h
@@ -7,6 +7,8 @@
 class FTileMapAssetTypeActions : public FAssetTypeActions_Base
 {
 public:
+	FTileMapAssetTypeActions(EAssetTypeCategories::Type InAssetCategory);
+
 	// IAssetTypeActions interface
 	virtual FText GetName() const override;
 	virtual FColor GetTypeColor() const override;
@@ -14,4 +16,7 @@ public:
 	virtual void OpenAssetEditor(const TArray<UObject*>& InObjects, TSharedPtr<class IToolkitHost> EditWithinLevelEditor = TSharedPtr<IToolkitHost>()) override;
 	virtual uint32 GetCategories() override;
 	// End of IAssetTypeActions interface
+
+private:
+	EAssetTypeCategories::Type MyAssetCategory;
 };
\ No newline at end of file
diff --git a/Engine/Plugins/2D/Paper2D/Source/Paper2DEditor/Private/TileSetAssetTypeActions.cpp b/Engine/Plugins/2D/Paper2D/Source/Paper2DEditor/Private/TileSetAssetTypeActions.cpp
--- a/Engine/Plugins/2D/Paper2D/Source/Paper2DEditor/Private/TileSetAssetTypeActions.cpp
+++ b/Engine/Plugins/2D/Paper2D/Source/Paper2DEditor/Private/TileSetAssetTypeActions.cpp
@@ -12,6 +12,11 @@
 //////////////////////////////////////////////////////////////////////////
 // FTileSetAssetTypeActions
 
+FTileSetAssetTypeActions::FTileSetAssetTypeActions(EAssetTypeCategories::Type InAssetCategory)
+	: MyAssetCategory(InAssetCategory)
+{
+}
+
 FText FTileSetAssetTypeActions::GetName() const
 {
 	return LOCTEXT("FTileSetAssetTypeActionsName", "Tile Set");
@@ -43,7 +48,7 @@ void FTileSetAssetTypeActions::OpenAssetEditor(const TArray<UObject*>& InObjects
 
 uint32 FTileSetAssetTypeActions::GetCategories()
 {
-	return EAssetTypeCategories::Misc;
+	return MyAssetCategory;
 }
 
 void FTileSetAssetTypeActions::GetActions(const TArray<UObject*>& InObjects, FMenuBuilder& MenuBuilder)
diff --git a/Engine/Plugins/2D/Paper2D/Source/Paper2DEditor/Private/TileSetAssetTypeActions.h b/Engine/Plugins/2D/Paper2D/Source/Paper2DEditor/Private/TileSetAssetTypeActions.h
--- a/Engine/Plugins/2D/Paper2D/Source/Paper2DEditor/Private/TileSetAssetTypeActions.h
+++ b/Engine/Plugins/2D/Paper2D/Source/Paper2DEditor/Private/TileSetAssetTypeActions.h
@@ -9,6 +9,8 @@ class UPaperTileSet;
 class FTileSetAssetTypeActions : public FAssetTypeActions_Base
 {
 public:
+	FTileSetAssetTypeActions(EAssetTypeCategories::Type InAssetCategory);
+
 	// IAssetTypeActions interface
 	virtual FText GetName() const override;
 	virtual FColor GetTypeColor() const override;
@@ -19,6 +21,9 @@ public:
 	virtual void GetActions(const TArray<UObject*>& InObjects, FMenuBuilder& MenuBuilder) override;
 	// End of IAssetTypeActions interface
 
+private:
+	EAssetTypeCategories::Type MyAssetCategory;
+
 private:
 	void ExecuteCreateTileMap(TWeakObjectPtr<UPaperTileSet> TileSetPtr);
 };
diff --git a/Engine/Plugins/2D/Paper2D/Source/Paper2DEditor/Public/Paper2DEditorModule.h b/Engine/Plugins/2D/Paper2D/Source/Paper2DEditor/Public/Paper2DEditorModule.h
--- a/Engine/Plugins/2D/Paper2D/Source/Paper2DEditor/Public/Paper2DEditorModule.h
+++ b/Engine/Plugins/2D/Paper2D/Source/Paper2DEditor/Public/Paper2DEditorModule.h
@@ -15,5 +15,7 @@ public:
 
 	virtual TSharedPtr<class FExtensibilityManager> GetFlipbookEditorMenuExtensibilityManager() { return nullptr; }
 	virtual TSharedPtr<class FExtensibilityManager> GetFlipbookEditorToolBarExtensibilityManager() { return nullptr; }
+
+	virtual uint32 GetPaper2DAssetCategory() const = 0;
 };
 
diff --git a/Engine/Source/Developer/AssetTools/Private/AssetTools.cpp b/Engine/Source/Developer/AssetTools/Private/AssetTools.cpp
--- a/Engine/Source/Developer/AssetTools/Private/AssetTools.cpp
+++ b/Engine/Source/Developer/AssetTools/Private/AssetTools.cpp
@@ -22,7 +22,19 @@
 FAssetTools::FAssetTools()
 	: AssetRenameManager( MakeShareable(new FAssetRenameManager) )
 	, AssetFixUpRedirectors( MakeShareable(new FAssetFixUpRedirectors) )
+	, NextUserCategoryBit( EAssetTypeCategories::FirstUser )
 {
+	// Register the built-in advanced categories
+	AllocatedCategoryBits.Add(TEXT("_BuiltIn_0"), FAdvancedAssetCategory(EAssetTypeCategories::Animation, LOCTEXT("AnimationAssetCategory", "Animation")));
+	AllocatedCategoryBits.Add(TEXT("_BuiltIn_1"), FAdvancedAssetCategory(EAssetTypeCategories::Blueprint, LOCTEXT("BlueprintAssetCategory", "Blueprints")));
+	AllocatedCategoryBits.Add(TEXT("_BuiltIn_2"), FAdvancedAssetCategory(EAssetTypeCategories::MaterialsAndTextures, LOCTEXT("MaterialAssetCategory", "Materials & Textures")));
+	AllocatedCategoryBits.Add(TEXT("_BuiltIn_3"), FAdvancedAssetCategory(EAssetTypeCategories::Sounds, LOCTEXT("SoundAssetCategory", "Sounds")));
+	AllocatedCategoryBits.Add(TEXT("_BuiltIn_4"), FAdvancedAssetCategory(EAssetTypeCategories::Physics, LOCTEXT("PhysicsAssetCategory", "Physics")));
+	AllocatedCategoryBits.Add(TEXT("_BuiltIn_5"), FAdvancedAssetCategory(EAssetTypeCategories::UI, LOCTEXT("UserInterfaceAssetCategory", "User Interface")));
+	AllocatedCategoryBits.Add(TEXT("_BuiltIn_6"), FAdvancedAssetCategory(EAssetTypeCategories::Misc, LOCTEXT("MiscellaneousAssetCategory", "Miscellaneous")));
+	AllocatedCategoryBits.Add(TEXT("_BuiltIn_7"), FAdvancedAssetCategory(EAssetTypeCategories::Gameplay, LOCTEXT("GameplayAssetCategory", "Gameplay")));
+
+	// Register the built-in asset type actions
 	RegisterAssetTypeActions( MakeShareable(new FAssetTypeActions_AnimationAsset) );
 	RegisterAssetTypeActions( MakeShareable(new FAssetTypeActions_AnimBlueprint) );
 	RegisterAssetTypeActions( MakeShareable(new FAssetTypeActions_AnimComposite) );
@@ -92,6 +104,9 @@ FAssetTools::FAssetTools()
 	RegisterAssetTypeActions( MakeShareable(new FAssetTypeActions_VectorFieldStatic) );
 	RegisterAssetTypeActions( MakeShareable(new FAssetTypeActions_VertexAnimation) );
 
+	// Note: Please don't add any more actions here!  They belong in an editor-only module that is more tightly
+	// coupled to your new system, and you should not create a dependency on your new system from AssetTools.
+
 	if ( UEditorEngine::IsUsingWorldAssets() )
 	{
 		RegisterAssetTypeActions( MakeShareable(new FAssetTypeActions_World) );
@@ -150,6 +165,53 @@ TWeakPtr<IAssetTypeActions> FAssetTools::GetAssetTypeActionsForClass( UClass* Cl
 	return MostDerivedAssetTypeActions;
 }
 
+EAssetTypeCategories::Type FAssetTools::RegisterAdvancedAssetCategory(FName CategoryKey, FText CategoryDisplayName)
+{
+	EAssetTypeCategories::Type Result = FindAdvancedAssetCategory(CategoryKey);
+	if (Result == EAssetTypeCategories::Misc)
+	{
+		if (NextUserCategoryBit != 0)
+		{
+			// Register the category
+			Result = (EAssetTypeCategories::Type)NextUserCategoryBit;
+			AllocatedCategoryBits.Add(CategoryKey, FAdvancedAssetCategory(Result, CategoryDisplayName));
+
+			// Advance to the next bit, or store that we're out
+			if (NextUserCategoryBit == EAssetTypeCategories::LastUser)
+			{
+				NextUserCategoryBit = 0;
+			}
+			else
+			{
+				NextUserCategoryBit = NextUserCategoryBit << 1;
+			}
+		}
+		else
+		{
+			UE_LOG(LogAssetTools, Warning, TEXT("RegisterAssetTypeCategory(\"%s\", \"%s\") failed as all user bits have been exhausted (placing into the Misc category instead)"), *CategoryKey.ToString(), *CategoryDisplayName.ToString());
+		}
+	}
+
+	return Result;
+}
+
+EAssetTypeCategories::Type FAssetTools::FindAdvancedAssetCategory(FName CategoryKey) const
+{
+	if (const FAdvancedAssetCategory* ExistingCategory = AllocatedCategoryBits.Find(CategoryKey))
+	{
+		return ExistingCategory->CategoryType;
+	}
+	else
+	{
+		return EAssetTypeCategories::Misc;
+	}
+}
+
+void FAssetTools::GetAllAdvancedAssetCategories(TArray<FAdvancedAssetCategory>& OutCategoryList) const
+{
+	AllocatedCategoryBits.GenerateValueArray(OutCategoryList);
+}
+
 void FAssetTools::RegisterClassTypeActions(const TSharedRef<IClassTypeActions>& NewActions)
 {
 	ClassTypeActionsList.Add(NewActions);
diff --git a/Engine/Source/Developer/AssetTools/Private/AssetTools.h b/Engine/Source/Developer/AssetTools/Private/AssetTools.h
--- a/Engine/Source/Developer/AssetTools/Private/AssetTools.h
+++ b/Engine/Source/Developer/AssetTools/Private/AssetTools.h
@@ -13,6 +13,9 @@ public:
 	virtual void UnregisterAssetTypeActions(const TSharedRef<IAssetTypeActions>& ActionsToRemove) override;
 	virtual void GetAssetTypeActionsList( TArray<TWeakPtr<IAssetTypeActions>>& OutAssetTypeActionsList ) const override;
 	virtual TWeakPtr<IAssetTypeActions> GetAssetTypeActionsForClass( UClass* Class ) const override;
+	virtual EAssetTypeCategories::Type RegisterAdvancedAssetCategory(FName CategoryKey, FText CategoryDisplayName) override;
+	virtual EAssetTypeCategories::Type FindAdvancedAssetCategory(FName CategoryKey) const override;
+	virtual void GetAllAdvancedAssetCategories(TArray<FAdvancedAssetCategory>& OutCategoryList) const override;
 	virtual void RegisterClassTypeActions(const TSharedRef<IClassTypeActions>& NewActions) override;
 	virtual void UnregisterClassTypeActions(const TSharedRef<IClassTypeActions>& ActionsToRemove) override;
 	virtual void GetClassTypeActionsList( TArray<TWeakPtr<IClassTypeActions>>& OutClassTypeActionsList ) const override;
@@ -79,4 +82,10 @@ private:
 
 	/** The list of all registered ClassTypeActions */
 	TArray<TSharedRef<IClassTypeActions>> ClassTypeActionsList;
+
+	/** The categories that have been allocated already */
+	TMap<FName, FAdvancedAssetCategory> AllocatedCategoryBits;
+	
+	/** The next user category bit to allocate (set to 0 when there are no more bits left) */
+	uint32 NextUserCategoryBit;
 };
diff --git a/Engine/Source/Developer/AssetTools/Public/IAssetTools.h b/Engine/Source/Developer/AssetTools/Public/IAssetTools.h
--- a/Engine/Source/Developer/AssetTools/Public/IAssetTools.h
+++ b/Engine/Source/Developer/AssetTools/Public/IAssetTools.h
@@ -28,6 +28,42 @@ struct FAssetRenameData
 
 DECLARE_MULTICAST_DELEGATE_OneParam(FAssetPostRenameEvent, const TArray<FAssetRenameData>&);
 
+// The list of categories for Asset Type Actions and UFactory subclasses
+namespace EAssetTypeCategories
+{
+	enum Type
+	{
+		None					= 0,
+		Basic					= 1 << 0,
+		Animation				= 1 << 1,
+		MaterialsAndTextures	= 1 << 2,
+		Sounds					= 1 << 3,
+		Physics					= 1 << 4,
+		UI						= 1 << 5,
+		Misc					= 1 << 6,
+		Gameplay				= 1 << 7,
+		Blueprint				= 1 << 8,
+
+		// Items below this will be allocated at runtime via RegisterAdvancedAssetCategory
+		FirstUser				= 1 << 9,
+		LastUser				= 1 << 31,
+		// Last allowed value is 1 << 31
+	};
+}
+
+
+struct FAdvancedAssetCategory
+{
+	EAssetTypeCategories::Type CategoryType;
+	FText CategoryName;
+
+	FAdvancedAssetCategory(EAssetTypeCategories::Type InCategoryType, FText InCategoryName)
+		: CategoryType(InCategoryType)
+		, CategoryName(InCategoryName)
+	{
+	}
+};
+
 class IAssetTools
 {
 public:
@@ -46,6 +82,18 @@ public:
 	/** Gets the appropriate AssetTypeActions for the supplied class */
 	virtual TWeakPtr<IAssetTypeActions> GetAssetTypeActionsForClass( UClass* Class ) const = 0;
 
+	/**
+	 * Allocates a Category bit for a user-defined Category, or EAssetTypeCategories::Misc if all available bits are allocated.
+	 * Ignores duplicate calls with the same CategoryKey (returns the existing bit but does not change the display name).
+	 */
+	virtual EAssetTypeCategories::Type RegisterAdvancedAssetCategory(FName CategoryKey, FText CategoryDisplayName) = 0;
+
+	/** Returns the allocated Category bit for a user-specified Category, or EAssetTypeCategories::Misc if it doesn't exist */
+	virtual EAssetTypeCategories::Type FindAdvancedAssetCategory(FName CategoryKey) const = 0;
+	
+	/** Returns the list of all advanced asset categories */
+	virtual void GetAllAdvancedAssetCategories(TArray<FAdvancedAssetCategory>& OutCategoryList) const = 0;
+
 	/** Registers a class type actions object so it can provide information about and actions for class asset types. */
 	virtual void RegisterClassTypeActions(const TSharedRef<IClassTypeActions>& NewActions) = 0;
 
diff --git a/Engine/Source/Developer/AssetTools/Public/IAssetTypeActions.h b/Engine/Source/Developer/AssetTools/Public/IAssetTypeActions.h
--- a/Engine/Source/Developer/AssetTools/Public/IAssetTypeActions.h
+++ b/Engine/Source/Developer/AssetTools/Public/IAssetTypeActions.h
@@ -2,26 +2,9 @@
 
 #pragma once
 
+#include "Developer/AssetTools/Public/IAssetTools.h"
 #include "Developer/Merge/Public/Merge.h" // for FOnMergeResolved
 
-namespace EAssetTypeCategories
-{
-	enum Type
-	{
-		None					= 0,
-		Basic					= 1 << 0,
-		Animation				= 1 << 1,
-		MaterialsAndTextures	= 1 << 2,
-		Sounds					= 1 << 3,
-		Physics					= 1 << 4,
-		UI						= 1 << 5,
-		Misc					= 1 << 6,
-		Gameplay				= 1 << 7,
-		Blueprint				= 1 << 8
-		// Last allowed value is 1 << 31
-	};
-}
-
 namespace EAssetTypeActivationMethod
 {
 	enum Type
diff --git a/Engine/Source/Editor/ContentBrowser/Private/NewAssetOrClassContextMenu.cpp b/Engine/Source/Editor/ContentBrowser/Private/NewAssetOrClassContextMenu.cpp
--- a/Engine/Source/Editor/ContentBrowser/Private/NewAssetOrClassContextMenu.cpp
+++ b/Engine/Source/Editor/ContentBrowser/Private/NewAssetOrClassContextMenu.cpp
@@ -42,12 +42,6 @@ TArray<FFactoryItem> FindFactoriesInCategory(EAssetTypeCategories::Type AssetTyp
 	return FactoriesInThisCategory;
 }
 
-struct FAdvancedAssetCategory
-{
-	EAssetTypeCategories::Type CategoryType;
-	FText CategoryName;
-};
-
 class SFactoryMenuEntry : public SCompoundWidget
 {
 public:
@@ -336,23 +330,15 @@ void FNewAssetOrClassContextMenu::MakeContextMenu(
 		// Add Advanced Asset
 		MenuBuilder.BeginSection("ContentBrowserNewAdvancedAsset", LOCTEXT("AdvancedAssetsMenuHeading", "Create Advanced Asset"));
 		{
-			// List of advanced categories to add, in order
-			const FAdvancedAssetCategory AdvancedAssetCategories[] = {
-				//	CategoryType								CategoryName
-				{	EAssetTypeCategories::Animation,			LOCTEXT("AnimationAssetCategory", "Animation")				},
-				{	EAssetTypeCategories::Blueprint,			LOCTEXT("BlueprintAssetCategory", "Blueprints")				},
-				{	EAssetTypeCategories::MaterialsAndTextures,	LOCTEXT("MaterialAssetCategory", "Materials & Textures")	},
-				{	EAssetTypeCategories::Sounds,				LOCTEXT("SoundAssetCategory", "Sounds")						},
-				{	EAssetTypeCategories::Physics,				LOCTEXT("PhysicsAssetCategory", "Physics")					},
-				{	EAssetTypeCategories::UI,					LOCTEXT("UserInterfaceAssetCategory", "User Interface")		},
-				{	EAssetTypeCategories::Misc,					LOCTEXT("MiscellaneousAssetCategory", "Miscellaneous")		},
-				{	EAssetTypeCategories::Gameplay,				LOCTEXT("GameplayAssetCategory", "Gameplay")				},
-			};
+			FAssetToolsModule& AssetToolsModule = FAssetToolsModule::GetModule();
+
+			TArray<FAdvancedAssetCategory> AdvancedAssetCategories;
+			AssetToolsModule.Get().GetAllAdvancedAssetCategories(/*out*/ AdvancedAssetCategories);
 
 			for (const FAdvancedAssetCategory& AdvancedAssetCategory : AdvancedAssetCategories)
 			{
-				TArray<FFactoryItem> AnimationFactories = FindFactoriesInCategory(AdvancedAssetCategory.CategoryType);
-				if (AnimationFactories.Num() > 0)
+				TArray<FFactoryItem> Factories = FindFactoriesInCategory(AdvancedAssetCategory.CategoryType);
+				if (Factories.Num() > 0)
 				{
 					MenuBuilder.AddSubMenu(
 						AdvancedAssetCategory.CategoryName,
diff --git a/Engine/Source/Editor/ContentBrowser/Private/SFilterList.cpp b/Engine/Source/Editor/ContentBrowser/Private/SFilterList.cpp
--- a/Engine/Source/Editor/ContentBrowser/Private/SFilterList.cpp
+++ b/Engine/Source/Editor/ContentBrowser/Private/SFilterList.cpp
@@ -970,6 +970,8 @@ bool IsFilteredByPicker(TArray<UClass*> FilterClassList, UClass* TestClass)
 
 TSharedRef<SWidget> SFilterList::MakeAddFilterMenu(EAssetTypeCategories::Type MenuExpansion)
 {
+	FAssetToolsModule& AssetToolsModule = FModuleManager::LoadModuleChecked<FAssetToolsModule>(TEXT("AssetTools"));
+
 	// A local struct to describe a category in the filter menu
 	struct FCategoryMenu
 	{
@@ -992,18 +994,23 @@ TSharedRef<SWidget> SFilterList::MakeAddFilterMenu(EAssetTypeCategories::Type Me
 
 	// Create a map of Categories to Menus
 	TMap<EAssetTypeCategories::Type, FCategoryMenu> CategoryToMenuMap;
+
+	// Add the Basic category
 	CategoryToMenuMap.Add(EAssetTypeCategories::Basic, FCategoryMenu( LOCTEXT("BasicFilter", "Basic"), LOCTEXT("BasicFilterTooltip", "Filter by basic assets."), "ContentBrowserFilterBasicAsset", LOCTEXT("BasicAssetsMenuHeading", "Basic Assets") ) );
-	CategoryToMenuMap.Add(EAssetTypeCategories::Animation, FCategoryMenu( LOCTEXT("AnimationFilter", "Animation"), LOCTEXT("AnimationFilterTooltip", "Filter by animation assets."), "ContentBrowserFilterAnimationAsset", LOCTEXT("AnimationAssetsMenuHeading", "Animation Assets") ) );
-	CategoryToMenuMap.Add(EAssetTypeCategories::Blueprint, FCategoryMenu( LOCTEXT("BlueprintFilter", "Blueprints"), LOCTEXT("BlueprintFilterTooltip", "Filter by blueprint assets."), "ContentBrowserFilterBlueprintAsset", LOCTEXT("BlueprintAssetsMenuHeading", "Blueprint Assets") ) );
-	CategoryToMenuMap.Add(EAssetTypeCategories::MaterialsAndTextures, FCategoryMenu(LOCTEXT("MaterialFilter", "Materials & Textures"), LOCTEXT("MaterialFilterTooltip", "Filter by material and texture assets."), "ContentBrowserFilterMaterialAsset", LOCTEXT("MaterialAssetsMenuHeading", "Material Assets")));
-	CategoryToMenuMap.Add(EAssetTypeCategories::Sounds, FCategoryMenu( LOCTEXT("SoundFilter", "Sounds"), LOCTEXT("SoundFilterTooltip", "Filter by sound assets."), "ContentBrowserFilterSoundAsset", LOCTEXT("SoundAssetsMenuHeading", "Sound Assets") ) );
-	CategoryToMenuMap.Add(EAssetTypeCategories::Physics, FCategoryMenu( LOCTEXT("PhysicsFilter", "Physics"), LOCTEXT("PhysicsFilterTooltip", "Filter by physics assets."), "ContentBrowserFilterPhysicsAsset", LOCTEXT("PhysicsAssetsMenuHeading", "Physics Assets") ) );
-	CategoryToMenuMap.Add(EAssetTypeCategories::UI, FCategoryMenu(LOCTEXT("UIFilter", "User Interface"), LOCTEXT("UIFilterTooltip", "Filter by UI assets."), "ContentBrowserFilterUIAsset", LOCTEXT("UIAssetsMenuHeading", "User Interface Assets")));
-	CategoryToMenuMap.Add(EAssetTypeCategories::Misc, FCategoryMenu( LOCTEXT("MiscFilter", "Miscellaneous"), LOCTEXT("MiscFilterTooltip", "Filter by miscellaneous assets."), "ContentBrowserFilterMiscAsset", LOCTEXT("MiscAssetsMenuHeading", "Misc Assets") ) );
-	CategoryToMenuMap.Add(EAssetTypeCategories::Gameplay, FCategoryMenu(LOCTEXT("GameplayFilter", "Gameplay"), LOCTEXT("GameplayFilterTooltip", "Filter by gameplay assets."), "ContentBrowserFilterGameplayAsset", LOCTEXT("GameplayAssetsMenuHeading", "Gameplay Assets")));
-
-	// Load the asset tools module to get access to the browser type maps
-	FAssetToolsModule& AssetToolsModule = FModuleManager::LoadModuleChecked<FAssetToolsModule>(TEXT("AssetTools"));
+
+	// Add the advanced categories
+	TArray<FAdvancedAssetCategory> AdvancedAssetCategories;
+	AssetToolsModule.Get().GetAllAdvancedAssetCategories(/*out*/ AdvancedAssetCategories);
+
+	for (const FAdvancedAssetCategory& AdvancedAssetCategory : AdvancedAssetCategories)
+	{
+		const FName ExtensionPoint = NAME_None;
+		const FText SectionHeading = FText::Format(LOCTEXT("WildcardFilterHeadingHeadingTooltip", "{0} Assets."), AdvancedAssetCategory.CategoryName);
+		const FText Tooltip = FText::Format(LOCTEXT("WildcardFilterTooltip", "Filter by {0}."), SectionHeading);
+		CategoryToMenuMap.Add(AdvancedAssetCategory.CategoryType, FCategoryMenu(AdvancedAssetCategory.CategoryName, Tooltip, ExtensionPoint, SectionHeading));
+	}
+
+	// Get the browser type maps
 	TArray<TWeakPtr<IAssetTypeActions>> AssetTypeActionsList;
 	AssetToolsModule.Get().GetAssetTypeActionsList(AssetTypeActionsList);
 
