Text FrontendFilter now tries to match with components of the path name, as well as the asset name itself.

diff --git a/Engine/Source/Editor/ContentBrowser/Private/FrontendFilters.cpp b/Engine/Source/Editor/ContentBrowser/Private/FrontendFilters.cpp
--- a/Engine/Source/Editor/ContentBrowser/Private/FrontendFilters.cpp
+++ b/Engine/Source/Editor/ContentBrowser/Private/FrontendFilters.cpp
@@ -11,15 +11,46 @@
 // FFrontendFilter_Text
 /////////////////////////////////////////
 
-void AssetDataToClassAndNameStrings(FAssetFilterType Asset, OUT TArray< FString >& Array)
+void AssetDataToNameString(FAssetFilterType Asset, OUT TArray< FString >& Array)
 {
-	Array.Add(Asset.AssetClass.ToString());
-	Array.Add(Asset.AssetName.ToString());
+	const FString FullPath = Asset.PackageName.ToString();
+	const TCHAR* Ptr = FullPath.GetCharArray().GetData();
+	if (Ptr == nullptr)
+	{
+		return;
+	}
+
+	// Enter each piece of the path name, apart from the first, as a potential match
+	bool bIsFirst = true;
+	while (const TCHAR* Delimiter = FCString::Strchr(Ptr, '/'))
+	{
+		const int32 Length = Delimiter - Ptr;
+
+		if (Length > 0)
+		{
+			if (bIsFirst)
+			{
+				bIsFirst = false;
+			}
+			else
+			{
+				Array.Emplace(Length, Ptr);
+			}
+		}
+
+		Ptr += (Length + 1);
+	}
+
+	if (*Ptr != 0)
+	{
+		Array.Emplace(Ptr);
+	}
 }
 
-void AssetDataToNameString(FAssetFilterType Asset, OUT TArray< FString >& Array)
+void AssetDataToClassAndNameStrings(FAssetFilterType Asset, OUT TArray< FString >& Array)
 {
-	Array.Add(Asset.AssetName.ToString());
+	Array.Add(Asset.AssetClass.ToString());
+	AssetDataToNameString(Asset, Array);
 }
 
 FFrontendFilter_Text::FFrontendFilter_Text()
