Allow aborting the "Move Folder Here" operation with a large number of assets within.

diff --git a/Engine/Source/Editor/ContentBrowser/Private/ContentBrowserUtils.cpp b/Engine/Source/Editor/ContentBrowser/Private/ContentBrowserUtils.cpp
--- a/Engine/Source/Editor/ContentBrowser/Private/ContentBrowserUtils.cpp
+++ b/Engine/Source/Editor/ContentBrowser/Private/ContentBrowserUtils.cpp
@@ -698,7 +698,7 @@ void ContentBrowserUtils::DisplayConfirmationPopup(const FText& Message, const F
 	Popup->OpenPopup(ParentContent);
 }
 
-void ContentBrowserUtils::CopyFolders(const TArray<FString>& InSourcePathNames, const FString& DestPath)
+bool ContentBrowserUtils::CopyFolders(const TArray<FString>& InSourcePathNames, const FString& DestPath)
 {
 	TMap<FString, TArray<UObject*> > SourcePathToLoadedAssets;
 
@@ -707,7 +707,10 @@ void ContentBrowserUtils::CopyFolders(const TArray<FString>& InSourcePathNames,
 	SourcePathNames.Remove(DestPath);
 
 	// Load all assets in the source paths
-	PrepareFoldersForDragDrop(SourcePathNames, SourcePathToLoadedAssets);
+	if (!PrepareFoldersForDragDrop(SourcePathNames, SourcePathToLoadedAssets))
+	{
+		return false;
+	}
 
 	// Load the Asset Registry to update paths during the copy
 	FAssetRegistryModule& AssetRegistryModule = FModuleManager::LoadModuleChecked<FAssetRegistryModule>("AssetRegistry");
@@ -729,9 +732,11 @@ void ContentBrowserUtils::CopyFolders(const TArray<FString>& InSourcePathNames,
 			ObjectTools::DuplicateObjects( PathIt.Value(), PathIt.Key(), Destination, /*bOpenDialog=*/false );
 		}
 	}
+
+	return true;
 }
 
-void ContentBrowserUtils::MoveFolders(const TArray<FString>& InSourcePathNames, const FString& DestPath)
+bool ContentBrowserUtils::MoveFolders(const TArray<FString>& InSourcePathNames, const FString& DestPath)
 {
 	TMap<FString, TArray<UObject*> > SourcePathToLoadedAssets;
 	FString DestPathWithTrailingSlash = DestPath / "";
@@ -752,7 +757,10 @@ void ContentBrowserUtils::MoveFolders(const TArray<FString>& InSourcePathNames,
 	}
 
 	// Load all assets in the source paths
-	PrepareFoldersForDragDrop(SourcePathNames, SourcePathToLoadedAssets);
+	if (!PrepareFoldersForDragDrop(SourcePathNames, SourcePathToLoadedAssets))
+	{
+		return false;
+	}
 	
 	// Load the Asset Registry to update paths during the move
 	FAssetRegistryModule& AssetRegistryModule = FModuleManager::LoadModuleChecked<FAssetRegistryModule>("AssetRegistry");
@@ -778,14 +786,14 @@ void ContentBrowserUtils::MoveFolders(const TArray<FString>& InSourcePathNames,
 		// Attempt to remove the old paths. This operation will silently fail if any assets failed to move.
 		AssetRegistryModule.Get().RemovePath(SourcePath);
 	}
+
+	return true;
 }
 
-void ContentBrowserUtils::PrepareFoldersForDragDrop(const TArray<FString>& SourcePathNames, TMap< FString, TArray<UObject*> >& OutSourcePathToLoadedAssets)
+bool ContentBrowserUtils::PrepareFoldersForDragDrop(const TArray<FString>& SourcePathNames, TMap< FString, TArray<UObject*> >& OutSourcePathToLoadedAssets)
 {
 	TSet<UObject*> AllFoundObjects;
 
-	GWarn->BeginSlowTask( LOCTEXT("FolderDragDrop_Loading", "Loading folders"), true);
-
 	// Load the Asset Registry to update paths during the move
 	FAssetRegistryModule& AssetRegistryModule = FModuleManager::LoadModuleChecked<FAssetRegistryModule>("AssetRegistry");
 
@@ -807,9 +815,15 @@ void ContentBrowserUtils::PrepareFoldersForDragDrop(const TArray<FString>& Sourc
 	TArray<FString> UnloadedObjects;
 	if(ShouldPromptToLoadAssets(ObjectPathsToWarnAbout, UnloadedObjects))
 	{
-		PromptToLoadAssets(UnloadedObjects);
+		const bool bShouldLoadAssets = PromptToLoadAssets(UnloadedObjects);
+		if (!bShouldLoadAssets)
+		{
+			return false;
+		}
 	}
 
+	GWarn->BeginSlowTask(LOCTEXT("FolderDragDrop_Loading", "Loading folders"), true);
+
 	// For every source path, load every package in the path (if necessary) and keep track of the assets that were loaded
 	for ( auto PathIt = SourcePathNames.CreateConstIterator(); PathIt; ++PathIt )
 	{
@@ -858,6 +872,7 @@ void ContentBrowserUtils::PrepareFoldersForDragDrop(const TArray<FString>& Sourc
 	GWarn->EndSlowTask();
 
 	ensure(SourcePathNames.Num() == OutSourcePathToLoadedAssets.Num());
+	return true;
 }
 
 void ContentBrowserUtils::CopyAssetReferencesToClipboard(const TArray<FAssetData>& AssetsToCopy)
diff --git a/Engine/Source/Editor/ContentBrowser/Private/ContentBrowserUtils.h b/Engine/Source/Editor/ContentBrowser/Private/ContentBrowserUtils.h
--- a/Engine/Source/Editor/ContentBrowser/Private/ContentBrowserUtils.h
+++ b/Engine/Source/Editor/ContentBrowser/Private/ContentBrowserUtils.h
@@ -98,10 +98,10 @@ namespace ContentBrowserUtils
 	void DisplayConfirmationPopup(const FText& Message, const FText& YesString, const FText& NoString, const TSharedRef<SWidget>& ParentContent, const FOnClicked& OnYesClicked, const FOnClicked& OnNoClicked = FOnClicked());
 
 	/** Copies all assets in all source paths to the destination path, preserving path structure */
-	void CopyFolders(const TArray<FString>& InSourcePathNames, const FString& DestPath);
+	bool CopyFolders(const TArray<FString>& InSourcePathNames, const FString& DestPath);
 
 	/** Moves all assets in all source paths to the destination path, preserving path structure */
-	void MoveFolders(const TArray<FString>& InSourcePathNames, const FString& DestPath);
+	bool MoveFolders(const TArray<FString>& InSourcePathNames, const FString& DestPath);
 
 	/**
 	  * A helper function for folder drag/drop which loads all assets in a path (including sub-paths) and returns the assets found
@@ -109,7 +109,7 @@ namespace ContentBrowserUtils
 	  * @param SourcePathNames				The paths to the folders to drag/drop
 	  * @param OutSourcePathToLoadedAssets	The map of source folder paths to assets found
 	  */
-	void PrepareFoldersForDragDrop(const TArray<FString>& SourcePathNames, TMap< FString, TArray<UObject*> >& OutSourcePathToLoadedAssets);
+	bool PrepareFoldersForDragDrop(const TArray<FString>& SourcePathNames, TMap< FString, TArray<UObject*> >& OutSourcePathToLoadedAssets);
 
 	/** Copies references to the specified assets to the clipboard */
 	void CopyAssetReferencesToClipboard(const TArray<FAssetData>& AssetsToCopy);
diff --git a/Engine/Source/Editor/ContentBrowser/Private/SPathView.cpp b/Engine/Source/Editor/ContentBrowser/Private/SPathView.cpp
--- a/Engine/Source/Editor/ContentBrowser/Private/SPathView.cpp
+++ b/Engine/Source/Editor/ContentBrowser/Private/SPathView.cpp
@@ -1308,7 +1308,10 @@ void SPathView::ExecuteTreeDropMove(TArray<FAssetData> AssetList, FString Destin
 
 void SPathView::ExecuteTreeDropCopyFolder(TArray<FString> PathNames, FString DestinationPath)
 {
-	ContentBrowserUtils::CopyFolders(PathNames, DestinationPath);
+	if (!ContentBrowserUtils::CopyFolders(PathNames, DestinationPath))
+	{
+		return;
+	}
 
 	TSharedPtr<FTreeItem> RootItem = FindItemRecursive(DestinationPath);
 	if (RootItem.IsValid())
@@ -1334,7 +1337,10 @@ void SPathView::ExecuteTreeDropCopyFolder(TArray<FString> PathNames, FString Des
 
 void SPathView::ExecuteTreeDropMoveFolder(TArray<FString> PathNames, FString DestinationPath)
 {
-	ContentBrowserUtils::MoveFolders(PathNames, DestinationPath);
+	if (!ContentBrowserUtils::MoveFolders(PathNames, DestinationPath))
+	{
+		return;
+	}
 
 	TSharedPtr<FTreeItem> RootItem = FindItemRecursive(DestinationPath);
 	if (RootItem.IsValid())
