support to override max asset length with cvar MaxAssetFullPath

--------
Integrated using branch UE4-To-UE4-Orion (reversed) of change#2557531 by [REDACTED_PERSON] on 2015/05/19 15:24:29.

diff --git a/Engine/Source/Editor/ContentBrowser/Private/ContentBrowserUtils.cpp b/Engine/Source/Editor/ContentBrowser/Private/ContentBrowserUtils.cpp
--- a/Engine/Source/Editor/ContentBrowser/Private/ContentBrowserUtils.cpp
+++ b/Engine/Source/Editor/ContentBrowser/Private/ContentBrowserUtils.cpp
@@ -1504,6 +1504,9 @@ FText ContentBrowserUtils::GetExploreFolderText()
 	return FText::Format(NSLOCTEXT("GenericPlatform", "ShowInFileManager", "Show In {FileManagerName}"), Args);
 }
 
+static const auto CVarMaxFullPathLength = 
+	IConsoleManager::Get().RegisterConsoleVariable( TEXT("MaxAssetFullPath"), PLATFORM_MAX_FILEPATH_LENGTH, TEXT("Maximum full path name of an asset.") )->AsVariableInt();
+
 bool ContentBrowserUtils::IsValidObjectPathForCreate(const FString& ObjectPath, FText& OutErrorMessage, bool bAllowExistingAsset)
 {
 	const FString ObjectName = FPackageName::ObjectPathToObjectName(ObjectPath);
@@ -1541,7 +1544,7 @@ bool ContentBrowserUtils::IsValidObjectPathForCreate(const FString& ObjectPath,
 	// Make sure we are not creating an path that is too long for the OS
 	const FString RelativePathFilename = FPackageName::LongPackageNameToFilename(PackageName, FPackageName::GetAssetPackageExtension());	// full relative path with name + extension
 	const FString FullPath = FPaths::ConvertRelativePathToFull(RelativePathFilename);	// path to file on disk
-	if ( ObjectPath.Len() > (PLATFORM_MAX_FILEPATH_LENGTH - MAX_CLASS_NAME_LENGTH) || FullPath.Len() > PLATFORM_MAX_FILEPATH_LENGTH )
+	if ( ObjectPath.Len() > (PLATFORM_MAX_FILEPATH_LENGTH - MAX_CLASS_NAME_LENGTH) || FullPath.Len() > CVarMaxFullPathLength->GetValueOnGameThread() )
 	{
 		// The full path for the asset is too long
 		OutErrorMessage = FText::Format( LOCTEXT("AssetPathTooLong", 
