Added the ability to filter collections

UETOOL-332 - Collections 2.0
UETOOL-368 - Search to find specific collections in the folder view

diff --git a/Engine/Source/Developer/CollectionManager/Public/CollectionManagerTypes.h b/Engine/Source/Developer/CollectionManager/Public/CollectionManagerTypes.h
--- a/Engine/Source/Developer/CollectionManager/Public/CollectionManagerTypes.h
+++ b/Engine/Source/Developer/CollectionManager/Public/CollectionManagerTypes.h
@@ -2,7 +2,7 @@
 
 #pragma once
 
-namespace ECollectionShareType
+struct ECollectionShareType
 {
 	enum Type
 	{
@@ -13,7 +13,57 @@ namespace ECollectionShareType
 
 		CST_All
 	};
-}
+
+	static Type FromString(const TCHAR* InString, const Type ReturnIfConversionFails = CST_All)
+	{
+		if(FPlatformString::Stricmp(InString, TEXT("System")) == 0)		return CST_System;
+		if(FPlatformString::Stricmp(InString, TEXT("Local")) == 0)		return CST_Local;
+		if(FPlatformString::Stricmp(InString, TEXT("Private")) == 0)	return CST_Private;
+		if(FPlatformString::Stricmp(InString, TEXT("Shared")) == 0)		return CST_Shared;
+		if(FPlatformString::Stricmp(InString, TEXT("All")) == 0)		return CST_All;
+		return ReturnIfConversionFails;
+	}
+
+	static const TCHAR* ToString(const Type InType)
+	{
+		switch(InType)
+		{
+		case CST_System:
+			return TEXT("System");
+		case CST_Local:
+			return TEXT("Local");
+		case CST_Private:
+			return TEXT("Private");
+		case CST_Shared:
+			return TEXT("Shared");
+		case CST_All:
+			return TEXT("All");
+		default:
+			break;
+		}
+		return TEXT("");
+	}
+
+	static FText ToText(const Type InType)
+	{
+		switch(InType)
+		{
+		case CST_System:
+			return NSLOCTEXT("ECollectionShareType", "CST_System", "System");
+		case CST_Local:
+			return NSLOCTEXT("ECollectionShareType", "CST_Local", "Local");
+		case CST_Private:
+			return NSLOCTEXT("ECollectionShareType", "CST_Private", "Private");
+		case CST_Shared:
+			return NSLOCTEXT("ECollectionShareType", "CST_Shared", "Shared");
+		case CST_All:
+			return NSLOCTEXT("ECollectionShareType", "CST_All", "All");
+		default:
+			break;
+		}
+		return FText::GetEmpty();
+	}
+};
 
 /** A name/type pair to uniquely identify a collection */
 struct FCollectionNameType
diff --git a/Engine/Source/Editor/ContentBrowser/Private/SCollectionView.cpp b/Engine/Source/Editor/ContentBrowser/Private/SCollectionView.cpp
--- a/Engine/Source/Editor/ContentBrowser/Private/SCollectionView.cpp
+++ b/Engine/Source/Editor/ContentBrowser/Private/SCollectionView.cpp
@@ -8,6 +8,7 @@
 #include "SourcesViewWidgets.h"
 #include "ContentBrowserModule.h"
 #include "SExpandableArea.h"
+#include "SSearchBox.h"
 
 #define LOCTEXT_NAMESPACE "ContentBrowser"
 
@@ -35,23 +36,25 @@ void SCollectionView::Construct( const FArguments& InArgs )
 	PreventSelectionChangedDelegateCount = 0;
 
 	TSharedRef< SWidget > HeaderContent = SNew(SHorizontalBox)
-			+SHorizontalBox::Slot()
-			.VAlign(VAlign_Center)
-			.FillWidth(1.f)
+			+ SHorizontalBox::Slot()
+			.FillWidth(1.0f)
+			.Padding(6.0f, 0.0f, 0.0f, 0.0f)
 			[
-				SNew(STextBlock)
-				.Font( FEditorStyle::GetFontStyle("ContentBrowser.SourceTitleFont") )
-				.Text( LOCTEXT("CollectionsListTitle", "Collections") )
+				SAssignNew(SearchBoxPtr, SSearchBox)
+				.HintText( LOCTEXT( "CollectionsViewSearchBoxHint", "Search Collections" ) )
+				.OnTextChanged( this, &SCollectionView::ApplyCollectionsSearchFilter )
 			]
 
 			+SHorizontalBox::Slot()
 			.AutoWidth()
-			.VAlign(VAlign_Bottom)
+			.VAlign(VAlign_Center)
+			.Padding(2.0f, 0.0f, 0.0f, 0.0f)
 			[
 				SNew(SButton)
+				.ButtonStyle(FEditorStyle::Get(), "FlatButton")
 				.ToolTipText(LOCTEXT("AddCollectionButtonTooltip", "Add a collection."))
 				.OnClicked(this, &SCollectionView::MakeAddCollectionMenu)
-				.ContentPadding(0)
+				.ContentPadding( FMargin(2, 2) )
 				.Visibility(this, &SCollectionView::GetAddCollectionButtonVisibility)
 				[
 					SNew(SImage) .Image( FEditorStyle::GetBrush("ContentBrowser.AddCollectionButtonIcon") )
@@ -71,7 +74,7 @@ void SCollectionView::Construct( const FArguments& InArgs )
 			.FillHeight(1.f)
 			[
 				SAssignNew(CollectionListPtr, SListView< TSharedPtr<FCollectionItem> >)
-				.ListItemsSource(&CollectionItems)
+				.ListItemsSource(&FilteredCollectionItems)
 				.OnGenerateRow( this, &SCollectionView::GenerateCollectionRow )
 				.ItemHeight(18)
 				.SelectionMode(ESelectionMode::Multi)
@@ -88,6 +91,7 @@ void SCollectionView::Construct( const FArguments& InArgs )
 		Content = SAssignNew(CollectionsExpandableAreaPtr, SExpandableArea)
 			.MaxHeight(200)
 			.BorderImage( FEditorStyle::GetBrush("NoBorder") )
+			.HeaderPadding( FMargin(4.0f, 0.0f, 0.0f, 0.0f) )
 			.HeaderContent()
 			[
 				HeaderContent
@@ -100,17 +104,17 @@ void SCollectionView::Construct( const FArguments& InArgs )
 	else
 	{
 		Content = SNew( SVerticalBox )
-			+SVerticalBox::Slot()
-			.AutoHeight()
-			[
-				HeaderContent
+		+SVerticalBox::Slot()
+		.AutoHeight()
+		[
+			HeaderContent
 		]
 
-			+SVerticalBox::Slot()
-			.AutoHeight()
-			[
-				BodyContent
-	];
+		+SVerticalBox::Slot()
+		.AutoHeight()
+		[
+			BodyContent
+		];
 	}
 
 	ChildSlot
@@ -124,19 +128,16 @@ void SCollectionView::Construct( const FArguments& InArgs )
 void SCollectionView::HandleCollectionCreated( const FCollectionNameType& Collection )
 {
 	UpdateCollectionItems();
-	CollectionListPtr->RequestListRefresh();
 }
 
 void SCollectionView::HandleCollectionRenamed( const FCollectionNameType& OriginalCollection, const FCollectionNameType& NewCollection )
 {
 	UpdateCollectionItems();
-	CollectionListPtr->RequestListRefresh();
 }
 
 void SCollectionView::HandleCollectionDestroyed( const FCollectionNameType& Collection )
 {
 	UpdateCollectionItems();
-	CollectionListPtr->RequestListRefresh();
 }
 
 void SCollectionView::UpdateCollectionItems()
@@ -168,6 +169,41 @@ void SCollectionView::UpdateCollectionItems()
 	}
 
 	CollectionItems.Sort( FCollectionItem::FCompareFCollectionItemByName() );
+
+	ApplyCollectionsSearchFilter();
+}
+
+void SCollectionView::ApplyCollectionsSearchFilter()
+{
+	ApplyCollectionsSearchFilter( SearchBoxPtr->GetText() );
+}
+
+void SCollectionView::ApplyCollectionsSearchFilter( const FText& InSearchText )
+{
+	if ( InSearchText.IsEmptyOrWhitespace() )
+	{
+		FilteredCollectionItems = CollectionItems;
+	}
+	else
+	{
+		const FString& SearchTextString = InSearchText.ToString();
+
+		FilteredCollectionItems.Reset();
+		for ( const auto& CollectionItem : CollectionItems )
+		{
+			if ( CollectionItem->CollectionName.Contains(SearchTextString) )
+			{
+				FilteredCollectionItems.Add(CollectionItem);
+			}
+		}
+	}
+
+	CollectionListPtr->RequestListRefresh();
+}
+
+FText SCollectionView::GetFilterText() const
+{
+	return SearchBoxPtr->GetText();
 }
 
 void SCollectionView::SetSelectedCollections(const TArray<FCollectionNameType>& CollectionsToSelect)
@@ -365,12 +401,16 @@ void SCollectionView::CreateCollectionItem( ECollectionShareType::Type Collectio
 		CollectionManagerModule.Get().CreateUniqueCollectionName(BaseCollectionName, CollectionType, CollectionName);
 		TSharedPtr<FCollectionItem> NewItem = MakeShareable(new FCollectionItem(CollectionName.ToString(), CollectionType));
 
+		// Adding a new collection now, so clear any filter we may have applied
+		SearchBoxPtr->SetText(FText::GetEmpty());
+
 		// Mark the new collection for rename and that it is new so it will be created upon successful rename
 		NewItem->bRenaming = true;
 		NewItem->bNewCollection = true;
 
 		CollectionItems.Add( NewItem );
 		CollectionItems.Sort( FCollectionItem::FCompareFCollectionItemByName() );
+		ApplyCollectionsSearchFilter();
 		CollectionListPtr->RequestScrollIntoView(NewItem);
 		CollectionListPtr->SetSelection( NewItem );
 	}
@@ -423,7 +463,7 @@ void SCollectionView::RemoveCollectionItems( const TArray<TSharedPtr<FCollection
 	}
 
 	// Refresh the list
-	CollectionListPtr->RequestListRefresh();
+	ApplyCollectionsSearchFilter();
 }
 
 EVisibility SCollectionView::GetCollectionListVisibility() const
@@ -446,6 +486,7 @@ TSharedRef<ITableRow> SCollectionView::GenerateCollectionRow( TSharedPtr<FCollec
 			.OnAssetsDragDropped(this, &SCollectionView::CollectionAssetsDropped)
 			.IsSelected( TableRow.Get(), &STableRow< TSharedPtr<FTreeItem> >::IsSelectedExclusively )
 			.IsReadOnly(this, &SCollectionView::IsCollectionNotRenamable)
+			.HighlightText( this, &SCollectionView::GetFilterText)
 		);
 
 	return TableRow.ToSharedRef();
@@ -532,7 +573,7 @@ bool SCollectionView::CollectionNameChangeCommit( const TSharedPtr< FCollectionI
 		if ( !bChangeConfirmed )
 		{
 			CollectionItems.Remove(CollectionItem);
-			CollectionListPtr->RequestListRefresh();
+			ApplyCollectionsSearchFilter();
 			return false;
 		}
 
@@ -540,7 +581,7 @@ bool SCollectionView::CollectionNameChangeCommit( const TSharedPtr< FCollectionI
 		{
 			// Failed to add the collection, remove it from the list
 			CollectionItems.Remove(CollectionItem);
-			CollectionListPtr->RequestListRefresh();
+			ApplyCollectionsSearchFilter();
 
 			OutWarningMessage = FText::Format( LOCTEXT("CreateCollectionFailed", "Failed to create the collection. {0}"), CollectionManagerModule.Get().GetLastError());
 			return false;
@@ -554,6 +595,12 @@ bool SCollectionView::CollectionNameChangeCommit( const TSharedPtr< FCollectionI
 			return true;
 		}
 
+		// If the new name doesn't pass our current filter, we need to clear it
+		if ( !NewNameFinal.Contains(SearchBoxPtr->GetText().ToString()) )
+		{
+			SearchBoxPtr->SetText(FText::GetEmpty());
+		}
+
 		// Otherwise perform the rename
 		if ( !CollectionManagerModule.Get().RenameCollection(FName(*CollectionItem->CollectionName), CollectionType, FName(*NewNameFinal), CollectionType) )
 		{
diff --git a/Engine/Source/Editor/ContentBrowser/Private/SCollectionView.h b/Engine/Source/Editor/ContentBrowser/Private/SCollectionView.h
--- a/Engine/Source/Editor/ContentBrowser/Private/SCollectionView.h
+++ b/Engine/Source/Editor/ContentBrowser/Private/SCollectionView.h
@@ -115,6 +115,16 @@ private:
 
 	/** Updates the collections shown in the list view */
 	void UpdateCollectionItems();
+
+	/** Update the visible collections based on the active search text */
+	void ApplyCollectionsSearchFilter();
+
+	/** Update the visible collections based on the given search text */
+	void ApplyCollectionsSearchFilter( const FText& InSearchText );
+
+	/** Get the currently filter text */
+	FText GetFilterText() const;
+
 private:
 
 	/** A helper class to manage PreventSelectionChangedDelegateCount by incrementing it when constructed (on the stack) and decrementing when destroyed */
@@ -137,12 +147,18 @@ private:
 		TSharedRef<SCollectionView> CollectionView;
 	};
 
+	/** The collection list search box */
+	TSharedPtr< SSearchBox > SearchBoxPtr;
+
 	/** The collection list widget */
 	TSharedPtr< SListView< TSharedPtr<FCollectionItem>> > CollectionListPtr;
 
-	/** The list of collections */
+	/** The list of available collections */
 	TArray< TSharedPtr<FCollectionItem> > CollectionItems;
 
+	/** The list of visible collections based on the current filter */
+	TArray< TSharedPtr<FCollectionItem> > FilteredCollectionItems;
+
 	/** The context menu logic and data */
 	TSharedPtr<class FCollectionContextMenu> CollectionContextMenu;
 
diff --git a/Engine/Source/Editor/ContentBrowser/Private/SourcesViewWidgets.cpp b/Engine/Source/Editor/ContentBrowser/Private/SourcesViewWidgets.cpp
--- a/Engine/Source/Editor/ContentBrowser/Private/SourcesViewWidgets.cpp
+++ b/Engine/Source/Editor/ContentBrowser/Private/SourcesViewWidgets.cpp
@@ -374,6 +374,7 @@ void SCollectionListItem::Construct( const FArguments& InArgs )
 			[
 				SAssignNew(InlineRenameWidget, SInlineEditableTextBlock)
 					.Text( this, &SCollectionListItem::GetNameText )
+					.HighlightText( InArgs._HighlightText )
 					.Font( FEditorStyle::GetFontStyle("ContentBrowser.SourceListItemFont") )
 					.OnBeginTextEdit(this, &SCollectionListItem::HandleBeginNameChange)
 					.OnTextCommitted(this, &SCollectionListItem::HandleNameCommitted)
diff --git a/Engine/Source/Editor/ContentBrowser/Private/SourcesViewWidgets.h b/Engine/Source/Editor/ContentBrowser/Private/SourcesViewWidgets.h
--- a/Engine/Source/Editor/ContentBrowser/Private/SourcesViewWidgets.h
+++ b/Engine/Source/Editor/ContentBrowser/Private/SourcesViewWidgets.h
@@ -187,6 +187,9 @@ public:
 		/** True if the item is read-only. It will not be able to be renamed if read-only */
 		SLATE_ATTRIBUTE( bool, IsReadOnly )
 
+		/** Text to highlight for this item */
+		SLATE_ATTRIBUTE( FText, HighlightText )
+
 	SLATE_END_ARGS()
 
 	/** Constructs this widget with InArgs */
diff --git a/Engine/Source/Editor/EditorStyle/Private/SlateEditorStyle.cpp b/Engine/Source/Editor/EditorStyle/Private/SlateEditorStyle.cpp
--- a/Engine/Source/Editor/EditorStyle/Private/SlateEditorStyle.cpp
+++ b/Engine/Source/Editor/EditorStyle/Private/SlateEditorStyle.cpp
@@ -5395,7 +5395,7 @@ void FSlateEditorStyle::FStyle::SetupContentBrowserStyle()
 		Set( "ContentBrowser.AssetTreeFolderDeveloper", new IMAGE_BRUSH( "Icons/FolderDeveloper", FVector2D(18, 16) ) );
 		Set( "ContentBrowser.AssetTreeFolderOpenCode", new IMAGE_BRUSH( "Icons/FolderOpen_Code", FVector2D(18, 16) ) );
 		Set( "ContentBrowser.AssetTreeFolderClosedCode", new IMAGE_BRUSH( "Icons/FolderClosed_Code", FVector2D(18, 16) ) );
-		Set( "ContentBrowser.AddCollectionButtonIcon", new IMAGE_BRUSH( "Icons/PlusSymbol_12x", Icon12x12, FLinearColor::Black ) );
+		Set( "ContentBrowser.AddCollectionButtonIcon", new IMAGE_BRUSH( "Icons/PlusSymbol_12x", Icon12x12 ) );
 
 		Set( "ContentBrowser.Splitter", FSplitterStyle()
 			.SetHandleNormalBrush( FSlateNoResource() )
