No longer show the collections filter when the collections list is collapsed

UETOOL-332 - Collections 2.0
UETOOL-368 - Search to find specific collections in the folder view

diff --git a/Engine/Source/Editor/ContentBrowser/Private/SCollectionView.cpp b/Engine/Source/Editor/ContentBrowser/Private/SCollectionView.cpp
--- a/Engine/Source/Editor/ContentBrowser/Private/SCollectionView.cpp
+++ b/Engine/Source/Editor/ContentBrowser/Private/SCollectionView.cpp
@@ -38,11 +38,25 @@ void SCollectionView::Construct( const FArguments& InArgs )
 	TSharedRef< SWidget > HeaderContent = SNew(SHorizontalBox)
 			+ SHorizontalBox::Slot()
 			.FillWidth(1.0f)
-			.Padding(6.0f, 0.0f, 0.0f, 0.0f)
+			.Padding(0.0f)
 			[
-				SAssignNew(SearchBoxPtr, SSearchBox)
-				.HintText( LOCTEXT( "CollectionsViewSearchBoxHint", "Search Collections" ) )
-				.OnTextChanged( this, &SCollectionView::ApplyCollectionsSearchFilter )
+				SNew(SHorizontalBox)
+
+				+ SHorizontalBox::Slot()
+				[
+					SNew(STextBlock)
+					.Font( FEditorStyle::GetFontStyle("ContentBrowser.SourceTitleFont") )
+					.Text( LOCTEXT("CollectionsListTitle", "Collections") )
+					.Visibility( this, &SCollectionView::GetCollectionsTitleTextVisibility )
+				]
+
+				+ SHorizontalBox::Slot()
+				[
+					SAssignNew(SearchBoxPtr, SSearchBox)
+					.HintText( LOCTEXT( "CollectionsViewSearchBoxHint", "Search Collections" ) )
+					.OnTextChanged( this, &SCollectionView::ApplyCollectionsSearchFilter )
+					.Visibility( this, &SCollectionView::GetCollectionsSearchBoxVisibility )
+				]
 			]
 
 			+SHorizontalBox::Slot()
@@ -94,7 +108,11 @@ void SCollectionView::Construct( const FArguments& InArgs )
 			.HeaderPadding( FMargin(4.0f, 0.0f, 0.0f, 0.0f) )
 			.HeaderContent()
 			[
-				HeaderContent
+				SNew(SBox)
+				.Padding(FMargin(6.0f, 0.0f, 0.0f, 0.0f))
+				[
+					HeaderContent
+				]
 			]
 			.BodyContent()
 			[
@@ -111,7 +129,6 @@ void SCollectionView::Construct( const FArguments& InArgs )
 		]
 
 		+SVerticalBox::Slot()
-		.AutoHeight()
 		[
 			BodyContent
 		];
@@ -385,6 +402,18 @@ FReply SCollectionView::MakeAddCollectionMenu()
 	return FReply::Handled();
 }
 
+EVisibility SCollectionView::GetCollectionsTitleTextVisibility() const
+{
+	// Only show the title text if we have an expansion area, but are collapsed
+	return (CollectionsExpandableAreaPtr.IsValid() && !CollectionsExpandableAreaPtr->IsExpanded()) ? EVisibility::Visible : EVisibility::Collapsed;
+}
+
+EVisibility SCollectionView::GetCollectionsSearchBoxVisibility() const
+{
+	// Only show the search box if we have an expanded expansion area, or aren't currently using an expansion area
+	return (!CollectionsExpandableAreaPtr.IsValid() || CollectionsExpandableAreaPtr->IsExpanded()) ? EVisibility::Visible : EVisibility::Collapsed;
+}
+
 EVisibility SCollectionView::GetAddCollectionButtonVisibility() const
 {
 	return (bAllowCollectionButtons && ( !CollectionsExpandableAreaPtr.IsValid() || CollectionsExpandableAreaPtr->IsExpanded() ) ) ? EVisibility::Visible : EVisibility::Collapsed;
diff --git a/Engine/Source/Editor/ContentBrowser/Private/SCollectionView.h b/Engine/Source/Editor/ContentBrowser/Private/SCollectionView.h
--- a/Engine/Source/Editor/ContentBrowser/Private/SCollectionView.h
+++ b/Engine/Source/Editor/ContentBrowser/Private/SCollectionView.h
@@ -65,6 +65,12 @@ private:
 	/** Creates the menu for the add collection button */
 	FReply MakeAddCollectionMenu();
 
+	/** Gets the visibility of the collections title text */
+	EVisibility GetCollectionsTitleTextVisibility() const;
+
+	/** Gets the visibility of the collections search box */
+	EVisibility GetCollectionsSearchBoxVisibility() const;
+
 	/** Gets the visibility of the add collection button */
 	EVisibility GetAddCollectionButtonVisibility() const;
 	
