Added quick asset management for collections in the editor

UETOOL-332 - Collections 2.0
UETOOL-370 - Ability to assign many collections to a single asset very quickly

The collections view now provides check boxes next to each collection. These can be used to quickly add or remove a large selection of assets to or from a collection.

This feature is also available in the Content Browser via the asset context menu. The "Collections" section now contains a sub-menu which lists the available collections, and allows you to quickly add or remove the current selection to or from a collection.

diff --git a/Engine/Source/Editor/ContentBrowser/Private/AssetContextMenu.cpp b/Engine/Source/Editor/ContentBrowser/Private/AssetContextMenu.cpp
--- a/Engine/Source/Editor/ContentBrowser/Private/AssetContextMenu.cpp
+++ b/Engine/Source/Editor/ContentBrowser/Private/AssetContextMenu.cpp
@@ -22,6 +22,7 @@
 #include "SourceControlWindows.h"
 #include "KismetEditorUtilities.h"
 #include "AssetToolsModule.h"
+#include "CollectionAssetManagement.h"
 #include "ComponentAssetBroker.h"
 #include "SNumericEntryBox.h"
 
@@ -106,12 +107,12 @@ TSharedRef<SWidget> FAssetContextMenu::MakeContextMenu(const TArray<FAssetData>&
 		// Add reference options
 		AddReferenceMenuOptions(MenuBuilder);
 
-		// Add documentation options
-		AddDocumentationMenuOptions(MenuBuilder);
-
 		// Add collection options
 		AddCollectionMenuOptions(MenuBuilder);
 
+		// Add documentation options
+		AddDocumentationMenuOptions(MenuBuilder);
+
 		// Add source control options
 		AddSourceControlMenuOptions(MenuBuilder);
 	}
@@ -807,27 +808,124 @@ bool FAssetContextMenu::CanExecuteSourceControlActions() const
 
 bool FAssetContextMenu::AddCollectionMenuOptions(FMenuBuilder& MenuBuilder)
 {
-	// "Remove from collection" (only display option if exactly one collection is selected)
-	if ( SourcesData.Collections.Num() == 1 )
+	class FManageCollectionsContextMenu
 	{
-		MenuBuilder.BeginSection("AssetContextCollections", LOCTEXT("AssetCollectionOptionsMenuHeading", "Collections"));
+	public:
+		static void CreateManageCollectionsSubMenu(FMenuBuilder& SubMenuBuilder, TSharedRef<FCollectionAssetManagement> QuickAssetManagement)
 		{
-			MenuBuilder.AddMenuEntry(
-				LOCTEXT("RemoveFromCollection", "Remove From Collection"),
-				LOCTEXT("RemoveFromCollection_ToolTip", "Removes the selected asset from the current collection."),
-				FSlateIcon(),
-				FUIAction(
-					FExecuteAction::CreateSP( this, &FAssetContextMenu::ExecuteRemoveFromCollection ),
-					FCanExecuteAction::CreateSP( this, &FAssetContextMenu::CanExecuteRemoveFromCollection )
-					)
-				);
+			static const FName CollectionManagerModuleName = "CollectionManager";
+			FCollectionManagerModule& CollectionManagerModule = FModuleManager::LoadModuleChecked<FCollectionManagerModule>(CollectionManagerModuleName);
+
+			// Get collections of all types
+			TArray<FCollectionNameType> AvailableCollections;
+			for (int32 TypeIdx = 0; TypeIdx < ECollectionShareType::CST_All; ++TypeIdx)
+			{
+				const ECollectionShareType::Type CollectionType = ECollectionShareType::Type(TypeIdx);
+
+				// Never display system collections
+				if (CollectionType == ECollectionShareType::CST_System)
+				{
+					continue;
+				}
+
+				TArray<FName> CollectionNames;
+				CollectionManagerModule.Get().GetCollectionNames(CollectionType, CollectionNames);
+
+				for (const FName& CollectionName : CollectionNames)
+				{
+					AvailableCollections.Add(FCollectionNameType(CollectionName, CollectionType));
+				}
+			}
+
+			AvailableCollections.Sort([](const FCollectionNameType& One, const FCollectionNameType& Two) -> bool
+			{
+				return One.Name < Two.Name;
+			});
+
+			for (const FCollectionNameType& AvailableCollection : AvailableCollections)
+			{
+				SubMenuBuilder.AddMenuEntry(
+					FText::FromName(AvailableCollection.Name), 
+					ECollectionShareType::GetDescription(AvailableCollection.Type), 
+					FSlateIcon(FEditorStyle::GetStyleSetName(), ECollectionShareType::GetIconStyleName(AvailableCollection.Type)), 
+					FUIAction(
+						FExecuteAction::CreateStatic(&FManageCollectionsContextMenu::OnCollectionClicked, QuickAssetManagement, AvailableCollection),
+						FCanExecuteAction::CreateStatic(&FManageCollectionsContextMenu::IsCollectionEnabled, QuickAssetManagement, AvailableCollection),
+						FGetActionCheckState::CreateStatic(&FManageCollectionsContextMenu::GetCollectionCheckState, QuickAssetManagement, AvailableCollection)
+						), 
+					NAME_None, 
+					EUserInterfaceActionType::ToggleButton
+					);
+			}
 		}
-		MenuBuilder.EndSection();
 
-		return true;
+	private:
+		static bool IsCollectionEnabled(TSharedRef<FCollectionAssetManagement> QuickAssetManagement, FCollectionNameType InCollectionKey)
+		{
+			return QuickAssetManagement->IsCollectionEnabled(InCollectionKey);
+		}
+
+		static ECheckBoxState GetCollectionCheckState(TSharedRef<FCollectionAssetManagement> QuickAssetManagement, FCollectionNameType InCollectionKey)
+		{
+			return QuickAssetManagement->GetCollectionCheckState(InCollectionKey);
+		}
+
+		static void OnCollectionClicked(TSharedRef<FCollectionAssetManagement> QuickAssetManagement, FCollectionNameType InCollectionKey)
+		{
+			// The UI actions don't give you the new check state, so we need to emulate the behavior of SCheckBox
+			// Basically, unchecked will transition to checked (adding items), and anything else will transition to unchecked (removing items)
+			if (GetCollectionCheckState(QuickAssetManagement, InCollectionKey) == ECheckBoxState::Unchecked)
+			{
+				QuickAssetManagement->AddCurrentAssetsToCollection(InCollectionKey);
+			}
+			else
+			{
+				QuickAssetManagement->RemoveCurrentAssetsFromCollection(InCollectionKey);
+			}
+		}
+	};
+
+	bool bHasAddedItems = false;
+
+	static const FName CollectionManagerModuleName = "CollectionManager";
+	FCollectionManagerModule& CollectionManagerModule = FModuleManager::LoadModuleChecked<FCollectionManagerModule>(CollectionManagerModuleName);
+
+	MenuBuilder.BeginSection("AssetContextCollections", LOCTEXT("AssetCollectionOptionsMenuHeading", "Collections"));
+
+	// Show a sub-menu that allows you to quickly add or remove the current asset selection from the available collections
+	if (CollectionManagerModule.Get().HasCollections())
+	{
+		TSharedRef<FCollectionAssetManagement> QuickAssetManagement = MakeShareable(new FCollectionAssetManagement());
+		QuickAssetManagement->SetCurrentAssets(SelectedAssets);
+
+		MenuBuilder.AddSubMenu(
+			LOCTEXT("ManageCollections", "Manage Collections"),
+			LOCTEXT("ManageCollections_ToolTip", "Manage the collections that the selected asset(s) belong to."),
+			FNewMenuDelegate::CreateStatic(&FManageCollectionsContextMenu::CreateManageCollectionsSubMenu, QuickAssetManagement)
+			);
+
+		bHasAddedItems = true;
 	}
 
-	return false;
+	// "Remove from collection" (only display option if exactly one collection is selected)
+	if ( SourcesData.Collections.Num() == 1 )
+	{
+		MenuBuilder.AddMenuEntry(
+			FText::Format(LOCTEXT("RemoveFromCollectionFmt", "Remove From {0}"), FText::FromName(SourcesData.Collections[0].Name)),
+			LOCTEXT("RemoveFromCollection_ToolTip", "Removes the selected asset from the current collection."),
+			FSlateIcon(),
+			FUIAction(
+				FExecuteAction::CreateSP( this, &FAssetContextMenu::ExecuteRemoveFromCollection ),
+				FCanExecuteAction::CreateSP( this, &FAssetContextMenu::CanExecuteRemoveFromCollection )
+				)
+			);
+
+		bHasAddedItems = true;
+	}
+
+	MenuBuilder.EndSection();
+
+	return bHasAddedItems;
 }
 
 bool FAssetContextMenu::AreImportedAssetActionsVisible() const
diff --git a/Engine/Source/Editor/ContentBrowser/Private/CollectionAssetManagement.cpp b/Engine/Source/Editor/ContentBrowser/Private/CollectionAssetManagement.cpp
new file mode 100644
--- /dev/null
+++ b/Engine/Source/Editor/ContentBrowser/Private/CollectionAssetManagement.cpp
@@ -0,0 +1,245 @@
+// Copyright 1998-2015 Epic Games, Inc. All Rights Reserved.
+
+#include "ContentBrowserPCH.h"
+#include "CollectionAssetManagement.h"
+#include "CollectionManagerModule.h"
+#include "NotificationManager.h"
+#include "SNotificationList.h"
+
+#define LOCTEXT_NAMESPACE "ContentBrowser"
+
+FCollectionAssetManagement::FCollectionAssetManagement()
+{
+	static const FName CollectionManagerModuleName = "CollectionManager";
+	FCollectionManagerModule& CollectionManagerModule = FModuleManager::Get().LoadModuleChecked<FCollectionManagerModule>(CollectionManagerModuleName);
+
+	// Register the notifications we need in order to keep things up-to-date
+	OnCollectionRenamedHandle = CollectionManagerModule.Get().OnCollectionRenamed().AddRaw(this, &FCollectionAssetManagement::HandleCollectionRenamed);
+	OnCollectionDestroyedHandle = CollectionManagerModule.Get().OnCollectionDestroyed().AddRaw(this, &FCollectionAssetManagement::HandleCollectionDestroyed);
+	OnAssetsAddedHandle = CollectionManagerModule.Get().OnAssetsAdded().AddRaw(this, &FCollectionAssetManagement::HandleAssetsAddedToCollection);
+	OnAssetsRemovedHandle = CollectionManagerModule.Get().OnAssetsRemoved().AddRaw(this, &FCollectionAssetManagement::HandleAssetsRemovedFromCollection);
+}
+
+FCollectionAssetManagement::~FCollectionAssetManagement()
+{
+	// Use Get rather than Load as we might be in the process of shutting down...
+	static const FName CollectionManagerModuleName = "CollectionManager";
+	FCollectionManagerModule* CollectionManagerModulePtr = FModuleManager::Get().GetModulePtr<FCollectionManagerModule>(CollectionManagerModuleName);
+
+	if (CollectionManagerModulePtr)
+	{
+		CollectionManagerModulePtr->Get().OnCollectionRenamed().Remove(OnCollectionRenamedHandle);
+		CollectionManagerModulePtr->Get().OnCollectionDestroyed().Remove(OnCollectionDestroyedHandle);
+		CollectionManagerModulePtr->Get().OnAssetsAdded().Remove(OnAssetsAddedHandle);
+		CollectionManagerModulePtr->Get().OnAssetsRemoved().Remove(OnAssetsRemovedHandle);
+	}
+}
+
+void FCollectionAssetManagement::SetCurrentAssets(const TArray<FAssetData>& CurrentAssets)
+{
+	CurrentAssetPaths.Empty();
+	for (const FAssetData& AssetData : CurrentAssets)
+	{
+		CurrentAssetPaths.Add(AssetData.ObjectPath);
+	}
+
+	UpdateAssetManagementState();
+}
+
+void FCollectionAssetManagement::AddCurrentAssetsToCollection(FCollectionNameType InCollectionKey)
+{
+	static const FName CollectionManagerModuleName = "CollectionManager";
+	FCollectionManagerModule& CollectionManagerModule = FModuleManager::Get().LoadModuleChecked<FCollectionManagerModule>(CollectionManagerModuleName);
+
+	const TArray<FName> ObjectPaths = CurrentAssetPaths.Array();
+
+	FText ResultText;
+	bool bSuccess = false;
+	{
+		int32 NumAdded = 0;
+		if (CollectionManagerModule.Get().AddToCollection(InCollectionKey.Name, InCollectionKey.Type, ObjectPaths, &NumAdded))
+		{
+			bSuccess = true;
+
+			FFormatNamedArguments Args;
+			Args.Add(TEXT("Number"), NumAdded);
+			Args.Add(TEXT("CollectionName"), FText::FromName(InCollectionKey.Name));
+			ResultText = FText::Format(LOCTEXT("CollectionAssetsAdded", "Added {Number} asset(s) to {CollectionName}"), Args);
+		}
+		else
+		{
+			ResultText = CollectionManagerModule.Get().GetLastError();
+		}
+	}
+
+	if (!ResultText.IsEmpty())
+	{
+		FNotificationInfo Info(ResultText);
+		Info.bFireAndForget = true;
+		Info.bUseLargeFont = false;
+
+		TSharedPtr<SNotificationItem> Item = FSlateNotificationManager::Get().AddNotification(Info);
+		if (Item.IsValid())
+		{
+			Item->SetCompletionState((bSuccess) ? SNotificationItem::CS_Success : SNotificationItem::CS_Fail);
+		}
+	}
+}
+
+void FCollectionAssetManagement::RemoveCurrentAssetsFromCollection(FCollectionNameType InCollectionKey)
+{
+	static const FName CollectionManagerModuleName = "CollectionManager";
+	FCollectionManagerModule& CollectionManagerModule = FModuleManager::Get().LoadModuleChecked<FCollectionManagerModule>(CollectionManagerModuleName);
+
+	const TArray<FName> ObjectPaths = CurrentAssetPaths.Array();
+
+	FText ResultText;
+	bool bSuccess = false;
+	{
+		int32 NumRemoved = 0;
+		if (CollectionManagerModule.Get().RemoveFromCollection(InCollectionKey.Name, InCollectionKey.Type, ObjectPaths, &NumRemoved))
+		{
+			bSuccess = true;
+
+			FFormatNamedArguments Args;
+			Args.Add(TEXT("Number"), NumRemoved);
+			Args.Add(TEXT("CollectionName"), FText::FromName(InCollectionKey.Name));
+			ResultText = FText::Format(LOCTEXT("CollectionAssetsRemoved", "Removed {Number} asset(s) from {CollectionName}"), Args);
+		}
+		else
+		{
+			ResultText = CollectionManagerModule.Get().GetLastError();
+		}
+	}
+
+	if (!ResultText.IsEmpty())
+	{
+		FNotificationInfo Info(ResultText);
+		Info.bFireAndForget = true;
+		Info.bUseLargeFont = false;
+
+		TSharedPtr<SNotificationItem> Item = FSlateNotificationManager::Get().AddNotification(Info);
+		if (Item.IsValid())
+		{
+			Item->SetCompletionState((bSuccess) ? SNotificationItem::CS_Success : SNotificationItem::CS_Fail);
+		}
+	}
+}
+
+bool FCollectionAssetManagement::IsCollectionEnabled(FCollectionNameType InCollectionKey) const
+{
+	// Non-local collections can only be changed if we have an available source control connection
+	const bool bCollectionWritable = (InCollectionKey.Type == ECollectionShareType::CST_Local) || (ISourceControlModule::Get().IsEnabled() && ISourceControlModule::Get().GetProvider().IsAvailable());
+	return bCollectionWritable && CurrentAssetPaths.Num() > 0;
+}
+
+ECheckBoxState FCollectionAssetManagement::GetCollectionCheckState(FCollectionNameType InCollectionKey) const
+{
+	// If the collection exists in the map, then it means that the current selection contains at least one asset using that collection
+	const ECheckBoxState* FoundCheckState = AssetManagementState.Find(InCollectionKey);
+	if (FoundCheckState)
+	{
+		return *FoundCheckState;
+	}
+
+	// If the collection doesn't exist in the map, then it's assumed to be unused by the current selection (and thus unchecked)
+	return ECheckBoxState::Unchecked;
+}
+
+void FCollectionAssetManagement::UpdateAssetManagementState()
+{
+	static const FName CollectionManagerModuleName = "CollectionManager";
+	FCollectionManagerModule& CollectionManagerModule = FModuleManager::Get().LoadModuleChecked<FCollectionManagerModule>(CollectionManagerModuleName);
+
+	AssetManagementState.Empty();
+
+	if (CurrentAssetPaths.Num() == 0)
+	{
+		return;
+	}
+
+	// The logic below is much simpler when only a single object is selected as we don't need to deal with set intersection
+	if (CurrentAssetPaths.Num() == 1)
+	{
+		TArray<FCollectionNameType> MatchedCollections;
+		CollectionManagerModule.Get().GetCollectionsContainingObject(*CurrentAssetPaths.CreateConstIterator(), MatchedCollections);
+
+		for (const FCollectionNameType& CollectionKey : MatchedCollections)
+		{
+			AssetManagementState.Add(CollectionKey, ECheckBoxState::Checked);
+		}
+	}
+	else
+	{
+		const TArray<FName> ObjectPaths = CurrentAssetPaths.Array();
+
+		TMap<FCollectionNameType, TArray<FName>> CollectionsAndMatchedObjects;
+		CollectionManagerModule.Get().GetCollectionsContainingObjects(ObjectPaths, CollectionsAndMatchedObjects);
+
+		for (const auto& MatchedCollection : CollectionsAndMatchedObjects)
+		{
+			const FCollectionNameType& CollectionKey = MatchedCollection.Key;
+			const TArray<FName>& MatchedObjects = MatchedCollection.Value;
+
+			// Collections that contain all of the selected assets are shown as checked, collections that only contain some of the selected assets are shown as undetermined
+			AssetManagementState.Add(
+				CollectionKey, 
+				(MatchedObjects.Num() == CurrentAssetPaths.Num()) ? ECheckBoxState::Checked : ECheckBoxState::Undetermined
+				);
+		}
+	}
+}
+
+void FCollectionAssetManagement::HandleCollectionRenamed(const FCollectionNameType& OriginalCollection, const FCollectionNameType& NewCollection)
+{
+	if (AssetManagementState.Contains(OriginalCollection))
+	{
+		AssetManagementState.Add(NewCollection, AssetManagementState[OriginalCollection]);
+		AssetManagementState.Remove(OriginalCollection);
+	}
+}
+
+void FCollectionAssetManagement::HandleCollectionDestroyed(const FCollectionNameType& Collection)
+{
+	AssetManagementState.Remove(Collection);
+}
+
+void FCollectionAssetManagement::HandleAssetsAddedToCollection(const FCollectionNameType& Collection, const TArray<FName>& AssetsAdded)
+{
+	// Only need to update if one of the added assets belongs to our current selection set
+	bool bNeedsUpdate = false;
+	for (const FName& AssetPath : AssetsAdded)
+	{
+		if (CurrentAssetPaths.Contains(AssetPath))
+		{
+			bNeedsUpdate = true;
+			break;
+		}
+	}
+
+	if (bNeedsUpdate)
+	{
+		UpdateAssetManagementState();
+	}
+}
+
+void FCollectionAssetManagement::HandleAssetsRemovedFromCollection(const FCollectionNameType& Collection, const TArray<FName>& AssetsRemoved)
+{
+	// Only need to update if one of the removed assets belongs to our current selection set
+	bool bNeedsUpdate = false;
+	for (const FName& AssetPath : AssetsRemoved)
+	{
+		if (CurrentAssetPaths.Contains(AssetPath))
+		{
+			bNeedsUpdate = true;
+			break;
+		}
+	}
+
+	if (bNeedsUpdate)
+	{
+		UpdateAssetManagementState();
+	}
+}
+
+#undef LOCTEXT_NAMESPACE
diff --git a/Engine/Source/Editor/ContentBrowser/Private/CollectionAssetManagement.h b/Engine/Source/Editor/ContentBrowser/Private/CollectionAssetManagement.h
new file mode 100644
--- /dev/null
+++ b/Engine/Source/Editor/ContentBrowser/Private/CollectionAssetManagement.h
@@ -0,0 +1,57 @@
+// Copyright 1998-2015 Epic Games, Inc. All Rights Reserved.
+
+#pragma once
+
+/** Handles the collection management for the given assets */
+class FCollectionAssetManagement
+{
+public:
+	/** Constructor */
+	FCollectionAssetManagement();
+
+	/** Destructor */
+	~FCollectionAssetManagement();
+
+	/** Set the assets that we are currently observing and managing the collection state of */
+	void SetCurrentAssets(const TArray<FAssetData>& CurrentAssets);
+
+	/** Add the current assets to the given collection */
+	void AddCurrentAssetsToCollection(FCollectionNameType InCollectionKey);
+
+	/** Remove the current assets from the given collection */
+	void RemoveCurrentAssetsFromCollection(FCollectionNameType InCollectionKey);
+
+	/** Return whether or not the given collection should be enabled in any management UIs */
+	bool IsCollectionEnabled(FCollectionNameType InCollectionKey) const;
+
+	/** Get the check box state the given collection should use in any management UIs */
+	ECheckBoxState GetCollectionCheckState(FCollectionNameType InCollectionKey) const;
+
+private:
+	/** Update the internal state used to track the check box status for each collection */
+	void UpdateAssetManagementState();
+
+	/** Handles an on collection renamed event */
+	void HandleCollectionRenamed(const FCollectionNameType& OriginalCollection, const FCollectionNameType& NewCollection);
+
+	/** Handles an on collection destroyed event */
+	void HandleCollectionDestroyed(const FCollectionNameType& Collection);
+
+	/** Handles assets being added to a collection */
+	void HandleAssetsAddedToCollection(const FCollectionNameType& Collection, const TArray<FName>& AssetsAdded);
+
+	/** Handles assets being removed from a collection */
+	void HandleAssetsRemovedFromCollection(const FCollectionNameType& Collection, const TArray<FName>& AssetsRemoved);
+
+	/** Set of asset paths that we are currently observing and managing the collection state of */
+	TSet<FName> CurrentAssetPaths;
+
+	/** Mapping between a collection and its asset management state (based on the current assets). A missing item is assumed to be unchecked */
+	TMap<FCollectionNameType, ECheckBoxState> AssetManagementState;
+
+	/** Delegate handles */
+	FDelegateHandle OnCollectionRenamedHandle;
+	FDelegateHandle OnCollectionDestroyedHandle;
+	FDelegateHandle OnAssetsAddedHandle;
+	FDelegateHandle OnAssetsRemovedHandle;
+};
diff --git a/Engine/Source/Editor/ContentBrowser/Private/SCollectionView.cpp b/Engine/Source/Editor/ContentBrowser/Private/SCollectionView.cpp
--- a/Engine/Source/Editor/ContentBrowser/Private/SCollectionView.cpp
+++ b/Engine/Source/Editor/ContentBrowser/Private/SCollectionView.cpp
@@ -2,6 +2,7 @@
 
 #include "ContentBrowserPCH.h"
 
+#include "CollectionAssetManagement.h"
 #include "CollectionViewTypes.h"
 #include "CollectionContextMenu.h"
 #include "ObjectTools.h"
@@ -27,6 +28,11 @@ void SCollectionView::Construct( const FArguments& InArgs )
 	CollectionContextMenu = MakeShareable(new FCollectionContextMenu( SharedThis(this) ));
 	CollectionContextMenu->BindCommands(Commands);
 
+	if ( InArgs._AllowQuickAssetManagement )
+	{
+		QuickAssetManagement = MakeShareable(new FCollectionAssetManagement());
+	}
+
 	FOnContextMenuOpening CollectionListContextMenuOpening;
 	if ( InArgs._AllowContextMenu )
 	{
@@ -275,6 +281,14 @@ TArray<FCollectionNameType> SCollectionView::GetSelectedCollections() const
 	return RetArray;
 }
 
+void SCollectionView::SetSelectedAssets(const TArray<FAssetData>& SelectedAssets)
+{
+	if ( QuickAssetManagement.IsValid() )
+	{
+		QuickAssetManagement->SetCurrentAssets(SelectedAssets);
+	}
+}
+
 void SCollectionView::ApplyHistoryData ( const FHistoryData& History )
 {	
 	// Prevent the selection changed delegate because it would add more history when we are just setting a state
@@ -504,6 +518,17 @@ TSharedRef<ITableRow> SCollectionView::GenerateCollectionRow( TSharedPtr<FCollec
 {
 	check(CollectionItem.IsValid());
 
+	// Only bind the check box callbacks if we're allowed to show check boxes
+	TAttribute<bool> IsCollectionCheckBoxEnabledAttribute;
+	TAttribute<ECheckBoxState> IsCollectionCheckedAttribute;
+	FOnCheckStateChanged OnCollectionCheckStateChangedDelegate;
+	if ( QuickAssetManagement.IsValid() )
+	{
+		IsCollectionCheckBoxEnabledAttribute.Bind(TAttribute<bool>::FGetter::CreateSP(this, &SCollectionView::IsCollectionCheckBoxEnabled, CollectionItem));
+		IsCollectionCheckedAttribute.Bind(TAttribute<ECheckBoxState>::FGetter::CreateSP(this, &SCollectionView::IsCollectionChecked, CollectionItem));
+		OnCollectionCheckStateChangedDelegate.BindSP(this, &SCollectionView::OnCollectionCheckStateChanged, CollectionItem);
+	}
+
 	TSharedPtr< STableRow< TSharedPtr<FTreeItem> > > TableRow = SNew( STableRow< TSharedPtr<FTreeItem> >, OwnerTable );
 	TableRow->SetContent
 		(
@@ -516,6 +541,9 @@ TSharedRef<ITableRow> SCollectionView::GenerateCollectionRow( TSharedPtr<FCollec
 			.IsSelected( TableRow.Get(), &STableRow< TSharedPtr<FTreeItem> >::IsSelectedExclusively )
 			.IsReadOnly(this, &SCollectionView::IsCollectionNotRenamable)
 			.HighlightText( this, &SCollectionView::GetFilterText)
+			.IsCheckBoxEnabled( IsCollectionCheckBoxEnabledAttribute )
+			.IsCollectionChecked( IsCollectionCheckedAttribute )
+			.OnCollectionCheckStateChanged( OnCollectionCheckStateChangedDelegate )
 		);
 
 	return TableRow.ToSharedRef();
@@ -531,6 +559,40 @@ TSharedPtr<SWidget> SCollectionView::MakeCollectionListContextMenu()
 	return CollectionContextMenu->MakeCollectionListContextMenu(Commands);
 }
 
+bool SCollectionView::IsCollectionCheckBoxEnabled( TSharedPtr<FCollectionItem> CollectionItem ) const
+{
+	return QuickAssetManagement.IsValid() && QuickAssetManagement->IsCollectionEnabled(FCollectionNameType(*CollectionItem->CollectionName, CollectionItem->CollectionType));
+}
+
+ECheckBoxState SCollectionView::IsCollectionChecked( TSharedPtr<FCollectionItem> CollectionItem ) const
+{
+	if ( QuickAssetManagement.IsValid() )
+	{
+		return QuickAssetManagement->GetCollectionCheckState(FCollectionNameType(*CollectionItem->CollectionName, CollectionItem->CollectionType));
+	}
+	return ECheckBoxState::Unchecked;
+}
+
+void SCollectionView::OnCollectionCheckStateChanged( ECheckBoxState NewState, TSharedPtr<FCollectionItem> CollectionItem )
+{
+	if ( QuickAssetManagement.IsValid() )
+	{
+		switch(NewState)
+		{
+		case ECheckBoxState::Checked:
+			QuickAssetManagement->AddCurrentAssetsToCollection(FCollectionNameType(*CollectionItem->CollectionName, CollectionItem->CollectionType));
+			break;
+
+		case ECheckBoxState::Unchecked:
+			QuickAssetManagement->RemoveCurrentAssetsFromCollection(FCollectionNameType(*CollectionItem->CollectionName, CollectionItem->CollectionType));
+			break;
+
+		default:
+			break;
+		}
+	}
+}
+
 void SCollectionView::CollectionSelectionChanged( TSharedPtr< FCollectionItem > CollectionItem, ESelectInfo::Type /*SelectInfo*/ )
 {
 	if ( ShouldAllowSelectionChangedDelegate() && OnCollectionSelected.IsBound() )
@@ -559,9 +621,20 @@ void SCollectionView::CollectionAssetsDropped(const TArray<FAssetData>& AssetLis
 	int32 NumAdded = 0;
 	if ( CollectionManagerModule.Get().AddToCollection(FName(*CollectionItem->CollectionName), CollectionItem->CollectionType, ObjectPaths, &NumAdded) )
 	{
-		FFormatNamedArguments Args;
-		Args.Add( TEXT("Number"), NumAdded );
-		OutMessage = FText::Format( LOCTEXT("CollectionAssetsAdded", "Added {Number} asset(s)"), Args );
+		if ( AssetList.Num() == 1 )
+		{
+			FFormatNamedArguments Args;
+			Args.Add( TEXT("AssetName"), FText::FromName(AssetList[0].AssetName) );
+			Args.Add( TEXT("CollectionName"), FText::FromString(CollectionItem->CollectionName) );
+			OutMessage = FText::Format( LOCTEXT("CollectionAssetsAdded", "Added {AssetName} to {CollectionName}"), Args );
+		}
+		else
+		{
+			FFormatNamedArguments Args;
+			Args.Add( TEXT("Number"), NumAdded );
+			Args.Add( TEXT("CollectionName"), FText::FromString(CollectionItem->CollectionName) );
+			OutMessage = FText::Format( LOCTEXT("CollectionAssetsAdded", "Added {Number} asset(s) to {CollectionName}"), Args );
+		}
 	}
 	else
 	{
diff --git a/Engine/Source/Editor/ContentBrowser/Private/SCollectionView.h b/Engine/Source/Editor/ContentBrowser/Private/SCollectionView.h
--- a/Engine/Source/Editor/ContentBrowser/Private/SCollectionView.h
+++ b/Engine/Source/Editor/ContentBrowser/Private/SCollectionView.h
@@ -6,6 +6,7 @@
 
 struct FCollectionItem;
 class SCollectionListItem;
+class FCollectionAssetManagement;
 
 /**
  * The list view of collections.
@@ -21,6 +22,7 @@ public:
 		, _AllowRightClickMenu(true)
 		, _AllowCollapsing(true)
 		, _AllowContextMenu(true)
+		, _AllowQuickAssetManagement(false)
 		{}
 
 		/** Called when a collection was selected */
@@ -32,6 +34,9 @@ public:
 		SLATE_ARGUMENT( bool, AllowCollapsing )
 		SLATE_ARGUMENT( bool, AllowContextMenu )
 
+		/** If true, check boxes that let you quickly add/remove the current selection from a collection will be displayed */
+		SLATE_ARGUMENT( bool, AllowQuickAssetManagement )
+
 	SLATE_END_ARGS()
 
 	/** Constructs this widget with InArgs */
@@ -46,6 +51,9 @@ public:
 	/** Gets all the currently selected collections */
 	TArray<FCollectionNameType> GetSelectedCollections() const;
 
+	/** Let the collections view know that the list of selected assets has changed, so that it can update the quick asset management check boxes */
+	void SetSelectedAssets(const TArray<FAssetData>& SelectedAssets);
+
 	/** Sets the state of the collection view to the one described by the history data */
 	void ApplyHistoryData ( const FHistoryData& History );
 
@@ -92,6 +100,15 @@ private:
 	/** Makes the context menu for the collection list */
 	TSharedPtr<SWidget> MakeCollectionListContextMenu();
 
+	/** Whether the check box of the given collection item is currently enabled */
+	bool IsCollectionCheckBoxEnabled( TSharedPtr<FCollectionItem> CollectionItem ) const;
+
+	/** Whether the check box of the given collection item is currently in a checked state */
+	ECheckBoxState IsCollectionChecked( TSharedPtr<FCollectionItem> CollectionItem ) const;
+
+	/** Handler for when the checked state of the given collection item check box is changed */
+	void OnCollectionCheckStateChanged( ECheckBoxState NewState, TSharedPtr<FCollectionItem> CollectionItem );
+
 	/** Handler for collection list selection changes */
 	void CollectionSelectionChanged( TSharedPtr< FCollectionItem > CollectionItem, ESelectInfo::Type SelectInfo );
 
@@ -185,4 +202,7 @@ private:
 
 	/** Commands handled by this widget */
 	TSharedPtr< FUICommandList > Commands;
+
+	/** Handles the collection management for the currently selected assets (if available) */
+	TSharedPtr<FCollectionAssetManagement> QuickAssetManagement;
 };
\ No newline at end of file
diff --git a/Engine/Source/Editor/ContentBrowser/Private/SContentBrowser.cpp b/Engine/Source/Editor/ContentBrowser/Private/SContentBrowser.cpp
--- a/Engine/Source/Editor/ContentBrowser/Private/SContentBrowser.cpp
+++ b/Engine/Source/Editor/ContentBrowser/Private/SContentBrowser.cpp
@@ -454,6 +454,7 @@ void SContentBrowser::Construct( const FArguments& InArgs, const FName& InInstan
 							SAssignNew(CollectionViewPtr, SCollectionView)
 							.OnCollectionSelected(this, &SContentBrowser::CollectionSelected)
 							.AddMetaData<FTagMetaData>(FTagMetaData(TEXT("ContentBrowserCollections")))
+							.AllowQuickAssetManagement(true)
 						]
 					]
 				]
@@ -1553,6 +1554,7 @@ void SContentBrowser::OnAssetSelectionChanged(const FAssetData& SelectedAsset)
 	
 	const TArray<FAssetData>& SelectedAssets = AssetViewPtr->GetSelectedAssets();
 	AssetContextMenu->SetSelectedAssets(SelectedAssets);
+	CollectionViewPtr->SetSelectedAssets(SelectedAssets);
 	if(AssetSelectionChangedDelegate.IsBound())
 	{
 		AssetSelectionChangedDelegate.Broadcast(SelectedAssets, bIsPrimaryBrowser);
diff --git a/Engine/Source/Editor/ContentBrowser/Private/SourcesViewWidgets.cpp b/Engine/Source/Editor/ContentBrowser/Private/SourcesViewWidgets.cpp
--- a/Engine/Source/Editor/ContentBrowser/Private/SourcesViewWidgets.cpp
+++ b/Engine/Source/Editor/ContentBrowser/Private/SourcesViewWidgets.cpp
@@ -323,31 +323,6 @@ void SCollectionListItem::Construct( const FArguments& InArgs )
 	OnAssetsDragDropped = InArgs._OnAssetsDragDropped;
 	bDraggedOver = false;
 
-	FString CollectionTypeImage;
-	FText CollectionTypeTooltip;
-	switch (InArgs._CollectionItem->CollectionType)
-	{
-	case ECollectionShareType::CST_Shared:
-		CollectionTypeImage = TEXT("ContentBrowser.Shared");
-		CollectionTypeTooltip = LOCTEXT("SharedCollectionTooltip", "Shared. This collection is visible to everyone.");
-		break;
-
-	case ECollectionShareType::CST_Private:
-		CollectionTypeImage = TEXT("ContentBrowser.Private");
-		CollectionTypeTooltip = LOCTEXT("PrivateCollectionTooltip", "Private. This collection is only visible to you.");
-		break;
-
-	case ECollectionShareType::CST_Local:
-		CollectionTypeImage = TEXT("ContentBrowser.Local");
-		CollectionTypeTooltip = LOCTEXT("LocalCollectionTooltip", "Local. This collection is only visible to you and is not in source control.");
-		break;
-
-	default:
-		CollectionTypeImage = TEXT("");
-		CollectionTypeTooltip = FText::GetEmpty();
-		break;
-	}
-
 	ChildSlot
 	[
 		SNew(SBorder)
@@ -355,16 +330,28 @@ void SCollectionListItem::Construct( const FArguments& InArgs )
 		.Padding(0)
 		[
 			SNew(SHorizontalBox)
-			.ToolTip( SNew(SToolTip).Text(CollectionTypeTooltip) )
+			.ToolTip( SNew(SToolTip).Text( ECollectionShareType::GetDescription(InArgs._CollectionItem->CollectionType) ) )
 
 			+SHorizontalBox::Slot()
 			.AutoWidth()
 			.VAlign(VAlign_Center)
-			.Padding(0, 0, 4, 0)
+			.Padding(0, 0, 2, 0)
+			[
+				SNew(SCheckBox)
+				.Visibility( InArgs._IsCollectionChecked.IsSet() ? EVisibility::Visible : EVisibility::Collapsed )
+				.IsEnabled( InArgs._IsCheckBoxEnabled )
+				.IsChecked( InArgs._IsCollectionChecked )
+				.OnCheckStateChanged( InArgs._OnCollectionCheckStateChanged )
+			]
+
+			+SHorizontalBox::Slot()
+			.AutoWidth()
+			.VAlign(VAlign_Center)
+			.Padding(0, 0, 2, 0)
 			[
 				// Type Icon
 				SNew(SImage)
-				.Image( FEditorStyle::GetBrush(*CollectionTypeImage) )
+				.Image( FEditorStyle::GetBrush( ECollectionShareType::GetIconStyleName(InArgs._CollectionItem->CollectionType) ) )
 				.ColorAndOpacity( this, &SCollectionListItem::GetCollectionColor )
 			]
 
@@ -373,14 +360,14 @@ void SCollectionListItem::Construct( const FArguments& InArgs )
 			.VAlign(VAlign_Center)
 			[
 				SAssignNew(InlineRenameWidget, SInlineEditableTextBlock)
-					.Text( this, &SCollectionListItem::GetNameText )
-					.HighlightText( InArgs._HighlightText )
-					.Font( FEditorStyle::GetFontStyle("ContentBrowser.SourceListItemFont") )
-					.OnBeginTextEdit(this, &SCollectionListItem::HandleBeginNameChange)
-					.OnTextCommitted(this, &SCollectionListItem::HandleNameCommitted)
-					.OnVerifyTextChanged(this, &SCollectionListItem::HandleVerifyNameChanged)
-					.IsSelected( InArgs._IsSelected )
-					.IsReadOnly( InArgs._IsReadOnly )
+				.Text( this, &SCollectionListItem::GetNameText )
+				.HighlightText( InArgs._HighlightText )
+				.Font( FEditorStyle::GetFontStyle("ContentBrowser.SourceListItemFont") )
+				.OnBeginTextEdit(this, &SCollectionListItem::HandleBeginNameChange)
+				.OnTextCommitted(this, &SCollectionListItem::HandleNameCommitted)
+				.OnVerifyTextChanged(this, &SCollectionListItem::HandleVerifyNameChanged)
+				.IsSelected( InArgs._IsSelected )
+				.IsReadOnly( InArgs._IsReadOnly )
 			]
 		]
 	];
diff --git a/Engine/Source/Editor/ContentBrowser/Private/SourcesViewWidgets.h b/Engine/Source/Editor/ContentBrowser/Private/SourcesViewWidgets.h
--- a/Engine/Source/Editor/ContentBrowser/Private/SourcesViewWidgets.h
+++ b/Engine/Source/Editor/ContentBrowser/Private/SourcesViewWidgets.h
@@ -190,6 +190,15 @@ public:
 		/** Text to highlight for this item */
 		SLATE_ATTRIBUTE( FText, HighlightText )
 
+		/** True if the check box of the collection item is enabled */
+		SLATE_ATTRIBUTE( bool, IsCheckBoxEnabled )
+
+		/** Whether the check box of the collection item is currently in a checked state (if unset, no check box will be shown) */
+		SLATE_ATTRIBUTE( ECheckBoxState, IsCollectionChecked )
+
+		/** Delegate for when the checked state of the collection item check box is changed */
+		SLATE_EVENT( FOnCheckStateChanged, OnCollectionCheckStateChanged )
+
 	SLATE_END_ARGS()
 
 	/** Constructs this widget with InArgs */
