Added a way to change the collection share type via the UI

UETOOL-332 - Collections 2.0

You can now right click on a collection and change its share type via the context menu.

diff --git a/Engine/Source/Editor/ContentBrowser/Private/CollectionContextMenu.cpp b/Engine/Source/Editor/ContentBrowser/Private/CollectionContextMenu.cpp
--- a/Engine/Source/Editor/ContentBrowser/Private/CollectionContextMenu.cpp
+++ b/Engine/Source/Editor/ContentBrowser/Private/CollectionContextMenu.cpp
@@ -73,6 +73,13 @@ TSharedPtr<SWidget> FCollectionContextMenu::MakeCollectionListContextMenu(TShare
 
 			if ( CollectionList.Num() == 1 )
 			{
+				// Share Type
+				MenuBuilder.AddSubMenu(
+					LOCTEXT("SetCollectionShareType", "Share Type"),
+					LOCTEXT("SetCollectionShareTypeTooltip", "Change the share type of this collection."),
+					FNewMenuDelegate::CreateRaw( this, &FCollectionContextMenu::MakeCollectionShareTypeSubMenu )
+					);
+
 				// Rename
 				MenuBuilder.AddMenuEntry( FGenericCommands::Get().Rename, NAME_None, LOCTEXT("RenameCollection", "Rename"), LOCTEXT("RenameCollectionTooltip", "Rename this collection."));
 			}
@@ -125,31 +132,77 @@ void FCollectionContextMenu::MakeNewCollectionSubMenu(FMenuBuilder& MenuBuilder)
 		MenuBuilder.AddMenuEntry(
 			LOCTEXT("NewCollection_Shared", "Shared Collection"),
 			LOCTEXT("NewCollection_SharedTooltip", "Create a collection that can be seen by anyone."),
-			FSlateIcon(),
+			FSlateIcon( FEditorStyle::GetStyleSetName(), ECollectionShareType::GetIconStyleName( ECollectionShareType::CST_Shared ) ),
 			FUIAction(
-			FExecuteAction::CreateSP( this, &FCollectionContextMenu::ExecuteNewCollection, ECollectionShareType::CST_Shared ),
-			FCanExecuteAction::CreateSP( this, &FCollectionContextMenu::CanExecuteNewCollection, ECollectionShareType::CST_Shared )
-			)
+				FExecuteAction::CreateSP( this, &FCollectionContextMenu::ExecuteNewCollection, ECollectionShareType::CST_Shared ),
+				FCanExecuteAction::CreateSP( this, &FCollectionContextMenu::CanExecuteNewCollection, ECollectionShareType::CST_Shared )
+				)
 			);
 
 		MenuBuilder.AddMenuEntry(
 			LOCTEXT("NewCollection_Private", "Private Collection"),
 			LOCTEXT("NewCollection_PrivateTooltip", "Create a collection that can only be seen by you."),
-			FSlateIcon(),
+			FSlateIcon( FEditorStyle::GetStyleSetName(), ECollectionShareType::GetIconStyleName( ECollectionShareType::CST_Private ) ),
 			FUIAction(
-			FExecuteAction::CreateSP( this, &FCollectionContextMenu::ExecuteNewCollection, ECollectionShareType::CST_Private ),
-			FCanExecuteAction::CreateSP( this, &FCollectionContextMenu::CanExecuteNewCollection, ECollectionShareType::CST_Private )
-			)
+				FExecuteAction::CreateSP( this, &FCollectionContextMenu::ExecuteNewCollection, ECollectionShareType::CST_Private ),
+				FCanExecuteAction::CreateSP( this, &FCollectionContextMenu::CanExecuteNewCollection, ECollectionShareType::CST_Private )
+				)
 			);
 
 		MenuBuilder.AddMenuEntry(
 			LOCTEXT("NewCollection_Local", "Local Collection"),
 			LOCTEXT("NewCollection_LocalTooltip", "Create a collection that is not in source control and can only be seen by you."),
-			FSlateIcon(),
+			FSlateIcon( FEditorStyle::GetStyleSetName(), ECollectionShareType::GetIconStyleName( ECollectionShareType::CST_Local ) ),
+			FUIAction(
+				FExecuteAction::CreateSP( this, &FCollectionContextMenu::ExecuteNewCollection, ECollectionShareType::CST_Local ),
+				FCanExecuteAction::CreateSP( this, &FCollectionContextMenu::CanExecuteNewCollection, ECollectionShareType::CST_Local )
+				)
+			);
+	}
+	MenuBuilder.EndSection();
+}
+
+void FCollectionContextMenu::MakeCollectionShareTypeSubMenu(FMenuBuilder& MenuBuilder)
+{
+	MenuBuilder.BeginSection("CollectionShareType", LOCTEXT("CollectionShareTypeMenuHeading", "Collection Share Type"));
+	{
+		MenuBuilder.AddMenuEntry(
+			ECollectionShareType::ToText( ECollectionShareType::CST_Shared ),
+			ECollectionShareType::GetDescription( ECollectionShareType::CST_Shared ),
+			FSlateIcon( FEditorStyle::GetStyleSetName(), ECollectionShareType::GetIconStyleName( ECollectionShareType::CST_Shared ) ),
+			FUIAction(
+				FExecuteAction::CreateSP( this, &FCollectionContextMenu::ExecuteSetCollectionShareType, ECollectionShareType::CST_Shared ),
+				FCanExecuteAction::CreateSP( this, &FCollectionContextMenu::CanExecuteSetCollectionShareType, ECollectionShareType::CST_Shared ),
+				FIsActionChecked::CreateSP( this, &FCollectionContextMenu::IsSetCollectionShareTypeChecked, ECollectionShareType::CST_Shared )
+				),
+			NAME_None,
+			EUserInterfaceActionType::Check
+			);
+
+		MenuBuilder.AddMenuEntry(
+			ECollectionShareType::ToText( ECollectionShareType::CST_Private ),
+			ECollectionShareType::GetDescription( ECollectionShareType::CST_Private ),
+			FSlateIcon( FEditorStyle::GetStyleSetName(), ECollectionShareType::GetIconStyleName( ECollectionShareType::CST_Private ) ),
+			FUIAction(
+				FExecuteAction::CreateSP( this, &FCollectionContextMenu::ExecuteSetCollectionShareType, ECollectionShareType::CST_Private ),
+				FCanExecuteAction::CreateSP( this, &FCollectionContextMenu::CanExecuteSetCollectionShareType, ECollectionShareType::CST_Private ),
+				FIsActionChecked::CreateSP( this, &FCollectionContextMenu::IsSetCollectionShareTypeChecked, ECollectionShareType::CST_Private )
+				),
+			NAME_None,
+			EUserInterfaceActionType::Check
+			);
+
+		MenuBuilder.AddMenuEntry(
+			ECollectionShareType::ToText( ECollectionShareType::CST_Local ),
+			ECollectionShareType::GetDescription( ECollectionShareType::CST_Local ),
+			FSlateIcon( FEditorStyle::GetStyleSetName(), ECollectionShareType::GetIconStyleName( ECollectionShareType::CST_Local ) ),
 			FUIAction(
-			FExecuteAction::CreateSP( this, &FCollectionContextMenu::ExecuteNewCollection, ECollectionShareType::CST_Local ),
-			FCanExecuteAction::CreateSP( this, &FCollectionContextMenu::CanExecuteNewCollection, ECollectionShareType::CST_Local )
-			)
+				FExecuteAction::CreateSP( this, &FCollectionContextMenu::ExecuteSetCollectionShareType, ECollectionShareType::CST_Local ),
+				FCanExecuteAction::CreateSP( this, &FCollectionContextMenu::CanExecuteSetCollectionShareType, ECollectionShareType::CST_Local ),
+				FIsActionChecked::CreateSP( this, &FCollectionContextMenu::IsSetCollectionShareTypeChecked, ECollectionShareType::CST_Local )
+				),
+			NAME_None,
+			EUserInterfaceActionType::Check
 			);
 	}
 	MenuBuilder.EndSection();
@@ -242,6 +295,27 @@ void FCollectionContextMenu::ExecuteNewCollection(ECollectionShareType::Type Col
 	CollectionView.Pin()->CreateCollectionItem(CollectionType);
 }
 
+void FCollectionContextMenu::ExecuteSetCollectionShareType(ECollectionShareType::Type CollectionType)
+{
+	if ( !ensure(CollectionView.IsValid()) )
+	{
+		return;
+	}
+
+	TArray<TSharedPtr<FCollectionItem>> CollectionList = CollectionView.Pin()->CollectionListPtr->GetSelectedItems();
+
+	if ( !ensure(CollectionList.Num() == 1) )
+	{
+		return;
+	}
+
+	static const FName CollectionManagerModuleName = "CollectionManager";
+	FCollectionManagerModule& CollectionManagerModule = FModuleManager::Get().LoadModuleChecked<FCollectionManagerModule>(CollectionManagerModuleName);
+
+	const FName CollectionName = *CollectionList[0]->CollectionName;
+	CollectionManagerModule.Get().RenameCollection(CollectionName, CollectionList[0]->CollectionType, CollectionName, CollectionType);
+}
+
 void FCollectionContextMenu::ExecuteRenameCollection()
 {
 	if ( !ensure(CollectionView.IsValid()) )
@@ -369,6 +443,45 @@ bool FCollectionContextMenu::CanExecuteNewCollection(ECollectionShareType::Type
 	return (CollectionType == ECollectionShareType::CST_Local) || (bProjectUnderSourceControl && ISourceControlModule::Get().IsEnabled() && ISourceControlModule::Get().GetProvider().IsAvailable());
 }
 
+bool FCollectionContextMenu::CanExecuteSetCollectionShareType(ECollectionShareType::Type CollectionType) const
+{
+	if ( !ensure(CollectionView.IsValid()) )
+	{
+		return false;
+	}
+
+	TArray<TSharedPtr<FCollectionItem>> CollectionList = CollectionView.Pin()->CollectionListPtr->GetSelectedItems();
+
+	if ( !ensure(CollectionList.Num() == 1) )
+	{
+		return false;
+	}
+
+	const bool bIsSourceControlAvailable = bProjectUnderSourceControl && ISourceControlModule::Get().IsEnabled() && ISourceControlModule::Get().GetProvider().IsAvailable();
+	const bool bIsCurrentTypeLocal = CollectionList[0]->CollectionType == ECollectionShareType::CST_Local;
+	const bool bIsNewTypeLocal = CollectionType == ECollectionShareType::CST_Local;
+	const bool bIsNewShareTypeDifferent = CollectionList[0]->CollectionType != CollectionType;
+
+	return bIsNewShareTypeDifferent && ((bIsCurrentTypeLocal && bIsNewTypeLocal) || bIsSourceControlAvailable);
+}
+
+bool FCollectionContextMenu::IsSetCollectionShareTypeChecked(ECollectionShareType::Type CollectionType) const
+{
+	if ( !ensure(CollectionView.IsValid()) )
+	{
+		return false;
+	}
+
+	TArray<TSharedPtr<FCollectionItem>> CollectionList = CollectionView.Pin()->CollectionListPtr->GetSelectedItems();
+
+	if ( !ensure(CollectionList.Num() == 1) )
+	{
+		return false;
+	}
+
+	return CollectionList[0]->CollectionType == CollectionType;
+}
+
 bool FCollectionContextMenu::CanExecuteRenameCollection() const
 {
 	return CanRenameSelectedCollections();
diff --git a/Engine/Source/Editor/ContentBrowser/Private/CollectionContextMenu.h b/Engine/Source/Editor/ContentBrowser/Private/CollectionContextMenu.h
--- a/Engine/Source/Editor/ContentBrowser/Private/CollectionContextMenu.h
+++ b/Engine/Source/Editor/ContentBrowser/Private/CollectionContextMenu.h
@@ -17,6 +17,9 @@ public:
 	/** Makes the new collection submenu */
 	void MakeNewCollectionSubMenu(FMenuBuilder& MenuBuilder);
 
+	/** Makes the collection share type submenu */
+	void MakeCollectionShareTypeSubMenu(FMenuBuilder& MenuBuilder);
+
 	/** Update the source control flag the 'CanExecute' functions rely on */
 	void UpdateProjectSourceControl();
 
@@ -31,6 +34,9 @@ private:
 	/** Handler for when a collection is selected in the "New" menu */
 	void ExecuteNewCollection(ECollectionShareType::Type CollectionType);
 
+	/** Handler for when a collection share type is changed in the "Share Type" menu */
+	void ExecuteSetCollectionShareType(ECollectionShareType::Type CollectionType);
+
 	/** Handler for when "Rename Collection" is selected */
 	void ExecuteRenameCollection();
 
@@ -49,6 +55,12 @@ private:
 	/** Handler to check to see if "New Collection" can be executed */
 	bool CanExecuteNewCollection(ECollectionShareType::Type CollectionType) const;
 
+	/** Handler to check to see if a entry in the "Share Type" menu can be executed */
+	bool CanExecuteSetCollectionShareType(ECollectionShareType::Type CollectionType) const;
+
+	/** Handler to check to see if an entry in the "Share Type" menu should be checked */
+	bool IsSetCollectionShareTypeChecked(ECollectionShareType::Type CollectionType) const;
+
 	/** Handler to check to see if "Rename Collection" can be executed */
 	bool CanExecuteRenameCollection() const;
 
