Converted FCollectionItem::CollectionName to be an FName rather than FString

UETOOL-332 - Collections 2.0

This avoids conversion from FString to FName when dealing with the collections manager API (which gave the collections view the data as an FName to begin with).

diff --git a/Engine/Source/Editor/ContentBrowser/Private/CollectionContextMenu.cpp b/Engine/Source/Editor/ContentBrowser/Private/CollectionContextMenu.cpp
--- a/Engine/Source/Editor/ContentBrowser/Private/CollectionContextMenu.cpp
+++ b/Engine/Source/Editor/ContentBrowser/Private/CollectionContextMenu.cpp
@@ -312,8 +312,7 @@ void FCollectionContextMenu::ExecuteSetCollectionShareType(ECollectionShareType:
 	static const FName CollectionManagerModuleName = "CollectionManager";
 	FCollectionManagerModule& CollectionManagerModule = FModuleManager::Get().LoadModuleChecked<FCollectionManagerModule>(CollectionManagerModuleName);
 
-	const FName CollectionName = *CollectionList[0]->CollectionName;
-	CollectionManagerModule.Get().RenameCollection(CollectionName, CollectionList[0]->CollectionType, CollectionName, CollectionType);
+	CollectionManagerModule.Get().RenameCollection(CollectionList[0]->CollectionName, CollectionList[0]->CollectionType, CollectionList[0]->CollectionName, CollectionType);
 }
 
 void FCollectionContextMenu::ExecuteRenameCollection()
@@ -345,7 +344,7 @@ void FCollectionContextMenu::ExecuteDestroyCollection()
 	FText Prompt;
 	if ( CollectionList.Num() == 1 )
 	{
-		Prompt = FText::Format(LOCTEXT("CollectionDestroyConfirm_Single", "Delete {0}?"), FText::FromString(CollectionList[0]->CollectionName));
+		Prompt = FText::Format(LOCTEXT("CollectionDestroyConfirm_Single", "Delete {0}?"), FText::FromName(CollectionList[0]->CollectionName));
 	}
 	else
 	{
@@ -370,7 +369,7 @@ FReply FCollectionContextMenu::ExecuteDestroyCollectionConfirmed(TArray<TSharedP
 	{
 		const TSharedPtr<FCollectionItem>& CollectionItem = CollectionList[CollectionIdx];
 
-		if ( CollectionManagerModule.Get().DestroyCollection(FName(*CollectionItem->CollectionName), CollectionItem->CollectionType) )
+		if ( CollectionManagerModule.Get().DestroyCollection(CollectionItem->CollectionName, CollectionItem->CollectionType) )
 		{
 			// Refresh the list now that the collection has been removed
 			ItemsToRemove.Add(CollectionItem);
diff --git a/Engine/Source/Editor/ContentBrowser/Private/CollectionViewTypes.h b/Engine/Source/Editor/ContentBrowser/Private/CollectionViewTypes.h
--- a/Engine/Source/Editor/ContentBrowser/Private/CollectionViewTypes.h
+++ b/Engine/Source/Editor/ContentBrowser/Private/CollectionViewTypes.h
@@ -14,7 +14,7 @@ struct FCollectionItem
 	};
 
 	/** The name of the collection */
-	FString CollectionName;
+	FName CollectionName;
 
 	/** The type of the collection */
 	ECollectionShareType::Type CollectionType;
@@ -32,7 +32,7 @@ struct FCollectionItem
 	FRenamedRequestEvent OnRenamedRequestEvent;
 
 	/** Constructor */
-	FCollectionItem(const FString& InCollectionName, const ECollectionShareType::Type& InCollectionType)
+	FCollectionItem(const FName& InCollectionName, const ECollectionShareType::Type& InCollectionType)
 		: CollectionName(InCollectionName)
 		, CollectionType(InCollectionType)
 		, bRenaming(false)
diff --git a/Engine/Source/Editor/ContentBrowser/Private/SCollectionView.cpp b/Engine/Source/Editor/ContentBrowser/Private/SCollectionView.cpp
--- a/Engine/Source/Editor/ContentBrowser/Private/SCollectionView.cpp
+++ b/Engine/Source/Editor/ContentBrowser/Private/SCollectionView.cpp
@@ -187,7 +187,7 @@ void SCollectionView::UpdateCollectionItems()
 		for (int32 CollectionIdx = 0; CollectionIdx < CollectionNames.Num(); ++CollectionIdx)
 		{
 			const FName& CollectionName = CollectionNames[CollectionIdx];
-			CollectionItems.Add( MakeShareable(new FCollectionItem(CollectionName.ToString(), CollectionType)) );
+			CollectionItems.Add( MakeShareable(new FCollectionItem(CollectionName, CollectionType)) );
 		}
 	}
 
@@ -214,7 +214,7 @@ void SCollectionView::ApplyCollectionsSearchFilter( const FText& InSearchText )
 		FilteredCollectionItems.Reset();
 		for ( const auto& CollectionItem : CollectionItems )
 		{
-			if ( CollectionItem->CollectionName.Contains(SearchTextString) )
+			if ( CollectionItem->CollectionName.ToString().Contains(SearchTextString) )
 			{
 				FilteredCollectionItems.Add(CollectionItem);
 			}
@@ -249,7 +249,7 @@ void SCollectionView::SetSelectedCollections(const TArray<FCollectionNameType>&
 		for ( auto SelectionIt = CollectionsToSelect.CreateConstIterator(); SelectionIt; ++SelectionIt )
 		{
 			const FCollectionNameType& Selection = *SelectionIt;
-			if ( Item->CollectionName == Selection.Name.ToString() && Item->CollectionType == Selection.Type )
+			if ( Item->CollectionName == Selection.Name && Item->CollectionType == Selection.Type )
 			{
 				CollectionListPtr->SetItemSelection(Item, true);
 				CollectionListPtr->RequestScrollIntoView(Item);
@@ -275,7 +275,7 @@ TArray<FCollectionNameType> SCollectionView::GetSelectedCollections() const
 	for ( int32 ItemIdx = 0; ItemIdx < Items.Num(); ++ItemIdx )
 	{
 		const TSharedPtr<FCollectionItem>& Item = Items[ItemIdx];
-		new(RetArray) FCollectionNameType(FName(*Item->CollectionName), Item->CollectionType);
+		RetArray.Add(FCollectionNameType(Item->CollectionName, Item->CollectionType));
 	}
 
 	return RetArray;
@@ -297,7 +297,7 @@ void SCollectionView::ApplyHistoryData ( const FHistoryData& History )
 	CollectionListPtr->ClearSelection();
 	for ( auto HistoryIt = History.SourcesData.Collections.CreateConstIterator(); HistoryIt; ++HistoryIt)
 	{
-		FString Name = (*HistoryIt).Name.ToString();
+		const FName Name = (*HistoryIt).Name;
 		ECollectionShareType::Type Type = (*HistoryIt).Type;
 		for ( auto CollectionIt = CollectionItems.CreateConstIterator(); CollectionIt; ++CollectionIt)
 		{
@@ -322,7 +322,7 @@ void SCollectionView::SaveSettings(const FString& IniFilename, const FString& In
 		}
 
 		const TSharedPtr<FCollectionItem>& Collection = *CollectionIt;
-		SelectedCollectionsString += Collection->CollectionName + TEXT("?") + FString::FromInt(Collection->CollectionType);
+		SelectedCollectionsString += Collection->CollectionName.ToString() + TEXT("?") + FString::FromInt(Collection->CollectionType);
 	}
 	const bool IsCollectionsExpanded = CollectionsExpandableAreaPtr.IsValid() ? CollectionsExpandableAreaPtr->IsExpanded() : true;
 	GConfig->SetBool(*IniSection, *(SettingsString + TEXT(".CollectionsExpanded")), IsCollectionsExpanded, IniFilename);
@@ -442,7 +442,7 @@ void SCollectionView::CreateCollectionItem( ECollectionShareType::Type Collectio
 		const FName BaseCollectionName = *LOCTEXT("NewCollectionName", "NewCollection").ToString();
 		FName CollectionName;
 		CollectionManagerModule.Get().CreateUniqueCollectionName(BaseCollectionName, CollectionType, CollectionName);
-		TSharedPtr<FCollectionItem> NewItem = MakeShareable(new FCollectionItem(CollectionName.ToString(), CollectionType));
+		TSharedPtr<FCollectionItem> NewItem = MakeShareable(new FCollectionItem(CollectionName, CollectionType));
 
 		// Adding a new collection now, so clear any filter we may have applied
 		SearchBoxPtr->SetText(FText::GetEmpty());
@@ -561,14 +561,14 @@ TSharedPtr<SWidget> SCollectionView::MakeCollectionListContextMenu()
 
 bool SCollectionView::IsCollectionCheckBoxEnabled( TSharedPtr<FCollectionItem> CollectionItem ) const
 {
-	return QuickAssetManagement.IsValid() && QuickAssetManagement->IsCollectionEnabled(FCollectionNameType(*CollectionItem->CollectionName, CollectionItem->CollectionType));
+	return QuickAssetManagement.IsValid() && QuickAssetManagement->IsCollectionEnabled(FCollectionNameType(CollectionItem->CollectionName, CollectionItem->CollectionType));
 }
 
 ECheckBoxState SCollectionView::IsCollectionChecked( TSharedPtr<FCollectionItem> CollectionItem ) const
 {
 	if ( QuickAssetManagement.IsValid() )
 	{
-		return QuickAssetManagement->GetCollectionCheckState(FCollectionNameType(*CollectionItem->CollectionName, CollectionItem->CollectionType));
+		return QuickAssetManagement->GetCollectionCheckState(FCollectionNameType(CollectionItem->CollectionName, CollectionItem->CollectionType));
 	}
 	return ECheckBoxState::Unchecked;
 }
@@ -580,11 +580,11 @@ void SCollectionView::OnCollectionCheckStateChanged( ECheckBoxState NewState, TS
 		switch(NewState)
 		{
 		case ECheckBoxState::Checked:
-			QuickAssetManagement->AddCurrentAssetsToCollection(FCollectionNameType(*CollectionItem->CollectionName, CollectionItem->CollectionType));
+			QuickAssetManagement->AddCurrentAssetsToCollection(FCollectionNameType(CollectionItem->CollectionName, CollectionItem->CollectionType));
 			break;
 
 		case ECheckBoxState::Unchecked:
-			QuickAssetManagement->RemoveCurrentAssetsFromCollection(FCollectionNameType(*CollectionItem->CollectionName, CollectionItem->CollectionType));
+			QuickAssetManagement->RemoveCurrentAssetsFromCollection(FCollectionNameType(CollectionItem->CollectionName, CollectionItem->CollectionType));
 			break;
 
 		default:
@@ -599,7 +599,7 @@ void SCollectionView::CollectionSelectionChanged( TSharedPtr< FCollectionItem >
 	{
 		if ( CollectionItem.IsValid() )
 		{
-			OnCollectionSelected.Execute(FCollectionNameType(FName(*CollectionItem->CollectionName), CollectionItem->CollectionType));
+			OnCollectionSelected.Execute(FCollectionNameType(CollectionItem->CollectionName, CollectionItem->CollectionType));
 		}
 		else
 		{
@@ -619,20 +619,20 @@ void SCollectionView::CollectionAssetsDropped(const TArray<FAssetData>& AssetLis
 
 	FCollectionManagerModule& CollectionManagerModule = FModuleManager::Get().LoadModuleChecked<FCollectionManagerModule>("CollectionManager");
 	int32 NumAdded = 0;
-	if ( CollectionManagerModule.Get().AddToCollection(FName(*CollectionItem->CollectionName), CollectionItem->CollectionType, ObjectPaths, &NumAdded) )
+	if ( CollectionManagerModule.Get().AddToCollection(CollectionItem->CollectionName, CollectionItem->CollectionType, ObjectPaths, &NumAdded) )
 	{
 		if ( AssetList.Num() == 1 )
 		{
 			FFormatNamedArguments Args;
 			Args.Add( TEXT("AssetName"), FText::FromName(AssetList[0].AssetName) );
-			Args.Add( TEXT("CollectionName"), FText::FromString(CollectionItem->CollectionName) );
+			Args.Add( TEXT("CollectionName"), FText::FromName(CollectionItem->CollectionName) );
 			OutMessage = FText::Format( LOCTEXT("CollectionAssetsAdded", "Added {AssetName} to {CollectionName}"), Args );
 		}
 		else
 		{
 			FFormatNamedArguments Args;
 			Args.Add( TEXT("Number"), NumAdded );
-			Args.Add( TEXT("CollectionName"), FText::FromString(CollectionItem->CollectionName) );
+			Args.Add( TEXT("CollectionName"), FText::FromName(CollectionItem->CollectionName) );
 			OutMessage = FText::Format( LOCTEXT("CollectionAssetsAdded", "Added {Number} asset(s) to {CollectionName}"), Args );
 		}
 	}
@@ -665,7 +665,7 @@ bool SCollectionView::CollectionNameChangeCommit( const TSharedPtr< FCollectionI
 	ECollectionShareType::Type CollectionType = CollectionItem->CollectionType;
 
 	// If new name is empty, set it back to the original name
-	FString NewNameFinal( NewName.IsEmpty() ? CollectionItem->CollectionName : NewName );
+	const FName NewNameFinal( NewName.IsEmpty() ? CollectionItem->CollectionName : FName(*NewName) );
 
 	if ( CollectionItem->bNewCollection )
 	{
@@ -679,7 +679,7 @@ bool SCollectionView::CollectionNameChangeCommit( const TSharedPtr< FCollectionI
 			return false;
 		}
 
-		if ( !CollectionManagerModule.Get().CreateCollection(FName(*NewNameFinal), CollectionType) )
+		if ( !CollectionManagerModule.Get().CreateCollection(NewNameFinal, CollectionType) )
 		{
 			// Failed to add the collection, remove it from the list
 			CollectionItems.Remove(CollectionItem);
@@ -698,13 +698,13 @@ bool SCollectionView::CollectionNameChangeCommit( const TSharedPtr< FCollectionI
 		}
 
 		// If the new name doesn't pass our current filter, we need to clear it
-		if ( !NewNameFinal.Contains(SearchBoxPtr->GetText().ToString()) )
+		if ( !NewNameFinal.ToString().Contains(SearchBoxPtr->GetText().ToString()) )
 		{
 			SearchBoxPtr->SetText(FText::GetEmpty());
 		}
 
 		// Otherwise perform the rename
-		if ( !CollectionManagerModule.Get().RenameCollection(FName(*CollectionItem->CollectionName), CollectionType, FName(*NewNameFinal), CollectionType) )
+		if ( !CollectionManagerModule.Get().RenameCollection(CollectionItem->CollectionName, CollectionType, NewNameFinal, CollectionType) )
 		{
 			// Failed to rename the collection
 			OutWarningMessage = FText::Format( LOCTEXT("RenameCollectionFailed", "Failed to rename the collection. {0}"), CollectionManagerModule.Get().GetLastError());
@@ -746,18 +746,20 @@ bool SCollectionView::CollectionNameChangeCommit( const TSharedPtr< FCollectionI
 
 bool SCollectionView::CollectionVerifyRenameCommit(const TSharedPtr< FCollectionItem >& CollectionItem, const FString& NewName, const FSlateRect& MessageAnchor, FText& OutErrorMessage)
 {
+	const FName NewNameFinal = *NewName;
+
 	// If the new name is the same as the old name, consider this to be unchanged, and accept it.
-	if (CollectionItem->CollectionName == NewName)
+	if (CollectionItem->CollectionName == NewNameFinal)
 	{
 		return true;
 	}
 
 	FCollectionManagerModule& CollectionManagerModule = FModuleManager::Get().LoadModuleChecked<FCollectionManagerModule>(TEXT("CollectionManager"));
 
-	if (CollectionManagerModule.Get().CollectionExists(*NewName, ECollectionShareType::CST_All))
+	if (CollectionManagerModule.Get().CollectionExists(NewNameFinal, ECollectionShareType::CST_All))
 	{
 		// This collection already exists, inform the user and continue
-		OutErrorMessage = FText::Format(LOCTEXT("RenameCollectionAlreadyExists", "A collection already exists with the name '{0}'."), FText::FromString(NewName));
+		OutErrorMessage = FText::Format(LOCTEXT("RenameCollectionAlreadyExists", "A collection already exists with the name '{0}'."), FText::FromName(NewNameFinal));
 
 		// Return false to indicate that the user should enter a new name
 		return false;
diff --git a/Engine/Source/Editor/ContentBrowser/Private/SourcesViewWidgets.cpp b/Engine/Source/Editor/ContentBrowser/Private/SourcesViewWidgets.cpp
--- a/Engine/Source/Editor/ContentBrowser/Private/SourcesViewWidgets.cpp
+++ b/Engine/Source/Editor/ContentBrowser/Private/SourcesViewWidgets.cpp
@@ -503,7 +503,7 @@ FText SCollectionListItem::GetNameText() const
 
 	if ( CollectionItemPtr.IsValid() )
 	{
-		return FText::FromString(CollectionItemPtr->CollectionName);
+		return FText::FromName(CollectionItemPtr->CollectionName);
 	}
 	else
 	{
@@ -517,7 +517,7 @@ FSlateColor SCollectionListItem::GetCollectionColor() const
 
 	if ( CollectionItemPtr.IsValid() )
 	{
-		const TSharedPtr<FLinearColor> Color = CollectionViewUtils::LoadColor(CollectionItemPtr->CollectionName, CollectionItemPtr->CollectionType);
+		const TSharedPtr<FLinearColor> Color = CollectionViewUtils::LoadColor(CollectionItemPtr->CollectionName.ToString(), CollectionItemPtr->CollectionType);
 		if( Color.IsValid() )
 		{
 			return *Color.Get();
