Added utility functions to get the FCollectionManagerModule handle

UETOOL-332 - Collections 2.0

This avoids the need to repeat the module name throughout the code (and potentially performing an FName look up each time).

diff --git a/Engine/Source/Developer/CollectionManager/Public/CollectionManagerModule.h b/Engine/Source/Developer/CollectionManager/Public/CollectionManagerModule.h
--- a/Engine/Source/Developer/CollectionManager/Public/CollectionManagerModule.h
+++ b/Engine/Source/Developer/CollectionManager/Public/CollectionManagerModule.h
@@ -13,6 +13,29 @@ public:
 
 	virtual ICollectionManager& Get() const;
 
+	/**
+	 * Singleton-like access to this module's interface.  This is just for convenience!
+	 * Beware of calling this during the shutdown phase, though.  Your module might have been unloaded already.
+	 *
+	 * @return Returns singleton instance, loading the module on demand if needed
+	 */
+	static inline FCollectionManagerModule& GetModule()
+	{
+		static const FName CollectionManagerModuleName("CollectionManager");
+		return FModuleManager::LoadModuleChecked<FCollectionManagerModule>(CollectionManagerModuleName);
+	}
+
+	/**
+	 * Checks to see if this module is loaded and ready.  It is only valid to call Get() if IsAvailable() returns true.
+	 *
+	 * @return True if the module is loaded and ready to use
+	 */
+	static inline bool IsModuleAvailable()
+	{
+		static const FName CollectionManagerModuleName("CollectionManager");
+		return FModuleManager::Get().IsModuleLoaded(CollectionManagerModuleName);
+	}
+
 private:
 	ICollectionManager* CollectionManager;
 	class FCollectionManagerConsoleCommands* ConsoleCommands;
diff --git a/Engine/Source/Editor/Blutility/Private/BlutilityShelf.cpp b/Engine/Source/Editor/Blutility/Private/BlutilityShelf.cpp
--- a/Engine/Source/Editor/Blutility/Private/BlutilityShelf.cpp
+++ b/Engine/Source/Editor/Blutility/Private/BlutilityShelf.cpp
@@ -56,7 +56,7 @@ void SBlutilityShelf::Construct(const FArguments& InArgs)
 	bInFavoritesMode = false;
 
 	// Check the shelf collection
-	FCollectionManagerModule& CollectionManagerModule = FModuleManager::LoadModuleChecked<FCollectionManagerModule>(TEXT("CollectionManager"));
+	FCollectionManagerModule& CollectionManagerModule = FCollectionManagerModule::GetModule();
 	if (CollectionManagerModule.Get().CollectionExists(BlutilityModule::BlutilityShelfCollectionName, ECollectionShareType::CST_Local))
 	{
 		// Start up in favorites mode if it contains something
@@ -176,7 +176,7 @@ void SBlutilityShelf::ToggleFavoriteStatusOnSelection(TArray<FAssetData> AssetLi
 		bool bAreAllSelected;
 		GetFavoriteStatus(AssetList, /*out*/ bAreAnySelected, /*out*/ bAreAllSelected);
 
-		FCollectionManagerModule& CollectionManagerModule = FModuleManager::LoadModuleChecked<FCollectionManagerModule>(TEXT("CollectionManager"));
+		FCollectionManagerModule& CollectionManagerModule = FCollectionManagerModule::GetModule();
 		if (bAreAllSelected)
 		{
 			// Remove them
@@ -201,7 +201,7 @@ bool SBlutilityShelf::GetFavoriteStatusOnSelection(TArray<FAssetData> AssetList)
 void SBlutilityShelf::GetFavoriteStatus(const TArray<FAssetData>& AssetList, /*out*/ bool& bAreAnySelected, /*out*/ bool& bAreAllSelected)
 {
 	TArray<FName> AssetPaths;
-	FCollectionManagerModule& CollectionManagerModule = FModuleManager::LoadModuleChecked<FCollectionManagerModule>(TEXT("CollectionManager"));
+	FCollectionManagerModule& CollectionManagerModule = FCollectionManagerModule::GetModule();
 	CollectionManagerModule.Get().GetAssetsInCollection(BlutilityModule::BlutilityShelfCollectionName, ECollectionShareType::CST_Local, /*out*/ AssetPaths);
 
 	bAreAnySelected = false;
diff --git a/Engine/Source/Editor/ContentBrowser/Private/AssetContextMenu.cpp b/Engine/Source/Editor/ContentBrowser/Private/AssetContextMenu.cpp
--- a/Engine/Source/Editor/ContentBrowser/Private/AssetContextMenu.cpp
+++ b/Engine/Source/Editor/ContentBrowser/Private/AssetContextMenu.cpp
@@ -813,8 +813,7 @@ bool FAssetContextMenu::AddCollectionMenuOptions(FMenuBuilder& MenuBuilder)
 	public:
 		static void CreateManageCollectionsSubMenu(FMenuBuilder& SubMenuBuilder, TSharedRef<FCollectionAssetManagement> QuickAssetManagement)
 		{
-			static const FName CollectionManagerModuleName = "CollectionManager";
-			FCollectionManagerModule& CollectionManagerModule = FModuleManager::LoadModuleChecked<FCollectionManagerModule>(CollectionManagerModuleName);
+			FCollectionManagerModule& CollectionManagerModule = FCollectionManagerModule::GetModule();
 
 			// Get collections of all types
 			TArray<FCollectionNameType> AvailableCollections;
@@ -887,8 +886,7 @@ bool FAssetContextMenu::AddCollectionMenuOptions(FMenuBuilder& MenuBuilder)
 
 	bool bHasAddedItems = false;
 
-	static const FName CollectionManagerModuleName = "CollectionManager";
-	FCollectionManagerModule& CollectionManagerModule = FModuleManager::LoadModuleChecked<FCollectionManagerModule>(CollectionManagerModuleName);
+	FCollectionManagerModule& CollectionManagerModule = FCollectionManagerModule::GetModule();
 
 	MenuBuilder.BeginSection("AssetContextCollections", LOCTEXT("AssetCollectionOptionsMenuHeading", "Collections"));
 
@@ -1545,7 +1543,7 @@ void FAssetContextMenu::ExecuteRemoveFromCollection()
 
 		if ( AssetsToRemove.Num() > 0 )
 		{
-			FCollectionManagerModule& CollectionManagerModule = FModuleManager::LoadModuleChecked<FCollectionManagerModule>("CollectionManager");
+			FCollectionManagerModule& CollectionManagerModule = FCollectionManagerModule::GetModule();
 
 			FName CollectionName = SourcesData.Collections[0].Name;
 			ECollectionShareType::Type CollectionType = SourcesData.Collections[0].Type;
diff --git a/Engine/Source/Editor/ContentBrowser/Private/AssetViewWidgets.cpp b/Engine/Source/Editor/ContentBrowser/Private/AssetViewWidgets.cpp
--- a/Engine/Source/Editor/ContentBrowser/Private/AssetViewWidgets.cpp
+++ b/Engine/Source/Editor/ContentBrowser/Private/AssetViewWidgets.cpp
@@ -396,8 +396,7 @@ TSharedRef<SToolTip> SAssetViewItem::CreateToolTipWidget() const
 			
 			// Add Collections
 			{
-				static const FName CollectionManagerModuleName = "CollectionManager";
-				FCollectionManagerModule& CollectionManagerModule = FModuleManager::LoadModuleChecked<FCollectionManagerModule>(CollectionManagerModuleName);
+				FCollectionManagerModule& CollectionManagerModule = FCollectionManagerModule::GetModule();
 
 				const FString CollectionNames = CollectionManagerModule.Get().GetCollectionsStringForObject(AssetData.ObjectPath, ECollectionShareType::CST_All);
 				if (!CollectionNames.IsEmpty())
diff --git a/Engine/Source/Editor/ContentBrowser/Private/CollectionAssetManagement.cpp b/Engine/Source/Editor/ContentBrowser/Private/CollectionAssetManagement.cpp
--- a/Engine/Source/Editor/ContentBrowser/Private/CollectionAssetManagement.cpp
+++ b/Engine/Source/Editor/ContentBrowser/Private/CollectionAssetManagement.cpp
@@ -10,8 +10,7 @@
 
 FCollectionAssetManagement::FCollectionAssetManagement()
 {
-	static const FName CollectionManagerModuleName = "CollectionManager";
-	FCollectionManagerModule& CollectionManagerModule = FModuleManager::Get().LoadModuleChecked<FCollectionManagerModule>(CollectionManagerModuleName);
+	FCollectionManagerModule& CollectionManagerModule = FCollectionManagerModule::GetModule();
 
 	// Register the notifications we need in order to keep things up-to-date
 	OnCollectionRenamedHandle = CollectionManagerModule.Get().OnCollectionRenamed().AddRaw(this, &FCollectionAssetManagement::HandleCollectionRenamed);
@@ -22,16 +21,15 @@ FCollectionAssetManagement::FCollectionAssetManagement()
 
 FCollectionAssetManagement::~FCollectionAssetManagement()
 {
-	// Use Get rather than Load as we might be in the process of shutting down...
-	static const FName CollectionManagerModuleName = "CollectionManager";
-	FCollectionManagerModule* CollectionManagerModulePtr = FModuleManager::Get().GetModulePtr<FCollectionManagerModule>(CollectionManagerModuleName);
-
-	if (CollectionManagerModulePtr)
+	// Check IsModuleAvailable as we might be in the process of shutting down...
+	if (FCollectionManagerModule::IsModuleAvailable())
 	{
-		CollectionManagerModulePtr->Get().OnCollectionRenamed().Remove(OnCollectionRenamedHandle);
-		CollectionManagerModulePtr->Get().OnCollectionDestroyed().Remove(OnCollectionDestroyedHandle);
-		CollectionManagerModulePtr->Get().OnAssetsAdded().Remove(OnAssetsAddedHandle);
-		CollectionManagerModulePtr->Get().OnAssetsRemoved().Remove(OnAssetsRemovedHandle);
+		FCollectionManagerModule& CollectionManagerModule = FCollectionManagerModule::GetModule();
+
+		CollectionManagerModule.Get().OnCollectionRenamed().Remove(OnCollectionRenamedHandle);
+		CollectionManagerModule.Get().OnCollectionDestroyed().Remove(OnCollectionDestroyedHandle);
+		CollectionManagerModule.Get().OnAssetsAdded().Remove(OnAssetsAddedHandle);
+		CollectionManagerModule.Get().OnAssetsRemoved().Remove(OnAssetsRemovedHandle);
 	}
 }
 
@@ -48,8 +46,7 @@ void FCollectionAssetManagement::SetCurrentAssets(const TArray<FAssetData>& Curr
 
 void FCollectionAssetManagement::AddCurrentAssetsToCollection(FCollectionNameType InCollectionKey)
 {
-	static const FName CollectionManagerModuleName = "CollectionManager";
-	FCollectionManagerModule& CollectionManagerModule = FModuleManager::Get().LoadModuleChecked<FCollectionManagerModule>(CollectionManagerModuleName);
+	FCollectionManagerModule& CollectionManagerModule = FCollectionManagerModule::GetModule();
 
 	const TArray<FName> ObjectPaths = CurrentAssetPaths.Array();
 
@@ -88,8 +85,7 @@ void FCollectionAssetManagement::AddCurrentAssetsToCollection(FCollectionNameTyp
 
 void FCollectionAssetManagement::RemoveCurrentAssetsFromCollection(FCollectionNameType InCollectionKey)
 {
-	static const FName CollectionManagerModuleName = "CollectionManager";
-	FCollectionManagerModule& CollectionManagerModule = FModuleManager::Get().LoadModuleChecked<FCollectionManagerModule>(CollectionManagerModuleName);
+	FCollectionManagerModule& CollectionManagerModule = FCollectionManagerModule::GetModule();
 
 	const TArray<FName> ObjectPaths = CurrentAssetPaths.Array();
 
@@ -148,8 +144,7 @@ ECheckBoxState FCollectionAssetManagement::GetCollectionCheckState(FCollectionNa
 
 void FCollectionAssetManagement::UpdateAssetManagementState()
 {
-	static const FName CollectionManagerModuleName = "CollectionManager";
-	FCollectionManagerModule& CollectionManagerModule = FModuleManager::Get().LoadModuleChecked<FCollectionManagerModule>(CollectionManagerModuleName);
+	FCollectionManagerModule& CollectionManagerModule = FCollectionManagerModule::GetModule();
 
 	AssetManagementState.Empty();
 
diff --git a/Engine/Source/Editor/ContentBrowser/Private/CollectionContextMenu.cpp b/Engine/Source/Editor/ContentBrowser/Private/CollectionContextMenu.cpp
--- a/Engine/Source/Editor/ContentBrowser/Private/CollectionContextMenu.cpp
+++ b/Engine/Source/Editor/ContentBrowser/Private/CollectionContextMenu.cpp
@@ -309,8 +309,7 @@ void FCollectionContextMenu::ExecuteSetCollectionShareType(ECollectionShareType:
 		return;
 	}
 
-	static const FName CollectionManagerModuleName = "CollectionManager";
-	FCollectionManagerModule& CollectionManagerModule = FModuleManager::Get().LoadModuleChecked<FCollectionManagerModule>(CollectionManagerModuleName);
+	FCollectionManagerModule& CollectionManagerModule = FCollectionManagerModule::GetModule();
 
 	CollectionManagerModule.Get().RenameCollection(CollectionList[0]->CollectionName, CollectionList[0]->CollectionType, CollectionList[0]->CollectionName, CollectionType);
 }
@@ -362,7 +361,7 @@ void FCollectionContextMenu::ExecuteDestroyCollection()
 
 FReply FCollectionContextMenu::ExecuteDestroyCollectionConfirmed(TArray<TSharedPtr<FCollectionItem>> CollectionList)
 {
-	FCollectionManagerModule& CollectionManagerModule = FModuleManager::Get().LoadModuleChecked<FCollectionManagerModule>(TEXT("CollectionManager"));
+	FCollectionManagerModule& CollectionManagerModule = FCollectionManagerModule::GetModule();
 
 	TArray<TSharedPtr<FCollectionItem>> ItemsToRemove;
 	for (int32 CollectionIdx = 0; CollectionIdx < CollectionList.Num(); ++CollectionIdx)
diff --git a/Engine/Source/Editor/ContentBrowser/Private/CollectionViewUtils.cpp b/Engine/Source/Editor/ContentBrowser/Private/CollectionViewUtils.cpp
--- a/Engine/Source/Editor/ContentBrowser/Private/CollectionViewUtils.cpp
+++ b/Engine/Source/Editor/ContentBrowser/Private/CollectionViewUtils.cpp
@@ -158,7 +158,7 @@ bool CollectionViewUtils::HasCustomColors( TArray< FLinearColor >* OutColors )
 	GConfig->GetSection(TEXT("CollectionColor"), Section, GEditorPerProjectIni);
 
 	bool bHasCustom = false;
-	const FCollectionManagerModule& CollectionManagerModule = FModuleManager::Get().LoadModuleChecked<FCollectionManagerModule>("CollectionManager");
+	const FCollectionManagerModule& CollectionManagerModule = FCollectionManagerModule::GetModule();
 	const ICollectionManager& CollectionManager = CollectionManagerModule.Get();
 
 	for(FString& EntryStr : Section)
diff --git a/Engine/Source/Editor/ContentBrowser/Private/SAssetView.cpp b/Engine/Source/Editor/ContentBrowser/Private/SAssetView.cpp
--- a/Engine/Source/Editor/ContentBrowser/Private/SAssetView.cpp
+++ b/Engine/Source/Editor/ContentBrowser/Private/SAssetView.cpp
@@ -88,7 +88,7 @@ void SAssetView::Construct( const FArguments& InArgs )
 	AssetRegistryModule.Get().OnPathAdded().AddSP( this, &SAssetView::OnAssetRegistryPathAdded );
 	AssetRegistryModule.Get().OnPathRemoved().AddSP( this, &SAssetView::OnAssetRegistryPathRemoved );
 
-	FCollectionManagerModule& CollectionManagerModule = FModuleManager::LoadModuleChecked<FCollectionManagerModule>(TEXT("CollectionManager"));
+	FCollectionManagerModule& CollectionManagerModule = FCollectionManagerModule::GetModule();
 	CollectionManagerModule.Get().OnAssetsAdded().AddSP( this, &SAssetView::OnAssetsAddedToCollection );
 	CollectionManagerModule.Get().OnAssetsRemoved().AddSP( this, &SAssetView::OnAssetsRemovedFromCollection );
 	CollectionManagerModule.Get().OnCollectionRenamed().AddSP( this, &SAssetView::OnCollectionRenamed );
@@ -1107,7 +1107,7 @@ FReply SAssetView::OnDragOver( const FGeometry& MyGeometry, const FDragDropEvent
 			if( DragAssetOp.IsValid() )
 			{
 				TArray< FName > ObjectPaths;
-				FCollectionManagerModule& CollectionManagerModule = FModuleManager::LoadModuleChecked<FCollectionManagerModule>(TEXT("CollectionManager"));
+				FCollectionManagerModule& CollectionManagerModule = FCollectionManagerModule::GetModule();
 				CollectionManagerModule.Get().GetObjectsInCollection( SourcesData.Collections[0].Name, SourcesData.Collections[0].Type, ObjectPaths );
 
 				bool IsValidDrop = false;
@@ -1172,7 +1172,7 @@ FReply SAssetView::OnDrop( const FGeometry& MyGeometry, const FDragDropEvent& Dr
 
 			if (ObjectPaths.Num() > 0)
 			{
-				FCollectionManagerModule& CollectionManagerModule = FModuleManager::LoadModuleChecked<FCollectionManagerModule>(TEXT("CollectionManager"));
+				FCollectionManagerModule& CollectionManagerModule = FCollectionManagerModule::GetModule();
 				CollectionManagerModule.Get().AddToCollection(SourcesData.Collections[0].Name, SourcesData.Collections[0].Type, ObjectPaths);
 			}
 
@@ -1377,7 +1377,7 @@ void SAssetView::RefreshSourceItems()
 		if ( bFilterAllowsClasses )
 		{
 			TArray< FName > ClassPaths;
-			FCollectionManagerModule& CollectionManagerModule = FModuleManager::GetModuleChecked<FCollectionManagerModule>(TEXT("CollectionManager"));
+			FCollectionManagerModule& CollectionManagerModule = FCollectionManagerModule::GetModule();
 			for (int32 Index = 0; Index < SourcesData.Collections.Num(); Index++)
 			{
 				CollectionManagerModule.Get().GetClassesInCollection( SourcesData.Collections[Index].Name, SourcesData.Collections[Index].Type, ClassPaths );
@@ -2174,7 +2174,7 @@ void SAssetView::RunAssetsThroughBackendFilter(TArray<FAssetData>& InOutAssetDat
 
 		if ( SourcesData.Collections.Num() > 0 )
 		{
-			FCollectionManagerModule& CollectionManagerModule = FModuleManager::GetModuleChecked<FCollectionManagerModule>(TEXT("CollectionManager"));
+			FCollectionManagerModule& CollectionManagerModule = FCollectionManagerModule::GetModule();
 			TArray< FName > CollectionObjectPaths;
 			for (int Index = 0; Index < SourcesData.Collections.Num(); Index++)
 			{
diff --git a/Engine/Source/Editor/ContentBrowser/Private/SCollectionView.cpp b/Engine/Source/Editor/ContentBrowser/Private/SCollectionView.cpp
--- a/Engine/Source/Editor/ContentBrowser/Private/SCollectionView.cpp
+++ b/Engine/Source/Editor/ContentBrowser/Private/SCollectionView.cpp
@@ -19,7 +19,7 @@ void SCollectionView::Construct( const FArguments& InArgs )
 	bAllowCollectionButtons = InArgs._AllowCollectionButtons;
 	bAllowRightClickMenu = InArgs._AllowRightClickMenu;
 
-	FCollectionManagerModule& CollectionManagerModule = FModuleManager::LoadModuleChecked<FCollectionManagerModule>(TEXT("CollectionManager"));
+	FCollectionManagerModule& CollectionManagerModule = FCollectionManagerModule::GetModule();
 	CollectionManagerModule.Get().OnCollectionCreated().AddSP( this, &SCollectionView::HandleCollectionCreated );
 	CollectionManagerModule.Get().OnCollectionRenamed().AddSP( this, &SCollectionView::HandleCollectionRenamed );
 	CollectionManagerModule.Get().OnCollectionDestroyed().AddSP( this, &SCollectionView::HandleCollectionDestroyed );
@@ -168,7 +168,7 @@ void SCollectionView::UpdateCollectionItems()
 	CollectionItems.Empty();
 
 	// Load the collection manager module
-	FCollectionManagerModule& CollectionManagerModule = FModuleManager::LoadModuleChecked<FCollectionManagerModule>(TEXT("CollectionManager"));
+	FCollectionManagerModule& CollectionManagerModule = FCollectionManagerModule::GetModule();
 
 	// Get collections of all types
 	for (int32 TypeIdx = 0; TypeIdx < ECollectionShareType::CST_All; ++TypeIdx)
@@ -437,7 +437,7 @@ void SCollectionView::CreateCollectionItem( ECollectionShareType::Type Collectio
 {
 	if ( ensure(CollectionType != ECollectionShareType::CST_All) )
 	{
-		FCollectionManagerModule& CollectionManagerModule = FModuleManager::Get().LoadModuleChecked<FCollectionManagerModule>("CollectionManager");
+		FCollectionManagerModule& CollectionManagerModule = FCollectionManagerModule::GetModule();
 
 		const FName BaseCollectionName = *LOCTEXT("NewCollectionName", "NewCollection").ToString();
 		FName CollectionName;
@@ -617,7 +617,7 @@ void SCollectionView::CollectionAssetsDropped(const TArray<FAssetData>& AssetLis
 		ObjectPaths.Add((*AssetIt).ObjectPath);
 	}
 
-	FCollectionManagerModule& CollectionManagerModule = FModuleManager::Get().LoadModuleChecked<FCollectionManagerModule>("CollectionManager");
+	FCollectionManagerModule& CollectionManagerModule = FCollectionManagerModule::GetModule();
 	int32 NumAdded = 0;
 	if ( CollectionManagerModule.Get().AddToCollection(CollectionItem->CollectionName, CollectionItem->CollectionType, ObjectPaths, &NumAdded) )
 	{
@@ -661,7 +661,7 @@ bool SCollectionView::CollectionNameChangeCommit( const TSharedPtr< FCollectionI
 	// There should only ever be one item selected when renaming
 	check(CollectionListPtr->GetNumItemsSelected() == 1);
 
-	FCollectionManagerModule& CollectionManagerModule = FModuleManager::Get().LoadModuleChecked<FCollectionManagerModule>(TEXT("CollectionManager"));
+	FCollectionManagerModule& CollectionManagerModule = FCollectionManagerModule::GetModule();
 	ECollectionShareType::Type CollectionType = CollectionItem->CollectionType;
 
 	// If new name is empty, set it back to the original name
@@ -754,7 +754,7 @@ bool SCollectionView::CollectionVerifyRenameCommit(const TSharedPtr< FCollection
 		return true;
 	}
 
-	FCollectionManagerModule& CollectionManagerModule = FModuleManager::Get().LoadModuleChecked<FCollectionManagerModule>(TEXT("CollectionManager"));
+	FCollectionManagerModule& CollectionManagerModule = FCollectionManagerModule::GetModule();
 
 	if (CollectionManagerModule.Get().CollectionExists(NewNameFinal, ECollectionShareType::CST_All))
 	{
diff --git a/Engine/Source/Editor/ContentBrowser/Private/SContentBrowser.cpp b/Engine/Source/Editor/ContentBrowser/Private/SContentBrowser.cpp
--- a/Engine/Source/Editor/ContentBrowser/Private/SContentBrowser.cpp
+++ b/Engine/Source/Editor/ContentBrowser/Private/SContentBrowser.cpp
@@ -26,11 +26,12 @@ SContentBrowser::~SContentBrowser()
 	UContentBrowserSettings::OnSettingChanged().RemoveAll( this );
 
 	// Remove listeners for when collections/paths are renamed/deleted
-	FCollectionManagerModule* CollectionManagerModule = FModuleManager::GetModulePtr<FCollectionManagerModule>(TEXT("CollectionManager"));
-	if (CollectionManagerModule != nullptr)
+	if (FCollectionManagerModule::IsModuleAvailable())
 	{
-		CollectionManagerModule->Get().OnCollectionRenamed().RemoveAll(this);
-		CollectionManagerModule->Get().OnCollectionDestroyed().RemoveAll(this);
+		FCollectionManagerModule& CollectionManagerModule = FCollectionManagerModule::GetModule();
+
+		CollectionManagerModule.Get().OnCollectionRenamed().RemoveAll(this);
+		CollectionManagerModule.Get().OnCollectionDestroyed().RemoveAll(this);
 	}
 
 	FAssetRegistryModule* AssetRegistryModule = FModuleManager::GetModulePtr<FAssetRegistryModule>(TEXT("AssetRegistry"));
@@ -620,7 +621,7 @@ void SContentBrowser::Construct( const FArguments& InArgs, const FName& InInstan
 	LoadSettings(InInstanceName);
 
 	// Bindings to manage history when items are deleted
-	FCollectionManagerModule& CollectionManagerModule = FModuleManager::LoadModuleChecked<FCollectionManagerModule>(TEXT("CollectionManager"));
+	FCollectionManagerModule& CollectionManagerModule = FCollectionManagerModule::GetModule();
 	CollectionManagerModule.Get().OnCollectionRenamed().AddSP(this, &SContentBrowser::HandleCollectionRenamed);
 	CollectionManagerModule.Get().OnCollectionDestroyed().AddSP(this, &SContentBrowser::HandleCollectionRemoved);
 
diff --git a/Engine/Source/Editor/ContentBrowser/Private/SourcesData.h b/Engine/Source/Editor/ContentBrowser/Private/SourcesData.h
--- a/Engine/Source/Editor/ContentBrowser/Private/SourcesData.h
+++ b/Engine/Source/Editor/ContentBrowser/Private/SourcesData.h
@@ -22,10 +22,10 @@ struct FSourcesData
 
 		// Collections
 		TArray<FName> ObjectPathsFromCollections;
-		if ( Collections.Num() )
+		if ( Collections.Num() && FCollectionManagerModule::IsModuleAvailable() )
 		{
 			// Collection manager module should already be loaded since it may cause a hitch on the first search
-			FCollectionManagerModule& CollectionManagerModule = FModuleManager::GetModuleChecked<FCollectionManagerModule>(TEXT("CollectionManager"));
+			FCollectionManagerModule& CollectionManagerModule = FCollectionManagerModule::GetModule();
 
 			for ( int32 CollectionIdx = 0; CollectionIdx < Collections.Num(); ++CollectionIdx )
 			{
diff --git a/Engine/Source/Editor/UnrealEd/Private/Commandlets/PackageUtilities.cpp b/Engine/Source/Editor/UnrealEd/Private/Commandlets/PackageUtilities.cpp
--- a/Engine/Source/Editor/UnrealEd/Private/Commandlets/PackageUtilities.cpp
+++ b/Engine/Source/Editor/UnrealEd/Private/Commandlets/PackageUtilities.cpp
@@ -234,31 +234,31 @@ class FCollectionPolicy
 public:
 	static bool CreateAssetSet(FName InSetName, ECollectionShareType::Type InSetType)
 	{
-		FCollectionManagerModule& CollectionManagerModule = FModuleManager::LoadModuleChecked<FCollectionManagerModule>(TEXT("CollectionManager"));
+		FCollectionManagerModule& CollectionManagerModule = FCollectionManagerModule::GetModule();
 		return CollectionManagerModule.Get().CreateCollection(InSetName, InSetType);
 	}
 
 	static bool DestroyAssetSet(FName InSetName, ECollectionShareType::Type InSetType )
 	{
-		FCollectionManagerModule& CollectionManagerModule = FModuleManager::LoadModuleChecked<FCollectionManagerModule>(TEXT("CollectionManager"));
+		FCollectionManagerModule& CollectionManagerModule = FCollectionManagerModule::GetModule();
 		return CollectionManagerModule.Get().DestroyCollection(InSetName, InSetType);
 	}
 
 	static bool RemoveAssetsFromSet(FName InSetName, ECollectionShareType::Type InSetType, const TArray<FName>& InAssetPathNames )
 	{
-		FCollectionManagerModule& CollectionManagerModule = FModuleManager::LoadModuleChecked<FCollectionManagerModule>(TEXT("CollectionManager"));
+		FCollectionManagerModule& CollectionManagerModule = FCollectionManagerModule::GetModule();
 		return CollectionManagerModule.Get().RemoveFromCollection(InSetName, InSetType, InAssetPathNames);
 	}
 
 	static bool AddAssetsToSet(FName InSetName, ECollectionShareType::Type InSetType, const TArray<FName>& InAssetPathNames )
 	{
-		FCollectionManagerModule& CollectionManagerModule = FModuleManager::LoadModuleChecked<FCollectionManagerModule>(TEXT("CollectionManager"));
+		FCollectionManagerModule& CollectionManagerModule = FCollectionManagerModule::GetModule();
 		return CollectionManagerModule.Get().AddToCollection(InSetName, InSetType, InAssetPathNames);
 	}
 
 	static bool QueryAssetsInSet(FName InSetName, ECollectionShareType::Type InSetType, TArray<FName>& OutAssetPathNames )
 	{
-		FCollectionManagerModule& CollectionManagerModule = FModuleManager::LoadModuleChecked<FCollectionManagerModule>(TEXT("CollectionManager"));
+		FCollectionManagerModule& CollectionManagerModule = FCollectionManagerModule::GetModule();
 		return CollectionManagerModule.Get().GetAssetsInCollection(InSetName, InSetType, OutAssetPathNames);
 	}
 };
diff --git a/Engine/Source/Editor/UnrealEd/Private/Toolkits/AssetEditorToolkit.cpp b/Engine/Source/Editor/UnrealEd/Private/Toolkits/AssetEditorToolkit.cpp
--- a/Engine/Source/Editor/UnrealEd/Private/Toolkits/AssetEditorToolkit.cpp
+++ b/Engine/Source/Editor/UnrealEd/Private/Toolkits/AssetEditorToolkit.cpp
@@ -323,8 +323,7 @@ FText FAssetEditorToolkit::GetToolTipTextForObject(const UObject* InObject)
 		ToolTipString += TEXT(": ");
 		ToolTipString += InObject->GetName();
 
-		static const FName CollectionManagerModuleName = "CollectionManager";
-		FCollectionManagerModule& CollectionManagerModule = FModuleManager::LoadModuleChecked<FCollectionManagerModule>(CollectionManagerModuleName);
+		FCollectionManagerModule& CollectionManagerModule = FCollectionManagerModule::GetModule();
 
 		const FString CollectionNames = CollectionManagerModule.Get().GetCollectionsStringForObject(*InObject->GetPathName(), ECollectionShareType::CST_All);
 		if (!CollectionNames.IsEmpty())
