Improved path searching in the Content Browser

Re-added the exported text for an asset to the list of simple search terms. This was broken a while ago, and messed with some peoples workflow as they used to Ctrl+C the asset to get the exported text which they could later search for in the Content Browser.

This also improves the handling of the "path" key, as it can now compare against paths in any of the following forms (previously it only tested the last one):
/Path/To/Asset.Asset
/Path/To/Asset
/Path/To

diff --git a/Engine/Source/Editor/ContentBrowser/Private/FrontendFilters.cpp b/Engine/Source/Editor/ContentBrowser/Private/FrontendFilters.cpp
--- a/Engine/Source/Editor/ContentBrowser/Private/FrontendFilters.cpp
+++ b/Engine/Source/Editor/ContentBrowser/Private/FrontendFilters.cpp
@@ -51,6 +51,7 @@ void AssetDataToNameString(FAssetFilterType Asset, OUT TArray< FString >& Array)
 void AssetDataToClassAndNameStrings(FAssetFilterType Asset, OUT TArray< FString >& Array)
 {
 	Array.Add(Asset.AssetClass.ToString());
+	Array.Add(Asset.GetExportTextName()); // only include this if we're searching the class name too, as the exported text contains the type in the string
 	AssetDataToNameString(Asset, Array);
 }
 
@@ -85,7 +86,18 @@ bool AssetDataTestComplexExpression(FAssetFilterType Asset, const FName& InKey,
 			return false;
 		}
 
-		const bool bIsMatch = TextFilterUtils::TestBasicStringExpression(Asset.PackagePath.ToString(), InValue, InTextComparisonMode);
+		// If the comparison mode is partial, then we only need to test the ObjectPath as that contains the other two as sub-strings
+		bool bIsMatch = false;
+		if (InTextComparisonMode == ETextFilterTextComparisonMode::Partial)
+		{
+			bIsMatch = TextFilterUtils::TestBasicStringExpression(Asset.ObjectPath.ToString(), InValue, InTextComparisonMode);
+		}
+		else
+		{
+			bIsMatch = TextFilterUtils::TestBasicStringExpression(Asset.ObjectPath.ToString(), InValue, InTextComparisonMode)
+				|| TextFilterUtils::TestBasicStringExpression(Asset.PackageName.ToString(), InValue, InTextComparisonMode)
+				|| TextFilterUtils::TestBasicStringExpression(Asset.PackagePath.ToString(), InValue, InTextComparisonMode);
+		}
 		return (InComparisonOperation == ETextFilterComparisonOperation::Equal) ? bIsMatch : !bIsMatch;
 	}
 
