Converted SCollectionView to use a TTextFilter

diff --git a/Engine/Source/Editor/ContentBrowser/Private/SCollectionView.cpp b/Engine/Source/Editor/ContentBrowser/Private/SCollectionView.cpp
--- a/Engine/Source/Editor/ContentBrowser/Private/SCollectionView.cpp
+++ b/Engine/Source/Editor/ContentBrowser/Private/SCollectionView.cpp
@@ -13,6 +13,50 @@
 
 #define LOCTEXT_NAMESPACE "ContentBrowser"
 
+namespace CollectionViewFilter
+{
+
+void GetBasicStrings(const FCollectionItem& InCollection, TArray<FString>& OutBasicStrings)
+{
+	OutBasicStrings.Add(InCollection.CollectionName.ToString());
+}
+
+bool TestComplexExpression(const FCollectionItem& InCollection, const FName& InKey, const FString& InValue, ETextFilterComparisonOperation InComparisonOperation, ETextFilterTextComparisonMode InTextComparisonMode)
+{
+	static const FName NameKeyName = "Name";
+	static const FName TypeKeyName = "Type";
+
+	// Handle the collection name
+	if (InKey == NameKeyName)
+	{
+		// Names can only work with Equal or NotEqual type tests
+		if (InComparisonOperation != ETextFilterComparisonOperation::Equal && InComparisonOperation != ETextFilterComparisonOperation::NotEqual)
+		{
+			return false;
+		}
+
+		const bool bIsMatch = TextFilterUtils::TestBasicStringExpression(InCollection.CollectionName.ToString(), InValue, InTextComparisonMode);
+		return (InComparisonOperation == ETextFilterComparisonOperation::Equal) ? bIsMatch : !bIsMatch;
+	}
+
+	// Handle the collection type
+	if (InKey == TypeKeyName)
+	{
+		// Types can only work with Equal or NotEqual type tests
+		if (InComparisonOperation != ETextFilterComparisonOperation::Equal && InComparisonOperation != ETextFilterComparisonOperation::NotEqual)
+		{
+			return false;
+		}
+
+		const bool bIsMatch = TextFilterUtils::TestBasicStringExpression(ECollectionShareType::ToString(InCollection.CollectionType), InValue, InTextComparisonMode);
+		return (InComparisonOperation == ETextFilterComparisonOperation::Equal) ? bIsMatch : !bIsMatch;
+	}
+
+	return false;
+}
+
+} // namespace CollectionViewFilter
+
 void SCollectionView::Construct( const FArguments& InArgs )
 {
 	OnCollectionSelected = InArgs._OnCollectionSelected;
@@ -28,6 +72,12 @@ void SCollectionView::Construct( const FArguments& InArgs )
 	CollectionContextMenu = MakeShareable(new FCollectionContextMenu( SharedThis(this) ));
 	CollectionContextMenu->BindCommands(Commands);
 
+	CollectionItemTextFilter = MakeShareable(new FCollectionItemTextFilter(
+		FCollectionItemTextFilter::FItemToStringArray::CreateStatic(&CollectionViewFilter::GetBasicStrings), 
+		FCollectionItemTextFilter::FItemTestComplexExpression::CreateStatic(&CollectionViewFilter::TestComplexExpression)
+		));
+	CollectionItemTextFilter->OnChanged().AddSP(this, &SCollectionView::UpdateFilteredCollectionItems);
+
 	if ( InArgs._AllowQuickAssetManagement )
 	{
 		QuickAssetManagement = MakeShareable(new FCollectionAssetManagement());
@@ -60,7 +110,7 @@ void SCollectionView::Construct( const FArguments& InArgs )
 				[
 					SAssignNew(SearchBoxPtr, SSearchBox)
 					.HintText( LOCTEXT( "CollectionsViewSearchBoxHint", "Search Collections" ) )
-					.OnTextChanged( this, &SCollectionView::ApplyCollectionsSearchFilter )
+					.OnTextChanged( this, &SCollectionView::SetCollectionsSearchFilterText )
 					.Visibility( this, &SCollectionView::GetCollectionsSearchBoxVisibility )
 				]
 			]
@@ -193,40 +243,32 @@ void SCollectionView::UpdateCollectionItems()
 
 	CollectionItems.Sort( FCollectionItem::FCompareFCollectionItemByName() );
 
-	ApplyCollectionsSearchFilter();
-}
-
-void SCollectionView::ApplyCollectionsSearchFilter()
-{
-	ApplyCollectionsSearchFilter( SearchBoxPtr->GetText() );
+	UpdateFilteredCollectionItems();
 }
 
-void SCollectionView::ApplyCollectionsSearchFilter( const FText& InSearchText )
+void SCollectionView::UpdateFilteredCollectionItems()
 {
-	if ( InSearchText.IsEmptyOrWhitespace() )
+	FilteredCollectionItems.Reset();
+	for ( const auto& CollectionItem : CollectionItems )
 	{
-		FilteredCollectionItems = CollectionItems;
-	}
-	else
-	{
-		const FString& SearchTextString = InSearchText.ToString();
-
-		FilteredCollectionItems.Reset();
-		for ( const auto& CollectionItem : CollectionItems )
+		if ( CollectionItemTextFilter->PassesFilter(*CollectionItem) )
 		{
-			if ( CollectionItem->CollectionName.ToString().Contains(SearchTextString) )
-			{
-				FilteredCollectionItems.Add(CollectionItem);
-			}
+			FilteredCollectionItems.Add(CollectionItem);
 		}
 	}
 
 	CollectionListPtr->RequestListRefresh();
 }
 
-FText SCollectionView::GetFilterText() const
+void SCollectionView::SetCollectionsSearchFilterText( const FText& InSearchText )
+{
+	CollectionItemTextFilter->SetRawFilterText( InSearchText );
+	SearchBoxPtr->SetError( CollectionItemTextFilter->GetFilterErrorText() );
+}
+
+FText SCollectionView::GetCollectionsSearchFilterText() const
 {
-	return SearchBoxPtr->GetText();
+	return CollectionItemTextFilter->GetRawFilterText();
 }
 
 void SCollectionView::SetSelectedCollections(const TArray<FCollectionNameType>& CollectionsToSelect)
@@ -289,7 +331,7 @@ void SCollectionView::SetSelectedAssets(const TArray<FAssetData>& SelectedAssets
 	}
 }
 
-void SCollectionView::ApplyHistoryData ( const FHistoryData& History )
+void SCollectionView::ApplyHistoryData( const FHistoryData& History )
 {	
 	// Prevent the selection changed delegate because it would add more history when we are just setting a state
 	FScopedPreventSelectionChangedDelegate DelegatePrevention( SharedThis(this) );
@@ -454,7 +496,7 @@ void SCollectionView::CreateCollectionItem( ECollectionShareType::Type Collectio
 
 		CollectionItems.Add( NewItem );
 		CollectionItems.Sort( FCollectionItem::FCompareFCollectionItemByName() );
-		ApplyCollectionsSearchFilter();
+		UpdateFilteredCollectionItems();
 		CollectionListPtr->RequestScrollIntoView(NewItem);
 		CollectionListPtr->SetSelection( NewItem );
 	}
@@ -507,7 +549,7 @@ void SCollectionView::RemoveCollectionItems( const TArray<TSharedPtr<FCollection
 	}
 
 	// Refresh the list
-	ApplyCollectionsSearchFilter();
+	UpdateFilteredCollectionItems();
 }
 
 EVisibility SCollectionView::GetCollectionListVisibility() const
@@ -541,7 +583,7 @@ TSharedRef<ITableRow> SCollectionView::GenerateCollectionRow( TSharedPtr<FCollec
 			.OnAssetsDragDropped(this, &SCollectionView::CollectionAssetsDropped)
 			.IsSelected( TableRow.Get(), &STableRow< TSharedPtr<FTreeItem> >::IsSelectedExclusively )
 			.IsReadOnly(this, &SCollectionView::IsCollectionNotRenamable)
-			.HighlightText( this, &SCollectionView::GetFilterText)
+			.HighlightText( this, &SCollectionView::GetCollectionsSearchFilterText)
 			.IsCheckBoxEnabled( IsCollectionCheckBoxEnabledAttribute )
 			.IsCollectionChecked( IsCollectionCheckedAttribute )
 			.OnCollectionCheckStateChanged( OnCollectionCheckStateChangedDelegate )
@@ -676,7 +718,7 @@ bool SCollectionView::CollectionNameChangeCommit( const TSharedPtr< FCollectionI
 		if ( !bChangeConfirmed )
 		{
 			CollectionItems.Remove(CollectionItem);
-			ApplyCollectionsSearchFilter();
+			UpdateFilteredCollectionItems();
 			return false;
 		}
 
@@ -684,7 +726,7 @@ bool SCollectionView::CollectionNameChangeCommit( const TSharedPtr< FCollectionI
 		{
 			// Failed to add the collection, remove it from the list
 			CollectionItems.Remove(CollectionItem);
-			ApplyCollectionsSearchFilter();
+			UpdateFilteredCollectionItems();
 
 			OutWarningMessage = FText::Format( LOCTEXT("CreateCollectionFailed", "Failed to create the collection. {0}"), CollectionManagerModule.Get().GetLastError());
 			return false;
@@ -699,7 +741,7 @@ bool SCollectionView::CollectionNameChangeCommit( const TSharedPtr< FCollectionI
 		}
 
 		// If the new name doesn't pass our current filter, we need to clear it
-		if ( !NewNameFinal.ToString().Contains(SearchBoxPtr->GetText().ToString()) )
+		if ( !CollectionItemTextFilter->PassesFilter( FCollectionItem(NewNameFinal, CollectionType) ) )
 		{
 			SearchBoxPtr->SetText(FText::GetEmpty());
 		}
diff --git a/Engine/Source/Editor/ContentBrowser/Private/SCollectionView.h b/Engine/Source/Editor/ContentBrowser/Private/SCollectionView.h
--- a/Engine/Source/Editor/ContentBrowser/Private/SCollectionView.h
+++ b/Engine/Source/Editor/ContentBrowser/Private/SCollectionView.h
@@ -139,14 +139,14 @@ private:
 	/** Updates the collections shown in the list view */
 	void UpdateCollectionItems();
 
-	/** Update the visible collections based on the active search text */
-	void ApplyCollectionsSearchFilter();
+	/** Update the visible collections based on the active filter text */
+	void UpdateFilteredCollectionItems();
 
-	/** Update the visible collections based on the given search text */
-	void ApplyCollectionsSearchFilter( const FText& InSearchText );
+	/** Set the active filter text */
+	void SetCollectionsSearchFilterText( const FText& InSearchText );
 
-	/** Get the currently filter text */
-	FText GetFilterText() const;
+	/** Get the active filter text */
+	FText GetCollectionsSearchFilterText() const;
 
 private:
 
@@ -182,6 +182,10 @@ private:
 	/** The list of visible collections based on the current filter */
 	TArray< TSharedPtr<FCollectionItem> > FilteredCollectionItems;
 
+	/** The filter to apply to the available collections */
+	typedef TTextFilter<const FCollectionItem&> FCollectionItemTextFilter;
+	TSharedPtr< FCollectionItemTextFilter > CollectionItemTextFilter;
+
 	/** The context menu logic and data */
 	TSharedPtr<class FCollectionContextMenu> CollectionContextMenu;
 
