Fixed SPathView::OnPathSelected being called with an empty string after the editor loads

UE-13120 - Content Browser - when launching the editor, if the path is a root directory OnPathSelected fires with an empty string.

SPathView::LoadSettings was passing a null item. It now passes the first selected item.

diff --git a/Engine/Source/Editor/ContentBrowser/Private/SPathView.cpp b/Engine/Source/Editor/ContentBrowser/Private/SPathView.cpp
--- a/Engine/Source/Editor/ContentBrowser/Private/SPathView.cpp
+++ b/Engine/Source/Editor/ContentBrowser/Private/SPathView.cpp
@@ -623,8 +623,12 @@ void SPathView::LoadSettings(const FString& IniFilename, const FString& IniSecti
 
 			if ( bSelectedAtLeastOnePath )
 			{
+				// Send the first selected item with the notification
+				const TArray<TSharedPtr<FTreeItem>> SelectedItems = TreeViewPtr->GetSelectedItems();
+				check(SelectedItems.Num() > 0);
+
 				// Signal a single selection changed event to let any listeners know that paths have changed
-				TreeSelectionChanged( TSharedPtr<FTreeItem>(), ESelectInfo::Direct );
+				TreeSelectionChanged( SelectedItems[0], ESelectInfo::Direct );
 			}
 		}
 		else
@@ -632,8 +636,12 @@ void SPathView::LoadSettings(const FString& IniFilename, const FString& IniSecti
 			// If all assets are already discovered, just select paths the best we can
 			SetSelectedPaths(NewSelectedPaths);
 
+			// Send the first selected item with the notification
+			const TArray<TSharedPtr<FTreeItem>> SelectedItems = TreeViewPtr->GetSelectedItems();
+			check(SelectedItems.Num() > 0);
+
 			// Signal a single selection changed event to let any listeners know that paths have changed
-			TreeSelectionChanged( TSharedPtr<FTreeItem>(), ESelectInfo::Direct );
+			TreeSelectionChanged( SelectedItems[0], ESelectInfo::Direct );
 		}
 	}
 }
