Changing the Content Browser path via the breadcrumb now notifies you of the change

UE-13121 - Content Browser - navigiating via breadcrumb trail does not fire OnPathSelected delegate.

diff --git a/Engine/Source/Editor/ContentBrowser/Private/SContentBrowser.cpp b/Engine/Source/Editor/ContentBrowser/Private/SContentBrowser.cpp
--- a/Engine/Source/Editor/ContentBrowser/Private/SContentBrowser.cpp
+++ b/Engine/Source/Editor/ContentBrowser/Private/SContentBrowser.cpp
@@ -1290,6 +1290,9 @@ void SContentBrowser::OnPathClicked( const FString& CrumbData )
 		TArray<FString> SelectedPaths;
 		SelectedPaths.Add(CrumbData);
 		PathViewPtr->SetSelectedPaths(SelectedPaths);
+
+		PathSelected(CrumbData);
+		
 		SourcesChanged(SelectedPaths, TArray<FCollectionNameType>());
 	}
 }
