Added support for nested collections

UETOOL-332 - Collections 2.0
UETOOL-369 - Want nested collections with collapsing

We now have version 2 collections which maintain a persistent GUID for each collection. Existing collections will be lazily updated to this version when they need to be re-saved.

This GUID is used by child collections to keep track their parents, and the collections view (as well as the quick asset management) now show a tree of collections. Collections in the main collection view tree can be re-parented via drag and drop.

Performing Content Browser searches against a given collection will also test to see if an object exists in child collections, and the asset view will now show you folder entries for child collections when viewing a parent (if folders are enabled in your Content Browser view settings).

diff --git a/Engine/Source/Developer/CollectionManager/Private/Collection.cpp b/Engine/Source/Developer/CollectionManager/Private/Collection.cpp
--- a/Engine/Source/Developer/CollectionManager/Private/Collection.cpp
+++ b/Engine/Source/Developer/CollectionManager/Private/Collection.cpp
@@ -5,70 +5,62 @@
 
 #define LOCTEXT_NAMESPACE "CollectionManager"
 
-
-enum ECollectionFileVersion
-{
-	// The initial version for collection files
-	COLLECTION_VER_INITIAL								= 1,
-
-	// The automatic version count. This must always be at the end of the enum
-	COLLECTION_VER_AUTOMATIC_VERSION_PLUS_ONE
-};
-
 struct FCollectionUtils
 {
-	static void AppendCollectionToArray(const TSet<FName>& InAssetSet, TArray<FName>& OutAssetArray)
+	static void AppendCollectionToArray(const TSet<FName>& InObjectSet, TArray<FName>& OutObjectArray)
 	{
-		OutAssetArray.Reserve(OutAssetArray.Num() + InAssetSet.Num());
-		for (const FName& AssetName : InAssetSet)
+		OutObjectArray.Reserve(OutObjectArray.Num() + InObjectSet.Num());
+		for (const FName& ObjectName : InObjectSet)
 		{
-			OutAssetArray.Add(AssetName);
+			OutObjectArray.Add(ObjectName);
 		}
 	}
-
-	static TArray<FName> CreateCollectionArray(const TSet<FName>& InAssetSet)
-	{
-		TArray<FName> AssetArray;
-		AppendCollectionToArray(InAssetSet, AssetArray);
-		return AssetArray;
-	}
 };
 
-FCollection::FCollection()
+FCollection::FCollection(const FString& InFilename, bool InUseSCC)
 {
+	ensure(InFilename.Len() > 0);
 
+	bUseSCC = InUseSCC;
+	SourceFilename = InFilename;
+	CollectionName = FName(*FPaths::GetBaseFilename(InFilename));
+
+	CollectionGuid = FGuid::NewGuid();
+
+	// Initialize the file version to the most recent
+	FileVersion = ECollectionVersion::CurrentVersion;
 }
 
-FCollection::FCollection(FName InCollectionName, const FString& SourceFolder, const FString& CollectionExtension, bool InUseSCC)
+TSharedRef<FCollection> FCollection::Clone(const FString& InFilename, bool InUseSCC, ECollectionCloneMode InCloneMode) const
 {
-	CollectionName = InCollectionName;
-	bUseSCC = InUseSCC;
+	TSharedRef<FCollection> NewCollection = MakeShareable(new FCollection(*this));
 
-	// Initialize the file version to the most recent
-	FileVersion = COLLECTION_VER_AUTOMATIC_VERSION_PLUS_ONE - 1;
+	// Set the new collection name and path
+	NewCollection->bUseSCC = InUseSCC;
+	NewCollection->SourceFilename = InFilename;
+	NewCollection->CollectionName = FName(*FPaths::GetBaseFilename(InFilename));
 
-	if ( ensure(SourceFolder.Len()) )
+	// Create a new GUID?
+	if (InCloneMode == ECollectionCloneMode::Unique)
 	{
-		SourceFilename = SourceFolder / CollectionName.ToString() + TEXT(".") + CollectionExtension;
-		SourceFilename.ReplaceInline(TEXT("//"), TEXT("/"));
+		NewCollection->CollectionGuid = FGuid::NewGuid();
 	}
+
+	return NewCollection;
 }
 
-bool FCollection::LoadFromFile(const FString& InFilename, bool InUseSCC)
+bool FCollection::Load(FText& OutError)
 {
-	Clear();
+	ObjectSet.Empty();
+	DiskSnapshot = FCollectionSnapshot();
 
 	FString FullFileContentsString;
-	if ( !FFileHelper::LoadFileToString(FullFileContentsString, *InFilename) )
+	if (!FFileHelper::LoadFileToString(FullFileContentsString, *SourceFilename))
 	{
+		OutError = FText::Format(LOCTEXT("LoadError_FailedToLoadFile", "Failed to load the collection '{0}' from disk."), FText::FromString(SourceFilename));
 		return false;
 	}
 
-	// Initialize the FCollection and set up the collection name
-	bUseSCC = InUseSCC;
-	SourceFilename = InFilename;
-	CollectionName = FName(*FPaths::GetBaseFilename(InFilename));
-		
 	// Normalize line endings and parse into array
 	TArray<FString> FileContents;
 	FullFileContentsString.ReplaceInline(TEXT("\r"), TEXT(""));
@@ -106,6 +98,7 @@ bool FCollection::LoadFromFile(const FString& InFilename, bool InUseSCC)
 	if ( !LoadHeaderPairs(HeaderPairs) )
 	{
 		// Bad header
+		OutError = FText::Format(LOCTEXT("LoadError_BadHeader", "The collection file '{0}' contains a bad header and could not be loaded."), FText::FromString(SourceFilename));
 		return false;
 	}
 
@@ -127,10 +120,10 @@ bool FCollection::LoadFromFile(const FString& InFilename, bool InUseSCC)
 				AddAssetToCollection(FName(*Line));
 			}
 		}
-
-		DiskAssetSet = AssetSet;
 	}
 
+	DiskSnapshot.TakeSnapshot(*this);
+
 	return true;
 }
 
@@ -171,7 +164,7 @@ bool FCollection::Save(FText& OutError)
 
 	// Start with the header
 	TMap<FString,FString> HeaderPairs;
-	GenerateHeaderPairs(HeaderPairs);
+	SaveHeaderPairs(HeaderPairs);
 	for (const auto& HeaderPair : HeaderPairs)
 	{
 		FileOutput += HeaderPair.Key + TEXT(":") + HeaderPair.Value + LINE_TERMINATOR;
@@ -186,13 +179,13 @@ bool FCollection::Save(FText& OutError)
 	else
 	{
 		// Write out the set as a sorted array to keep things in a known order for diffing
-		TArray<FName> AssetList = FCollectionUtils::CreateCollectionArray(AssetSet);
-		AssetList.Sort();
+		TArray<FName> ObjectList = ObjectSet.Array();
+		ObjectList.Sort();
 
-		// Static collection. Save a flat list of all assets in the collection.
-		for (const FName& AssetName : AssetList)
+		// Static collection. Save a flat list of all objects in the collection.
+		for (const FName& ObjectName : ObjectList)
 		{
-			FileOutput += AssetName.ToString() + LINE_TERMINATOR;
+			FileOutput += ObjectName.ToString() + LINE_TERMINATOR;
 		}
 	}
 
@@ -251,7 +244,10 @@ bool FCollection::Save(FText& OutError)
 
 	if ( bSaveSuccessful )
 	{
-		DiskAssetSet = AssetSet;
+		// Files are always saved at the latest version as loading should take care of data upgrades
+		FileVersion = ECollectionVersion::CurrentVersion;
+
+		DiskSnapshot.TakeSnapshot(*this);
 	}
 
 	GWarn->EndSlowTask();
@@ -288,7 +284,7 @@ bool FCollection::DeleteSourceFile(FText& OutError)
 
 	if ( bSuccessfullyDeleted )
 	{
-		DiskAssetSet.Empty();
+		DiskSnapshot = FCollectionSnapshot();
 	}
 
 	return bSuccessfullyDeleted;
@@ -297,10 +293,9 @@ bool FCollection::DeleteSourceFile(FText& OutError)
 bool FCollection::AddAssetToCollection(FName ObjectPath)
 {
 	// @todo collection Does this apply to dynamic collections?
-	if ( !AssetSet.Contains(ObjectPath) )
+	if (!ObjectSet.Contains(ObjectPath))
 	{
-		AssetSet.Add(ObjectPath);
-
+		ObjectSet.Add(ObjectPath);
 		return true;
 	}
 
@@ -310,17 +305,17 @@ bool FCollection::AddAssetToCollection(FName ObjectPath)
 bool FCollection::RemoveAssetFromCollection(FName ObjectPath)
 {
 	// @todo collection Does this apply to dynamic collections?
-	return AssetSet.Remove(ObjectPath) > 0;
+	return ObjectSet.Remove(ObjectPath) > 0;
 }
 
 void FCollection::GetAssetsInCollection(TArray<FName>& Assets) const
 {
 	// @todo collection Does this apply to dynamic collections?
-	for (const FName& AssetName : AssetSet)
+	for (const FName& ObjectName : ObjectSet)
 	{
-		if (!AssetName.ToString().StartsWith(TEXT("/Script/")))
+		if (!ObjectName.ToString().StartsWith(TEXT("/Script/")))
 		{
-			Assets.Add(AssetName);
+			Assets.Add(ObjectName);
 		}
 	}
 }
@@ -328,11 +323,11 @@ void FCollection::GetAssetsInCollection(TArray<FName>& Assets) const
 void FCollection::GetClassesInCollection(TArray<FName>& Classes) const
 {
 	// @todo collection Does this apply to dynamic collections?
-	for (const FName& AssetName : AssetSet)
+	for (const FName& ObjectName : ObjectSet)
 	{
-		if (AssetName.ToString().StartsWith(TEXT("/Script/")))
+		if (ObjectName.ToString().StartsWith(TEXT("/Script/")))
 		{
-			Classes.Add(AssetName);
+			Classes.Add(ObjectName);
 		}
 	}
 }
@@ -340,22 +335,13 @@ void FCollection::GetClassesInCollection(TArray<FName>& Classes) const
 void FCollection::GetObjectsInCollection(TArray<FName>& Objects) const
 {
 	// @todo collection Does this apply to dynamic collections?
-	FCollectionUtils::AppendCollectionToArray(AssetSet, Objects);
+	FCollectionUtils::AppendCollectionToArray(ObjectSet, Objects);
 }
 
 bool FCollection::IsObjectInCollection(FName ObjectPath) const
 {
 	// @todo collection Does this apply to dynamic collections?
-	return AssetSet.Contains(ObjectPath);
-}
-
-void FCollection::Clear()
-{
-	CollectionName = NAME_None;
-	bUseSCC = false;
-	SourceFilename = TEXT("");
-	AssetSet.Empty();
-	DiskAssetSet.Empty();
+	return ObjectSet.Contains(ObjectPath);
 }
 
 bool FCollection::IsDynamic() const
@@ -363,10 +349,10 @@ bool FCollection::IsDynamic() const
 	// @todo collection Dynamic collections
 	return false;
 }
-	
+
 bool FCollection::IsEmpty() const
 {
-	return AssetSet.Num() == 0;
+	return ObjectSet.Num() == 0;
 }
 
 void FCollection::PrintCollection() const
@@ -381,22 +367,23 @@ void FCollection::PrintCollection() const
 		UE_LOG(LogCollectionManager, Log, TEXT("    ============================="));
 
 		// Print the set as a sorted array to keep things in a sane order
-		TArray<FName> AssetList = FCollectionUtils::CreateCollectionArray(AssetSet);
-		AssetList.Sort();
+		TArray<FName> ObjectList = ObjectSet.Array();
+		ObjectList.Sort();
 
-		for (const FName& AssetName : AssetList)
+		for (const FName& ObjectName : ObjectList)
 		{
-			UE_LOG(LogCollectionManager, Log, TEXT("        %s"), *AssetName.ToString());
+			UE_LOG(LogCollectionManager, Log, TEXT("        %s"), *ObjectName.ToString());
 		}
 	}
 }
 
-void FCollection::GenerateHeaderPairs(TMap<FString,FString>& OutHeaderPairs)
+void FCollection::SaveHeaderPairs(TMap<FString,FString>& OutHeaderPairs) const
 {
 	// These pairs will appear at the top of the file followed by a newline
-	ensure(FileVersion > 0);
-	OutHeaderPairs.Add(TEXT("FileVersion"), FString::FromInt(FileVersion));
+	OutHeaderPairs.Add(TEXT("FileVersion"), FString::FromInt(ECollectionVersion::CurrentVersion)); // Files are always saved at the latest version as loading should take care of data upgrades
 	OutHeaderPairs.Add(TEXT("Type"), IsDynamic() ? TEXT("Dynamic") : TEXT("Static"));
+	OutHeaderPairs.Add(TEXT("Guid"), CollectionGuid.ToString(EGuidFormats::DigitsWithHyphens));
+	OutHeaderPairs.Add(TEXT("ParentGuid"), ParentCollectionGuid.ToString(EGuidFormats::DigitsWithHyphens));
 }
 
 bool FCollection::LoadHeaderPairs(const TMap<FString,FString>& InHeaderPairs)
@@ -417,14 +404,30 @@ bool FCollection::LoadHeaderPairs(const TMap<FString,FString>& InHeaderPairs)
 		return false;
 	}
 
-	FileVersion = FCString::Atoi(**Version);
+	FileVersion = (ECollectionVersion::Type)FCString::Atoi(**Version);
+
+	if (FileVersion >= ECollectionVersion::AddedCollectionGuid)
+	{
+		const FString* GuidStr = InHeaderPairs.Find(TEXT("Guid"));
+		if ( !GuidStr || !FGuid::Parse(*GuidStr, CollectionGuid) )
+		{
+			// Guid is required
+			return false;
+		}
+
+		const FString* ParentGuidStr = InHeaderPairs.Find(TEXT("ParentGuid"));
+		if ( !ParentGuidStr || !FGuid::Parse(*ParentGuidStr, ParentCollectionGuid) )
+		{
+			ParentCollectionGuid = FGuid();
+		}
+	}
 
 	if ( *Type == TEXT("Dynamic") )
 	{
 		// @todo Set this file up to be dynamic
 	}
 
-	return FileVersion > 0;
+	return FileVersion > 0 && FileVersion <= ECollectionVersion::CurrentVersion;
 }
 
 void FCollection::MergeWithCollection(const FCollection& Other)
@@ -436,45 +439,45 @@ void FCollection::MergeWithCollection(const FCollection& Other)
 	else
 	{
 		// Gather the differences from the file on disk
-		TArray<FName> AssetsAdded;
-		TArray<FName> AssetsRemoved;
-		GetDifferencesFromDisk(AssetsAdded, AssetsRemoved);
+		TArray<FName> ObjectsAdded;
+		TArray<FName> ObjectsRemoved;
+		GetObjectDifferencesFromDisk(ObjectsAdded, ObjectsRemoved);
 
 		// Copy asset list from other collection
-		DiskAssetSet = Other.DiskAssetSet;
-		AssetSet = DiskAssetSet;
+		DiskSnapshot = Other.DiskSnapshot;
+		ObjectSet = DiskSnapshot.ObjectSet;
 
-		// Add the assets that were added before the merge
-		for (const FName& AddedAssetName : AssetsAdded)
+		// Add the objects that were added before the merge
+		for (const FName& AddedObjectName : ObjectsAdded)
 		{
-			AssetSet.Add(AddedAssetName);
+			ObjectSet.Add(AddedObjectName);
 		}
 
-		// Remove the assets that were removed before the merge
-		for (const FName& RemovedAssetName : AssetsRemoved)
+		// Remove the objects that were removed before the merge
+		for (const FName& RemovedObjectName : ObjectsRemoved)
 		{
-			AssetSet.Remove(RemovedAssetName);
+			ObjectSet.Remove(RemovedObjectName);
 		}
 	}
 }
 
-void FCollection::GetDifferencesFromDisk(TArray<FName>& AssetsAdded, TArray<FName>& AssetsRemoved)
+void FCollection::GetObjectDifferencesFromDisk(TArray<FName>& ObjectsAdded, TArray<FName>& ObjectsRemoved)
 {
-	// Find the assets that were removed since the disk list
-	for (const FName& DiskAssetName : DiskAssetSet)
+	// Find the objects that were removed since the disk list
+	for (const FName& DiskObjectName : DiskSnapshot.ObjectSet)
 	{
-		if (!AssetSet.Contains(DiskAssetName))
+		if (!ObjectSet.Contains(DiskObjectName))
 		{
-			AssetsRemoved.Add(DiskAssetName);
+			ObjectsRemoved.Add(DiskObjectName);
 		}
 	}
 
-	// Find the assets that were added since the disk list
-	for (const FName& MemoryAssetName : AssetSet)
+	// Find the objects that were added since the disk list
+	for (const FName& MemoryObjectName : ObjectSet)
 	{
-		if (!DiskAssetSet.Contains(MemoryAssetName))
+		if (!DiskSnapshot.ObjectSet.Contains(MemoryObjectName))
 		{
-			AssetsAdded.Add(MemoryAssetName);
+			ObjectsAdded.Add(MemoryObjectName);
 		}
 	}
 }
@@ -528,21 +531,15 @@ bool FCollection::CheckoutCollection(FText& OutError)
 		}
 
 		// Check to see if the file exists at the head revision
-		if ( IFileManager::Get().FileSize(*SourceFilename) < 0 )
-		{
-			// File was deleted at head...
-			// Just add our changes to an empty list so we can mark it for add
-			FCollection NewCollection;
-			MergeWithCollection(NewCollection);
-		}
-		else
+		if ( IFileManager::Get().FileExists(*SourceFilename) )
 		{
 			// File found! Load it and merge with our local changes
-			FCollection NewCollection;
-			if ( !NewCollection.LoadFromFile(SourceFilename, false) )
+			FText LoadErrorText;
+			FCollection NewCollection(SourceFilename, false);
+			if ( !NewCollection.Load(LoadErrorText) )
 			{
 				// Failed to load the head revision file so it isn't safe to delete it
-				OutError = LOCTEXT("Error_SCCBadHead", "Failed to load the collection at the head revision.");
+				OutError = FText::Format(LOCTEXT("Error_SCCBadHead", "Failed to load the collection at the head revision. {0}"), LoadErrorText);
 				return false;
 			}
 
@@ -640,92 +637,87 @@ bool FCollection::CheckinCollection(FText& OutError)
 
 	// Form an appropriate summary for the changelist
 	const FText CollectionNameText = FText::FromName( CollectionName );
-	FText ChangelistDesc = FText::Format( LOCTEXT("CollectionGenericModifiedDesc", "Modified collection: {0}"), CollectionNameText );
-	if ( SourceControlState.IsValid() && SourceControlState->IsAdded() )
-	{
-		ChangelistDesc = FText::Format( LOCTEXT("CollectionAddedNewDesc", "Added collection"), CollectionNameText);
-	}
-	else if ( IsDynamic() )
+	FTextBuilder ChangelistDescBuilder;
+
+	if (SourceControlState.IsValid() && SourceControlState->IsAdded())
 	{
-		// @todo collection Change description for dynamic collections
+		ChangelistDescBuilder.AppendLineFormat(LOCTEXT("CollectionAddedNewDesc", "Added collection '{0}'"), CollectionNameText);
 	}
 	else
 	{
-		// Gather differences from disk
-		TArray<FName> AssetsAdded;
-		TArray<FName> AssetsRemoved;
-		GetDifferencesFromDisk(AssetsAdded, AssetsRemoved);
-
-		AssetsAdded.Sort();
-		AssetsRemoved.Sort();
+		if (IsDynamic())
+		{
+			// @todo collection Change description for dynamic collections
+		}
+		else
+		{
+			// Gather differences from disk
+			TArray<FName> ObjectsAdded;
+			TArray<FName> ObjectsRemoved;
+			GetObjectDifferencesFromDisk(ObjectsAdded, ObjectsRemoved);
 
-		// Clear description
-		ChangelistDesc = FText::GetEmpty();
+			ObjectsAdded.Sort();
+			ObjectsRemoved.Sort();
 
-		// Report added files
-		FFormatNamedArguments Args;
-		Args.Add( TEXT("AssetAdded"), AssetsAdded.Num() > 0 ? FText::FromName( AssetsAdded[0] ) : NSLOCTEXT( "Core", "None", "None" ) );
-		Args.Add( TEXT("NumberAdded"), FText::AsNumber( FMath::Max( AssetsAdded.Num() - 1, 0 ) ));
-		Args.Add( TEXT("CollectionName"), CollectionNameText );
+			// Report added files
+			FFormatNamedArguments Args;
+			Args.Add(TEXT("FirstObjectAdded"), ObjectsAdded.Num() > 0 ? FText::FromName(ObjectsAdded[0]) : NSLOCTEXT("Core", "None", "None"));
+			Args.Add(TEXT("NumberAdded"), FText::AsNumber(ObjectsAdded.Num()));
+			Args.Add(TEXT("FirstObjectRemoved"), ObjectsRemoved.Num() > 0 ? FText::FromName(ObjectsRemoved[0]) : NSLOCTEXT("Core", "None", "None"));
+			Args.Add(TEXT("NumberRemoved"), FText::AsNumber(ObjectsRemoved.Num()));
+			Args.Add(TEXT("CollectionName"), CollectionNameText);
 
-		if ( AssetsRemoved.Num() == 0 )
-		{
-			if (AssetsAdded.Num() == 1)
+			if (ObjectsAdded.Num() == 1)
 			{
-				ChangelistDesc = FText::Format( LOCTEXT("CollectionAddedSingleDesc", "Added {AssetAdded} to collection: {CollectionName}"), Args );
+				ChangelistDescBuilder.AppendLineFormat(LOCTEXT("CollectionAddedSingleDesc", "Added '{FirstObjectAdded}' to collection '{CollectionName}'"), Args);
 			}
-			else if (AssetsAdded.Num() > 1)
+			else if (ObjectsAdded.Num() > 1)
 			{
-				ChangelistDesc = FText::Format( LOCTEXT("CollectionAddedMultipleDesc", "Added {AssetAdded} and {NumberAdded} other(s) to collection: {CollectionName}"), Args );
-			}
-		}
-		else
-		{
-			Args.Add( TEXT("AssetRemoved"), FText::FromName( AssetsRemoved[0] ) );
-			Args.Add( TEXT("NumberRemoved"), FText::AsNumber( AssetsRemoved.Num() - 1 ) );
+				ChangelistDescBuilder.AppendLineFormat(LOCTEXT("CollectionAddedMultipleDesc", "Added {NumberAdded} objects to collection '{CollectionName}':"), Args);
 
-			if ( AssetsAdded.Num() == 1 )
-			{
-				if ( AssetsRemoved.Num() == 1 )
-				{
-					ChangelistDesc = FText::Format( LOCTEXT("CollectionRemovedSingle_AddedSingleDesc", "Added {AssetAdded} to collection: {CollectionName} : Removed {AssetRemoved} from collection: {CollectionName}"), Args );
-				}
-				else if (AssetsRemoved.Num() > 1)
+				ChangelistDescBuilder.Indent();
+				for (const FName& AddedObjectName : ObjectsAdded)
 				{
-					ChangelistDesc = FText::Format( LOCTEXT("CollectionRemovedMultiple_AddedSingleDesc", "Added {AssetAdded} to collection: {CollectionName} : Removed {AssetRemoved} and {NumberRemoved} other(s) from collection: {CollectionName}"), Args );
+					ChangelistDescBuilder.AppendLine(FText::FromName(AddedObjectName));
 				}
+				ChangelistDescBuilder.Unindent();
 			}
-			else if (AssetsAdded.Num() > 1)
+
+			if ( ObjectsRemoved.Num() == 1 )
 			{
-				if ( AssetsRemoved.Num() == 1 )
-				{
-					ChangelistDesc = FText::Format( LOCTEXT("CollectionRemovedSingle_AddedMultpleDesc", "Added {AssetAdded} and {NumberAdded} other(s) to collection: {CollectionName} : Removed {AssetRemoved} from collection: {CollectionName}"), Args );
-				}
-				else if (AssetsRemoved.Num() > 1)
-				{
-					ChangelistDesc = FText::Format( LOCTEXT("CollectionRemovedMultiple_AddedMultpleDesc", "Added {AssetAdded} and {NumberAdded} other(s) to collection: {CollectionName} : Removed {AssetRemoved} and {NumberRemoved} other(s) from collection: {CollectionName}"), Args );
-				}
+				ChangelistDescBuilder.AppendLineFormat(LOCTEXT("CollectionRemovedSingleDesc", "Removed '{FirstObjectRemoved}' from collection '{CollectionName}'"), Args);
 			}
-			else
+			else if (ObjectsRemoved.Num() > 1)
 			{
-				if ( AssetsRemoved.Num() == 1 )
-				{
-					ChangelistDesc = FText::Format( LOCTEXT("CollectionRemovedSingleDesc", "Removed {AssetRemoved} from collection: {CollectionName}"), Args );
-				}
-				else if (AssetsRemoved.Num() > 1)
+				ChangelistDescBuilder.AppendLineFormat(LOCTEXT("CollectionRemovedMultipleDesc", "Removed {NumberRemoved} objects from collection '{CollectionName}'"), Args);
+
+				ChangelistDescBuilder.Indent();
+				for (const FName& RemovedObjectName : ObjectsRemoved)
 				{
-					ChangelistDesc = FText::Format( LOCTEXT("CollectionRemovedMultipleDesc", "Removed {AssetRemoved} and {NumberRemoved} other(s) from collection: {CollectionName}"), Args );
+					ChangelistDescBuilder.AppendLine(FText::FromName(RemovedObjectName));
 				}
+				ChangelistDescBuilder.Unindent();
 			}
 		}
+
+		// Parent change?
+		if (DiskSnapshot.ParentCollectionGuid != ParentCollectionGuid)
+		{
+			ChangelistDescBuilder.AppendLineFormat(LOCTEXT("CollectionChangedParentDesc", "Changed the parent of collection '{0}'"), CollectionNameText);
+		}
+
+		// Version bump?
+		if (FileVersion < ECollectionVersion::CurrentVersion)
+		{
+			ChangelistDescBuilder.AppendLineFormat(LOCTEXT("CollectionUpgradedDesc", "Upgraded collection '{0}' (was version {1}, now version {2})"), CollectionNameText, FText::AsNumber(FileVersion), FText::AsNumber(ECollectionVersion::CurrentVersion));
+		}
 	}
 
-	if ( ChangelistDesc.IsEmpty() )
+	FText ChangelistDesc = ChangelistDescBuilder.ToText();
+	if (ChangelistDesc.IsEmpty())
 	{
-		// No files were added or removed
-		FFormatNamedArguments Args;
-		Args.Add( TEXT("CollectionName"), CollectionNameText );
-		ChangelistDesc = FText::Format( LOCTEXT("CollectionNotModifiedDesc", "Collection not modified: {CollectionName}"), Args );
+		// No changes could be detected
+		ChangelistDesc = FText::Format(LOCTEXT("CollectionNotModifiedDesc", "Collection '{0}' not modified"), CollectionNameText);
 	}
 
 	// Finally check in the file
@@ -843,20 +835,20 @@ bool FCollection::DeleteFromSourceControl(FText& OutError)
 		}
 
 		// Check to see if the file exists at the head revision
-		if ( IFileManager::Get().FileSize(*SourceFilename) < 0 )
+		if ( !IFileManager::Get().FileExists(*SourceFilename) )
 		{
 			// File was already deleted, consider this a success
 			GWarn->EndSlowTask();
 			return true;
 		}
 			
-		FCollection NewCollection;
-				
-		if ( !NewCollection.LoadFromFile(SourceFilename, false) )
+		FCollection NewCollection(SourceFilename, false);
+		FText LoadErrorText;
+		if ( !NewCollection.Load(LoadErrorText) )
 		{
 			// Failed to load the head revision file so it isn't safe to delete it
 			GWarn->EndSlowTask();
-			OutError = LOCTEXT("Error_SCCBadHead", "Failed to load the collection at the head revision.");
+			OutError = FText::Format(LOCTEXT("Error_SCCBadHead", "Failed to load the collection at the head revision. {0}"), LoadErrorText);
 			return false;
 		}
 
diff --git a/Engine/Source/Developer/CollectionManager/Private/Collection.h b/Engine/Source/Developer/CollectionManager/Private/Collection.h
--- a/Engine/Source/Developer/CollectionManager/Private/Collection.h
+++ b/Engine/Source/Developer/CollectionManager/Private/Collection.h
@@ -2,23 +2,54 @@
 
 #pragma once
 
+struct ECollectionVersion
+{
+	enum Type
+	{
+		_ZeroVersion = 0,
+
+		/** The initial version for collection files */
+		Initial,
+
+		/** 
+		 * Added GUIDs to collections to allow them to be used as a parent of another collection without relying on their name/share-type combo
+		 * Collections that are older than this version must be re-saved before they can be used as a parent for another collection
+		 */
+		AddedCollectionGuid,
+
+		/** -----<new versions can be added before this line>------------------------------------------------- */
+		_AutomaticVersionPlusOne,
+		CurrentVersion = _AutomaticVersionPlusOne - 1
+	};
+};
+
+enum class ECollectionCloneMode : uint8
+{
+	/** Clone this collection exactly as it is now, preserving its existing GUID data */
+	Exact,
+	/** Clone this collection, but make sure it gets unique GUIDs */
+	Unique,
+};
+
 /** A class to represent a collection of assets */
 class FCollection
 {
 public:
-	FCollection();
-	FCollection(FName InCollectionName, const FString& SourceFolder, const FString& CollectionExtension, bool InUseSCC);
+	FCollection(const FString& InFilename, bool InUseSCC);
 
-	/** Loads content from a file into this collection. */
-	bool LoadFromFile(const FString& InFilename, bool InUseSCC);
+	/** Clone this collection to a new location */
+	TSharedRef<FCollection> Clone(const FString& InFilename, bool InUseSCC, ECollectionCloneMode InCloneMode) const;
+
+	/** Loads content from the SourceFilename into this collection. If false, OutError is a human readable warning depicting the error. */
+	bool Load(FText& OutError);
 	/** Saves this collection to SourceFilename. If false, OutError is a human readable warning depicting the error. */
 	bool Save(FText& OutError);
 	/** Deletes the source file for this collection. If false, OutError is a human readable warning depicting the error. */
 	bool DeleteSourceFile(FText& OutError);
 
-	/** Adds a single asset to the collection */
+	/** Adds a single asset to the collection. Static collections only. */
 	bool AddAssetToCollection(FName ObjectPath);
-	/** Removes a single asset from the collection */
+	/** Removes a single asset from the collection. Static collections only. */
 	bool RemoveAssetFromCollection(FName ObjectPath);
 	/** Gets a list of assets in the collection. Static collections only. */
 	void GetAssetsInCollection(TArray<FName>& Assets) const;
@@ -28,8 +59,6 @@ public:
 	void GetObjectsInCollection(TArray<FName>& Objects) const;
 	/** Returns true when the specified object is in the collection. Static collections only. */
 	bool IsObjectInCollection(FName ObjectPath) const;
-	/** Resets a collection to its default values */
-	void Clear();
 
 	/** Returns true if the collection contains rules instead of a flat list */
 	bool IsDynamic() const;
@@ -43,9 +72,21 @@ public:
 	/** Returns the name of the collection */
 	FORCEINLINE const FName& GetCollectionName() const { return CollectionName; }
 
+	/** Returns the GUID of the collection */
+	FORCEINLINE const FGuid& GetCollectionGuid() const { return CollectionGuid; }
+
+	/** Returns the GUID of the collection we are parented under */
+	FORCEINLINE const FGuid& GetParentCollectionGuid() const { return ParentCollectionGuid; }
+
+	/** Set the GUID of the collection we are parented under */
+	FORCEINLINE void SetParentCollectionGuid(const FGuid& NewGuid) { ParentCollectionGuid = NewGuid; }
+
+	/** Returns the file version of the collection */
+	FORCEINLINE ECollectionVersion::Type GetCollectionVersion() const { return FileVersion; }
+
 private:
 	/** Generates the header pairs for the collection file. */
-	void GenerateHeaderPairs(TMap<FString,FString>& OutHeaderPairs);
+	void SaveHeaderPairs(TMap<FString,FString>& OutHeaderPairs) const;
 
 	/** 
 	  * Processes header pairs from the top of a collection file.
@@ -57,8 +98,8 @@ private:
 
 	/** Merges the assets from the specified collection with this collection */
 	void MergeWithCollection(const FCollection& Other);
-	/** Gets the differences between Lists A and B */
-	void GetDifferencesFromDisk(TArray<FName>& AssetsAdded, TArray<FName>& AssetsRemoved);
+	/** Gets the object differences between collections A and B */
+	void GetObjectDifferencesFromDisk(TArray<FName>& ObjectsAdded, TArray<FName>& ObjectsRemoved);
 	/** Checks the shared collection out from source control so it may be saved. If false, OutError is a human readable warning depicting the error. */
 	bool CheckoutCollection(FText& OutError);
 	/** Checks the shared collection in to source control after it is saved. If false, OutError is a human readable warning depicting the error. */
@@ -69,21 +110,43 @@ private:
 	bool DeleteFromSourceControl(FText& OutError);
 
 private:
+	/** Snapshot data for a collection. Used to take snapshots and provide a diff message */
+	struct FCollectionSnapshot
+	{
+		void TakeSnapshot(const FCollection& InCollection)
+		{
+			ParentCollectionGuid = InCollection.ParentCollectionGuid;
+			ObjectSet = InCollection.ObjectSet;
+		}
+
+		/** The GUID of the collection we are parented under */
+		FGuid ParentCollectionGuid;
+
+		/** The set of objects in the collection. Takes the form PackageName.AssetName */
+		TSet<FName> ObjectSet;
+	};
+
 	/** The name of the collection */
 	FName CollectionName;
 
+	/** The GUID of the collection */
+	FGuid CollectionGuid;
+
+	/** The GUID of the collection we are parented under */
+	FGuid ParentCollectionGuid;
+
 	/** Source control is used if true */
 	bool bUseSCC;
 
 	/** The filename used to load this collection. Empty if it is new or never loaded from disk. */
 	FString SourceFilename;
 
-	/** The set of assets in the collection. Takes the form PackageName.AssetName */
-	TSet<FName> AssetSet;
-
-	/** The set of assets that were in the collection last time it was loaded from or saved to disk. Takes the form PackageName.AssetName */
-	TSet<FName> DiskAssetSet;
+	/** The set of objects in the collection. Takes the form PackageName.AssetName */
+	TSet<FName> ObjectSet;
 
 	/** The file version for this collection */
-	int32 FileVersion;
+	ECollectionVersion::Type FileVersion;
+
+	/** The state of the collection the last time it was loaded from or saved to disk. */
+	FCollectionSnapshot DiskSnapshot;
 };
diff --git a/Engine/Source/Developer/CollectionManager/Private/CollectionManager.cpp b/Engine/Source/Developer/CollectionManager/Private/CollectionManager.cpp
--- a/Engine/Source/Developer/CollectionManager/Private/CollectionManager.cpp
+++ b/Engine/Source/Developer/CollectionManager/Private/CollectionManager.cpp
@@ -26,6 +26,16 @@ bool FCollectionManager::HasCollections() const
 	return CachedCollections.Num() > 0;
 }
 
+void FCollectionManager::GetCollections(TArray<FCollectionNameType>& OutCollections) const
+{
+	OutCollections.Reserve(CachedCollections.Num());
+	for (const auto& CachedCollection : CachedCollections)
+	{
+		const FCollectionNameType& CollectionKey = CachedCollection.Key;
+		OutCollections.Add(CollectionKey);
+	}
+}
+
 void FCollectionManager::GetCollectionNames(ECollectionShareType::Type ShareType, TArray<FName>& CollectionNames) const
 {
 	for (const auto& CachedCollection : CachedCollections)
@@ -33,9 +43,115 @@ void FCollectionManager::GetCollectionNames(ECollectionShareType::Type ShareType
 		const FCollectionNameType& CollectionKey = CachedCollection.Key;
 		if (ShareType == ECollectionShareType::CST_All || ShareType == CollectionKey.Type)
 		{
-			CollectionNames.Add(CollectionKey.Name);
+			CollectionNames.AddUnique(CollectionKey.Name);
+		}
+	}
+}
+
+void FCollectionManager::GetRootCollections(TArray<FCollectionNameType>& OutCollections) const
+{
+	OutCollections.Reserve(CachedCollections.Num());
+	for (const auto& CachedCollection : CachedCollections)
+	{
+		const FCollectionNameType& CollectionKey = CachedCollection.Key;
+		const TSharedRef<FCollection>& Collection = CachedCollection.Value;
+
+		// A root collection either has no parent GUID, or a parent GUID that cannot currently be found - the check below handles both
+		if (!CachedCollectionNamesFromGuids.Contains(Collection->GetParentCollectionGuid()))
+		{
+			OutCollections.Add(CollectionKey);
+		}
+	}
+}
+
+void FCollectionManager::GetRootCollectionNames(ECollectionShareType::Type ShareType, TArray<FName>& CollectionNames) const
+{
+	for (const auto& CachedCollection : CachedCollections)
+	{
+		const FCollectionNameType& CollectionKey = CachedCollection.Key;
+		const TSharedRef<FCollection>& Collection = CachedCollection.Value;
+
+		if (ShareType == ECollectionShareType::CST_All || ShareType == CollectionKey.Type)
+		{
+			// A root collection either has no parent GUID, or a parent GUID that cannot currently be found - the check below handles both
+			if (!CachedCollectionNamesFromGuids.Contains(Collection->GetParentCollectionGuid()))
+			{
+				CollectionNames.AddUnique(CollectionKey.Name);
+			}
+		}
+	}
+}
+
+void FCollectionManager::GetChildCollections(FName CollectionName, ECollectionShareType::Type ShareType, TArray<FCollectionNameType>& OutCollections) const
+{
+	auto GetChildCollectionsInternal = [this, &OutCollections](const FCollectionNameType& InCollectionKey)
+	{
+		const TSharedRef<FCollection>* const CollectionRefPtr = CachedCollections.Find(InCollectionKey);
+		if (CollectionRefPtr)
+		{
+			const auto* ChildCollectionGuids = CachedHierarchy.Find((*CollectionRefPtr)->GetCollectionGuid());
+			if (ChildCollectionGuids)
+			{
+				for (const FGuid& ChildCollectionGuid : *ChildCollectionGuids)
+				{
+					const FCollectionNameType* const ChildCollectionKeyPtr = CachedCollectionNamesFromGuids.Find(ChildCollectionGuid);
+					if (ChildCollectionKeyPtr)
+					{
+						OutCollections.Add(*ChildCollectionKeyPtr);
+					}
+				}
+			}
+		}
+	};
+
+	if (ShareType == ECollectionShareType::CST_All)
+	{
+		// Asked for all share types, find children in the specified collection name in any cache
+		for (int32 CacheIdx = 0; CacheIdx < ECollectionShareType::CST_All; ++CacheIdx)
+		{
+			GetChildCollectionsInternal(FCollectionNameType(CollectionName, ECollectionShareType::Type(CacheIdx)));
 		}
 	}
+	else
+	{
+		GetChildCollectionsInternal(FCollectionNameType(CollectionName, ShareType));
+	}
+}
+
+void FCollectionManager::GetChildCollectionNames(FName CollectionName, ECollectionShareType::Type ShareType, ECollectionShareType::Type ChildShareType, TArray<FName>& CollectionNames) const
+{
+	auto GetChildCollectionsInternal = [this, &ChildShareType, &CollectionNames](const FCollectionNameType& InCollectionKey)
+	{
+		const TSharedRef<FCollection>* const CollectionRefPtr = CachedCollections.Find(InCollectionKey);
+		if (CollectionRefPtr)
+		{
+			const auto* ChildCollectionGuids = CachedHierarchy.Find((*CollectionRefPtr)->GetCollectionGuid());
+			if (ChildCollectionGuids)
+			{
+				for (const FGuid& ChildCollectionGuid : *ChildCollectionGuids)
+				{
+					const FCollectionNameType* const ChildCollectionKeyPtr = CachedCollectionNamesFromGuids.Find(ChildCollectionGuid);
+					if (ChildCollectionKeyPtr && (ChildShareType == ECollectionShareType::CST_All || ChildShareType == ChildCollectionKeyPtr->Type))
+					{
+						CollectionNames.AddUnique(ChildCollectionKeyPtr->Name);
+					}
+				}
+			}
+		}
+	};
+
+	if (ShareType == ECollectionShareType::CST_All)
+	{
+		// Asked for all share types, find children in the specified collection name in any cache
+		for (int32 CacheIdx = 0; CacheIdx < ECollectionShareType::CST_All; ++CacheIdx)
+		{
+			GetChildCollectionsInternal(FCollectionNameType(CollectionName, ECollectionShareType::Type(CacheIdx)));
+		}
+	}
+	else
+	{
+		GetChildCollectionsInternal(FCollectionNameType(CollectionName, ShareType));
+	}
 }
 
 bool FCollectionManager::CollectionExists(FName CollectionName, ECollectionShareType::Type ShareType) const
@@ -61,145 +177,161 @@ bool FCollectionManager::CollectionExists(FName CollectionName, ECollectionShare
 	}
 }
 
-bool FCollectionManager::GetAssetsInCollection(FName CollectionName, ECollectionShareType::Type ShareType, TArray<FName>& AssetsPaths) const
+bool FCollectionManager::GetAssetsInCollection(FName CollectionName, ECollectionShareType::Type ShareType, TArray<FName>& AssetsPaths, ECollectionRecursionFlags::Flags RecursionMode) const
 {
-	auto GetAssetsInCollectionInternal = [this, &AssetsPaths](const FCollectionNameType& InCollectionKey) -> bool
+	bool bFoundAssets = false;
+
+	auto GetAssetsInCollectionWorker = [&](const FCollectionNameType& InCollectionKey, ECollectionRecursionFlags::Flag InReason) -> ERecursiveWorkerFlowControl
 	{
 		const TSharedRef<FCollection>* const CollectionRefPtr = CachedCollections.Find(InCollectionKey);
 		if (CollectionRefPtr)
 		{
 			(*CollectionRefPtr)->GetAssetsInCollection(AssetsPaths);
-			return true;
+			bFoundAssets = true;
 		}
-		return false;
+		return ERecursiveWorkerFlowControl::Continue;
 	};
 
 	if (ShareType == ECollectionShareType::CST_All)
 	{
 		// Asked for all share types, find assets in the specified collection name in any cache
-		bool bFoundAssets = false;
 		for (int32 CacheIdx = 0; CacheIdx < ECollectionShareType::CST_All; ++CacheIdx)
 		{
-			bFoundAssets |= GetAssetsInCollectionInternal(FCollectionNameType(CollectionName, ECollectionShareType::Type(CacheIdx)));
+			RecursionHelper_DoWork(FCollectionNameType(CollectionName, ECollectionShareType::Type(CacheIdx)), RecursionMode, GetAssetsInCollectionWorker);
 		}
-		return bFoundAssets;
 	}
 	else
 	{
-		return GetAssetsInCollectionInternal(FCollectionNameType(CollectionName, ShareType));
+		RecursionHelper_DoWork(FCollectionNameType(CollectionName, ShareType), RecursionMode, GetAssetsInCollectionWorker);
 	}
+
+	return bFoundAssets;
 }
 
-bool FCollectionManager::GetClassesInCollection(FName CollectionName, ECollectionShareType::Type ShareType, TArray<FName>& ClassPaths) const
+bool FCollectionManager::GetClassesInCollection(FName CollectionName, ECollectionShareType::Type ShareType, TArray<FName>& ClassPaths, ECollectionRecursionFlags::Flags RecursionMode) const
 {
-	auto GetClassesInCollectionInternal = [this, &ClassPaths](const FCollectionNameType& InCollectionKey) -> bool
+	bool bFoundClasses = false;
+
+	auto GetClassesInCollectionWorker = [&](const FCollectionNameType& InCollectionKey, ECollectionRecursionFlags::Flag InReason) -> ERecursiveWorkerFlowControl
 	{
 		const TSharedRef<FCollection>* const CollectionRefPtr = CachedCollections.Find(InCollectionKey);
 		if (CollectionRefPtr)
 		{
 			(*CollectionRefPtr)->GetClassesInCollection(ClassPaths);
-			return true;
+			bFoundClasses = true;
 		}
-		return false;
+		return ERecursiveWorkerFlowControl::Continue;
 	};
 
 	if (ShareType == ECollectionShareType::CST_All)
 	{
 		// Asked for all share types, find classes in the specified collection name in any cache
-		bool bFoundAssets = false;
 		for (int32 CacheIdx = 0; CacheIdx < ECollectionShareType::CST_All; ++CacheIdx)
 		{
-			bFoundAssets |= GetClassesInCollectionInternal(FCollectionNameType(CollectionName, ECollectionShareType::Type(CacheIdx)));
+			RecursionHelper_DoWork(FCollectionNameType(CollectionName, ECollectionShareType::Type(CacheIdx)), RecursionMode, GetClassesInCollectionWorker);
 		}
-		return bFoundAssets;
 	}
 	else
 	{
-		return GetClassesInCollectionInternal(FCollectionNameType(CollectionName, ShareType));
+		RecursionHelper_DoWork(FCollectionNameType(CollectionName, ShareType), RecursionMode, GetClassesInCollectionWorker);
 	}
+
+	return bFoundClasses;
 }
 
-bool FCollectionManager::GetObjectsInCollection(FName CollectionName, ECollectionShareType::Type ShareType, TArray<FName>& ObjectPaths) const
+bool FCollectionManager::GetObjectsInCollection(FName CollectionName, ECollectionShareType::Type ShareType, TArray<FName>& ObjectPaths, ECollectionRecursionFlags::Flags RecursionMode) const
 {
-	auto GetObjectsInCollectionInternal = [this, &ObjectPaths](const FCollectionNameType& InCollectionKey) -> bool
+	bool bFoundObjects = false;
+
+	auto GetObjectsInCollectionWorker = [&](const FCollectionNameType& InCollectionKey, ECollectionRecursionFlags::Flag InReason) -> ERecursiveWorkerFlowControl
 	{
 		const TSharedRef<FCollection>* const CollectionRefPtr = CachedCollections.Find(InCollectionKey);
 		if (CollectionRefPtr)
 		{
 			(*CollectionRefPtr)->GetObjectsInCollection(ObjectPaths);
-			return true;
+			bFoundObjects = true;
 		}
-		return false;
+		return ERecursiveWorkerFlowControl::Continue;
 	};
 
 	if (ShareType == ECollectionShareType::CST_All)
 	{
-		// Asked for all share types, find objects in the specified collection name in any cache
-		bool bFoundAssets = false;
+		// Asked for all share types, find classes in the specified collection name in any cache
 		for (int32 CacheIdx = 0; CacheIdx < ECollectionShareType::CST_All; ++CacheIdx)
 		{
-			bFoundAssets |= GetObjectsInCollectionInternal(FCollectionNameType(CollectionName, ECollectionShareType::Type(CacheIdx)));
+			RecursionHelper_DoWork(FCollectionNameType(CollectionName, ECollectionShareType::Type(CacheIdx)), RecursionMode, GetObjectsInCollectionWorker);
 		}
-		return bFoundAssets;
 	}
 	else
 	{
-		return GetObjectsInCollectionInternal(FCollectionNameType(CollectionName, ShareType));
+		RecursionHelper_DoWork(FCollectionNameType(CollectionName, ShareType), RecursionMode, GetObjectsInCollectionWorker);
 	}
+
+	return bFoundObjects;
 }
 
-void FCollectionManager::GetCollectionsContainingObject(FName ObjectPath, ECollectionShareType::Type ShareType, TArray<FName>& OutCollectionNames) const
+void FCollectionManager::GetCollectionsContainingObject(FName ObjectPath, ECollectionShareType::Type ShareType, TArray<FName>& OutCollectionNames, ECollectionRecursionFlags::Flags RecursionMode) const
 {
-	const auto* ObjectCollectionsPtr = CachedObjects.Find(ObjectPath);
-	if (ObjectCollectionsPtr)
+	const auto* ObjectCollectionInfosPtr = CachedObjects.Find(ObjectPath);
+	if (ObjectCollectionInfosPtr)
 	{
-		for (const FCollectionNameType& CollectionKey : *ObjectCollectionsPtr)
+		for (const FObjectCollectionInfo& ObjectCollectionInfo : *ObjectCollectionInfosPtr)
 		{
-			if (ShareType == ECollectionShareType::CST_All || ShareType == CollectionKey.Type)
+			if ((ShareType == ECollectionShareType::CST_All || ShareType == ObjectCollectionInfo.CollectionKey.Type) && (RecursionMode & ObjectCollectionInfo.Reason) != 0)
 			{
-				OutCollectionNames.Add(CollectionKey.Name);
+				OutCollectionNames.Add(ObjectCollectionInfo.CollectionKey.Name);
 			}
 		}
 	}
 }
 
-void FCollectionManager::GetCollectionsContainingObject(FName ObjectPath, TArray<FCollectionNameType>& OutCollections) const
+void FCollectionManager::GetCollectionsContainingObject(FName ObjectPath, TArray<FCollectionNameType>& OutCollections, ECollectionRecursionFlags::Flags RecursionMode) const
 {
-	const auto* ObjectCollectionsPtr = CachedObjects.Find(ObjectPath);
-	if (ObjectCollectionsPtr)
+	const auto* ObjectCollectionInfosPtr = CachedObjects.Find(ObjectPath);
+	if (ObjectCollectionInfosPtr)
 	{
-		OutCollections.Append(*ObjectCollectionsPtr);
+		OutCollections.Reserve(OutCollections.Num() + ObjectCollectionInfosPtr->Num());
+		for (const FObjectCollectionInfo& ObjectCollectionInfo : *ObjectCollectionInfosPtr)
+		{
+			if ((RecursionMode & ObjectCollectionInfo.Reason) != 0)
+			{
+				OutCollections.Add(ObjectCollectionInfo.CollectionKey);
+			}
+		}
 	}
 }
 
-void FCollectionManager::GetCollectionsContainingObjects(const TArray<FName>& ObjectPaths, TMap<FCollectionNameType, TArray<FName>>& OutCollectionsAndMatchedObjects) const
+void FCollectionManager::GetCollectionsContainingObjects(const TArray<FName>& ObjectPaths, TMap<FCollectionNameType, TArray<FName>>& OutCollectionsAndMatchedObjects, ECollectionRecursionFlags::Flags RecursionMode) const
 {
 	for (const FName& ObjectPath : ObjectPaths)
 	{
-		const auto* ObjectCollectionsPtr = CachedObjects.Find(ObjectPath);
-		if (ObjectCollectionsPtr)
+		const auto* ObjectCollectionInfosPtr = CachedObjects.Find(ObjectPath);
+		if (ObjectCollectionInfosPtr)
 		{
-			for (const FCollectionNameType& CollectionKey : *ObjectCollectionsPtr)
+			for (const FObjectCollectionInfo& ObjectCollectionInfo : *ObjectCollectionInfosPtr)
 			{
-				TArray<FName>& MatchedObjects = OutCollectionsAndMatchedObjects.FindOrAdd(CollectionKey);
-				MatchedObjects.Add(ObjectPath);
+				if ((RecursionMode & ObjectCollectionInfo.Reason) != 0)
+				{
+					TArray<FName>& MatchedObjects = OutCollectionsAndMatchedObjects.FindOrAdd(ObjectCollectionInfo.CollectionKey);
+					MatchedObjects.Add(ObjectPath);
+				}
 			}
 		}
 	}
 }
 
-FString FCollectionManager::GetCollectionsStringForObject(FName ObjectPath, ECollectionShareType::Type ShareType) const
+FString FCollectionManager::GetCollectionsStringForObject(FName ObjectPath, ECollectionShareType::Type ShareType, ECollectionRecursionFlags::Flags RecursionMode) const
 {
-	const auto* ObjectCollectionsPtr = CachedObjects.Find(ObjectPath);
-	if (ObjectCollectionsPtr)
+	const auto* ObjectCollectionInfosPtr = CachedObjects.Find(ObjectPath);
+	if (ObjectCollectionInfosPtr)
 	{
 		TArray<FString> CollectionNameStrings;
 
-		for (const FCollectionNameType& CollectionKey : *ObjectCollectionsPtr)
+		for (const FObjectCollectionInfo& ObjectCollectionInfo : *ObjectCollectionInfosPtr)
 		{
-			if (ShareType == ECollectionShareType::CST_All || ShareType == CollectionKey.Type)
+			if ((ShareType == ECollectionShareType::CST_All || ShareType == ObjectCollectionInfo.CollectionKey.Type) && (RecursionMode & ObjectCollectionInfo.Reason) != 0)
 			{
-				CollectionNameStrings.Add(CollectionKey.Name.ToString());
+				CollectionNameStrings.Add(ObjectCollectionInfo.CollectionKey.Name.ToString());
 			}
 		}
 
@@ -249,10 +381,10 @@ bool FCollectionManager::CreateCollection(FName CollectionName, ECollectionShare
 	}
 
 	// Try to add the collection
-	bool bUseSCC = ShouldUseSCC(ShareType);
-	FString SourceFolder = CollectionFolders[ShareType];
+	const bool bUseSCC = ShouldUseSCC(ShareType);
+	const FString CollectionFilename = GetCollectionFilename(CollectionName, ShareType);
 
-	TSharedRef<FCollection> NewCollection = MakeShareable(new FCollection(CollectionName, SourceFolder, CollectionExtension, bUseSCC));
+	TSharedRef<FCollection> NewCollection = MakeShareable(new FCollection(CollectionFilename, bUseSCC));
 	if (!AddCollection(NewCollection, ShareType))
 	{
 		// Failed to add the collection, it already exists
@@ -283,46 +415,158 @@ bool FCollectionManager::RenameCollection(FName CurrentCollectionName, ECollecti
 		return false;
 	}
 
-	// Get the object paths for the assets in the collection
-	TArray<FName> ObjectPaths;
-	if (!GetAssetsInCollection(CurrentCollectionName, CurrentShareType, ObjectPaths))
+	const FCollectionNameType OriginalCollectionKey(CurrentCollectionName, CurrentShareType);
+	const FCollectionNameType NewCollectionKey(NewCollectionName, NewShareType);
+
+	TSharedRef<FCollection>* const CollectionRefPtr = CachedCollections.Find(OriginalCollectionKey);
+	if (!CollectionRefPtr)
+	{
+		// The collection doesn't exist
+		LastError = LOCTEXT("Error_DoesntExist", "The collection doesn't exist.");
+		return false;
+	}
+
+	// Add the new collection
+	TSharedPtr<FCollection> NewCollection;
+	{
+		const bool bUseSCC = ShouldUseSCC(NewShareType);
+		const FString NewCollectionFilename = GetCollectionFilename(NewCollectionName, NewShareType);
+
+		// Create an exact copy of the collection using its new path - this will preserve its GUID and avoid losing hierarchy data
+		NewCollection = (*CollectionRefPtr)->Clone(NewCollectionFilename, bUseSCC, ECollectionCloneMode::Exact);
+		if (!AddCollection(NewCollection.ToSharedRef(), NewShareType))
+		{
+			// Failed to add the collection, it already exists
+			LastError = LOCTEXT("Error_AlreadyExists", "The collection already exists.");
+			return false;
+		}
+
+		if (!NewCollection->Save(LastError))
+		{
+			// Collection failed to save, remove it from the cache
+			RemoveCollection(NewCollection.ToSharedRef(), NewShareType);
+
+			// Restore the original GUID mapping since AddCollection replaced it
+			CachedCollectionNamesFromGuids.Add((*CollectionRefPtr)->GetCollectionGuid(), OriginalCollectionKey);
+
+			return false;
+		}
+	}
+
+	// Remove the old collection
+	{
+		if ((*CollectionRefPtr)->DeleteSourceFile(LastError))
+		{
+			RemoveCollection(*CollectionRefPtr, CurrentShareType);
+
+			// Re-add the new GUID mapping since RemoveCollection removed it
+			CachedCollectionNamesFromGuids.Add(NewCollection->GetCollectionGuid(), NewCollectionKey);
+		}
+		else
+		{
+			// Failed to remove the old collection, so remove the collection we created.
+			NewCollection->DeleteSourceFile(LastError);
+			RemoveCollection(NewCollection.ToSharedRef(), NewShareType);
+
+			// Restore the original GUID mapping since RemoveCollection removed it
+			CachedCollectionNamesFromGuids.Add((*CollectionRefPtr)->GetCollectionGuid(), OriginalCollectionKey);
+
+			return false;
+		}
+	}
+
+	RebuildCachedObjects();
+
+	// Success
+	CollectionRenamedEvent.Broadcast(OriginalCollectionKey, NewCollectionKey);
+	return true;
+}
+
+bool FCollectionManager::ReparentCollection(FName CollectionName, ECollectionShareType::Type ShareType, FName ParentCollectionName, ECollectionShareType::Type ParentShareType)
+{
+	if (!ensure(ShareType < ECollectionShareType::CST_All) || (!ParentCollectionName.IsNone() && !ensure(ParentShareType < ECollectionShareType::CST_All)))
 	{
-		// Failed to get assets in the current collection
+		// Bad share type
+		LastError = LOCTEXT("Error_Internal", "There was an internal error.");
 		return false;
 	}
 
-	// Create a new collection
-	if (!CreateCollection(NewCollectionName, NewShareType))
+	const FCollectionNameType CollectionKey(CollectionName, ShareType);
+	TSharedRef<FCollection>* const CollectionRefPtr = CachedCollections.Find(CollectionKey);
+	if (!CollectionRefPtr)
 	{
-		// Failed to create collection
+		// The collection doesn't exist
+		LastError = LOCTEXT("Error_DoesntExist", "The collection doesn't exist.");
 		return false;
 	}
 
-	if (ObjectPaths.Num() > 0)
+	const FGuid OldParentGuid = (*CollectionRefPtr)->GetParentCollectionGuid();
+	FGuid NewParentGuid;
+
+	TOptional<FCollectionNameType> OldParentCollectionKey;
+	TOptional<FCollectionNameType> NewParentCollectionKey;
+
+	if (!ParentCollectionName.IsNone())
 	{
-		// Add all the objects from the old collection to the new collection
-		if (!AddToCollection(NewCollectionName, NewShareType, ObjectPaths))
+		// Find and set the new parent GUID
+		NewParentCollectionKey = FCollectionNameType(ParentCollectionName, ParentShareType);
+		TSharedRef<FCollection>* const ParentCollectionRefPtr = CachedCollections.Find(NewParentCollectionKey.GetValue());
+		if (!ParentCollectionRefPtr)
+		{
+			// The collection doesn't exist
+			LastError = LOCTEXT("Error_DoesntExist", "The collection doesn't exist.");
+			return false;
+		}
+
+		// Does the parent collection need saving in order to have a stable GUID?
+		if ((*ParentCollectionRefPtr)->GetCollectionVersion() < ECollectionVersion::AddedCollectionGuid)
+		{
+			// Try and re-save the parent collection now
+			if (!(*ParentCollectionRefPtr)->Save(LastError))
+			{
+				return false;
+			}
+		}
+
+		if (!IsValidParentCollection(CollectionName, ShareType, ParentCollectionName, ParentShareType))
 		{
-			// Failed to add paths to the new collection. Destroy the collection we created.
-			DestroyCollection(NewCollectionName, NewShareType);
+			// IsValidParentCollection fills in LastError itself
 			return false;
 		}
+
+		NewParentGuid = (*ParentCollectionRefPtr)->GetCollectionGuid();
+	}
+
+	// Anything changed?
+	if (OldParentGuid == NewParentGuid)
+	{
+		return true;
 	}
 
-	// Delete the old collection
-	if (!DestroyCollection(CurrentCollectionName, CurrentShareType))
+	(*CollectionRefPtr)->SetParentCollectionGuid(NewParentGuid);
+
+	// Try and save with the new parent GUID
+	if (!(*CollectionRefPtr)->Save(LastError))
 	{
-		// Failed to destroy the old collection. Destroy the collection we created.
-		DestroyCollection(NewCollectionName, NewShareType);
+		// Failed to save... rollback the collection to use its old parent GUID
+		(*CollectionRefPtr)->SetParentCollectionGuid(OldParentGuid);
 		return false;
 	}
 
+	// Find the old parent so we can notify about the change
+	{
+		const FCollectionNameType* const OldParentCollectionKeyPtr = CachedCollectionNamesFromGuids.Find(OldParentGuid);
+		if (OldParentCollectionKeyPtr)
+		{
+			OldParentCollectionKey = *OldParentCollectionKeyPtr;
+		}
+	}
+
+	RebuildCachedHierarchy();
 	RebuildCachedObjects();
 
 	// Success
-	const FCollectionNameType OriginalCollectionKey(CurrentCollectionName, CurrentShareType);
-	const FCollectionNameType NewCollectionKey(NewCollectionName, NewShareType);
-	CollectionRenamedEvent.Broadcast(OriginalCollectionKey, NewCollectionKey);
+	CollectionReparentedEvent.Broadcast(CollectionKey, OldParentCollectionKey, NewParentCollectionKey);
 	return true;
 }
 
@@ -348,6 +592,7 @@ bool FCollectionManager::DestroyCollection(FName CollectionName, ECollectionShar
 	{
 		RemoveCollection(*CollectionRefPtr, ShareType);
 
+		RebuildCachedHierarchy();
 		RebuildCachedObjects();
 
 		CollectionDestroyedEvent.Broadcast(CollectionKey);
@@ -512,7 +757,7 @@ bool FCollectionManager::EmptyCollection(FName CollectionName, ECollectionShareT
 	}
 
 	TArray<FName> ObjectPaths;
-	if (!GetAssetsInCollection(CollectionName, ShareType, ObjectPaths))
+	if (!GetObjectsInCollection(CollectionName, ShareType, ObjectPaths))
 	{
 		// Failed to load collection
 		return false;
@@ -552,7 +797,7 @@ bool FCollectionManager::IsCollectionEmpty(FName CollectionName, ECollectionShar
 	return true;
 }
 
-bool FCollectionManager::IsObjectInCollection(FName ObjectPath, FName CollectionName, ECollectionShareType::Type ShareType) const
+bool FCollectionManager::IsObjectInCollection(FName ObjectPath, FName CollectionName, ECollectionShareType::Type ShareType, ECollectionRecursionFlags::Flags RecursionMode) const
 {
 	if (!ensure(ShareType < ECollectionShareType::CST_All))
 	{
@@ -561,14 +806,57 @@ bool FCollectionManager::IsObjectInCollection(FName ObjectPath, FName Collection
 		return true;
 	}
 
-	const FCollectionNameType CollectionKey(CollectionName, ShareType);
-	const TSharedRef<FCollection>* const CollectionRefPtr = CachedCollections.Find(CollectionKey);
-	if (CollectionRefPtr)
+	bool bFoundObject = false;
+
+	auto IsObjectInCollectionWorker = [&](const FCollectionNameType& InCollectionKey, ECollectionRecursionFlags::Flag InReason) -> ERecursiveWorkerFlowControl
+	{
+		const TSharedRef<FCollection>* const CollectionRefPtr = CachedCollections.Find(InCollectionKey);
+		if (CollectionRefPtr)
+		{
+			bFoundObject = (*CollectionRefPtr)->IsObjectInCollection(ObjectPath);
+		}
+		return (bFoundObject) ? ERecursiveWorkerFlowControl::Stop : ERecursiveWorkerFlowControl::Continue;
+	};
+
+	RecursionHelper_DoWork(FCollectionNameType(CollectionName, ShareType), RecursionMode, IsObjectInCollectionWorker);
+
+	return bFoundObject;
+}
+
+bool FCollectionManager::IsValidParentCollection(FName CollectionName, ECollectionShareType::Type ShareType, FName ParentCollectionName, ECollectionShareType::Type ParentShareType) const
+{
+	if (!ensure(ShareType < ECollectionShareType::CST_All) || (!ParentCollectionName.IsNone() && !ensure(ParentShareType < ECollectionShareType::CST_All)))
+	{
+		// Bad share type
+		LastError = LOCTEXT("Error_Internal", "There was an internal error.");
+		return true;
+	}
+
+	if (ParentCollectionName.IsNone())
 	{
-		return (*CollectionRefPtr)->IsObjectInCollection(ObjectPath);
+		// Clearing the parent is always valid
+		return true;
 	}
 
-	return false;
+	bool bValidParent = true;
+
+	auto IsValidParentCollectionWorker = [&](const FCollectionNameType& InCollectionKey, ECollectionRecursionFlags::Flag InReason) -> ERecursiveWorkerFlowControl
+	{
+		const bool bMatchesCollectionBeingReparented = (CollectionName == InCollectionKey.Name && ShareType == InCollectionKey.Type);
+		if (bMatchesCollectionBeingReparented)
+		{
+			bValidParent = false;
+			LastError = (InReason == ECollectionRecursionFlags::Self) 
+				? LOCTEXT("InvalidParent_CannotParentToSelf", "A collection cannot be parented to itself") 
+				: LOCTEXT("InvalidParent_CannotParentToChildren", "A collection cannot be parented to its children");
+			return ERecursiveWorkerFlowControl::Stop;
+		}
+		return ERecursiveWorkerFlowControl::Continue;
+	};
+
+	RecursionHelper_DoWork(FCollectionNameType(ParentCollectionName, ParentShareType), ECollectionRecursionFlags::SelfAndParents, IsValidParentCollectionWorker);
+
+	return bValidParent;
 }
 
 void FCollectionManager::LoadCollections()
@@ -587,22 +875,24 @@ void FCollectionManager::LoadCollections()
 		for (const FString& BaseFilename : Filenames)
 		{
 			const FString Filename = CollectionFolder / BaseFilename;
-
-			TSharedRef<FCollection> NewCollection = MakeShareable(new FCollection());
 			const bool bUseSCC = ShouldUseSCC(ECollectionShareType::Type(CacheIdx));
-			if (NewCollection->LoadFromFile(Filename, bUseSCC))
+
+			FText LoadErrorText;
+			TSharedRef<FCollection> NewCollection = MakeShareable(new FCollection(Filename, bUseSCC));
+			if (NewCollection->Load(LoadErrorText))
 			{
 				AddCollection(NewCollection, ECollectionShareType::Type(CacheIdx));
 			}
 			else
 			{
-				UE_LOG(LogCollectionManager, Warning, TEXT("Failed to load collection file %s"), *Filename);
+				UE_LOG(LogCollectionManager, Warning, TEXT("%s"), *LoadErrorText.ToString());
 			}
 		}
 	}
 
 	UE_LOG(LogCollectionManager, Log, TEXT( "Loaded %d collections in %0.6f seconds" ), CachedCollections.Num() - PrevNumCollections, FPlatformTime::Seconds() - LoadStartTime);
 
+	RebuildCachedHierarchy();
 	RebuildCachedObjects();
 }
 
@@ -618,16 +908,74 @@ void FCollectionManager::RebuildCachedObjects()
 		const TSharedRef<FCollection>& Collection = CachedCollection.Value;
 
 		TArray<FName> ObjectsInCollection;
-		Collection->GetObjectsInCollection(ObjectsInCollection);
+		CachedCollection.Value->GetObjectsInCollection(ObjectsInCollection);
+
+		if (ObjectsInCollection.Num() > 0)
+		{
+			auto RebuildCachedObjectsWorker = [&](const FCollectionNameType& InCollectionKey, ECollectionRecursionFlags::Flag InReason) -> ERecursiveWorkerFlowControl
+			{
+				// The worker reason will tell us why this collection is being processed (eg, because it is a parent of the collection we told it to DoWork on),
+				// however, the reason this object exists in that parent collection is because a child collection contains it, and this is the reason we need
+				// to put into the FObjectCollectionInfo, since that's what we'll test against later when we do the "do my children contain this object"? test
+				// That's why we flip the reason logic here...
+				ECollectionRecursionFlags::Flag ReasonObjectInCollection = InReason;
+				switch (InReason)
+				{
+				case ECollectionRecursionFlags::Parents:
+					ReasonObjectInCollection = ECollectionRecursionFlags::Children;
+					break;
+				case ECollectionRecursionFlags::Children:
+					ReasonObjectInCollection = ECollectionRecursionFlags::Parents;
+					break;
+				default:
+					break;
+				}
+
+				for (const FName& ObjectPath : ObjectsInCollection)
+				{
+					auto& ObjectCollectionInfos = CachedObjects.FindOrAdd(ObjectPath);
+					FObjectCollectionInfo* ObjectInfoPtr = ObjectCollectionInfos.FindByPredicate([&](const FObjectCollectionInfo& InCollectionInfo) { return InCollectionInfo.CollectionKey == InCollectionKey; });
+					if (ObjectInfoPtr)
+					{
+						ObjectInfoPtr->Reason |= ReasonObjectInCollection;
+					}
+					else
+					{
+						ObjectCollectionInfos.Add(FObjectCollectionInfo(InCollectionKey, ReasonObjectInCollection));
+					}
+				}
+				return ERecursiveWorkerFlowControl::Continue;
+			};
+
+			// Recursively process all collections so that they know they contain these objects (and why!)
+			RecursionHelper_DoWork(CollectionKey, ECollectionRecursionFlags::All, RebuildCachedObjectsWorker);
+		}
+	}
+
+	UE_LOG(LogCollectionManager, Log, TEXT( "Rebuilt the object cache for %d collections in %0.6f seconds (found %d objects)" ), CachedCollections.Num(), FPlatformTime::Seconds() - LoadStartTime, CachedObjects.Num());
+}
+
+void FCollectionManager::RebuildCachedHierarchy()
+{
+	const double LoadStartTime = FPlatformTime::Seconds();
+
+	CachedHierarchy.Empty();
+
+	for (const auto& CachedCollection : CachedCollections)
+	{
+		const FCollectionNameType& CollectionKey = CachedCollection.Key;
+		const TSharedRef<FCollection>& Collection = CachedCollection.Value;
 
-		for (const FName& ObjectPath : ObjectsInCollection)
+		// Make sure this is a known parent GUID before adding it to the map
+		const FGuid& ParentCollectionGuid = Collection->GetParentCollectionGuid();
+		if (CachedCollectionNamesFromGuids.Contains(ParentCollectionGuid))
 		{
-			auto& ObjectCollections = CachedObjects.FindOrAdd(ObjectPath);
-			ObjectCollections.AddUnique(CollectionKey);
+			auto& CollectionChildren = CachedHierarchy.FindOrAdd(ParentCollectionGuid);
+			CollectionChildren.AddUnique(Collection->GetCollectionGuid());
 		}
 	}
 
-	UE_LOG(LogCollectionManager, Log, TEXT( "Rebuilt the cache for %d objects in %0.6f seconds" ), CachedObjects.Num(), FPlatformTime::Seconds() - LoadStartTime);
+	UE_LOG(LogCollectionManager, Log, TEXT( "Rebuilt the hierarchy cache for %d collections in %0.6f seconds" ), CachedCollections.Num(), FPlatformTime::Seconds() - LoadStartTime);
 }
 
 bool FCollectionManager::ShouldUseSCC(ECollectionShareType::Type ShareType) const
@@ -635,6 +983,13 @@ bool FCollectionManager::ShouldUseSCC(ECollectionShareType::Type ShareType) cons
 	return ShareType != ECollectionShareType::CST_Local && ShareType != ECollectionShareType::CST_System;
 }
 
+FString FCollectionManager::GetCollectionFilename(const FName& InCollectionName, const ECollectionShareType::Type InCollectionShareType) const
+{
+	FString CollectionFilename = CollectionFolders[InCollectionShareType] / InCollectionName.ToString() + TEXT(".") + CollectionExtension;
+	FPaths::NormalizeFilename(CollectionFilename);
+	return CollectionFilename;
+}
+
 bool FCollectionManager::AddCollection(const TSharedRef<FCollection>& CollectionRef, ECollectionShareType::Type ShareType)
 {
 	if (!ensure(ShareType < ECollectionShareType::CST_All))
@@ -651,6 +1006,7 @@ bool FCollectionManager::AddCollection(const TSharedRef<FCollection>& Collection
 	}
 
 	CachedCollections.Add(CollectionKey, CollectionRef);
+	CachedCollectionNamesFromGuids.Add(CollectionRef->GetCollectionGuid(), CollectionKey);
 	return true;
 }
 
@@ -663,7 +1019,69 @@ bool FCollectionManager::RemoveCollection(const TSharedRef<FCollection>& Collect
 	}
 
 	const FCollectionNameType CollectionKey(CollectionRef->GetCollectionName(), ShareType);
+	CachedCollectionNamesFromGuids.Remove(CollectionRef->GetCollectionGuid());
 	return CachedCollections.Remove(CollectionKey) > 0;
 }
 
+void FCollectionManager::RecursionHelper_DoWork(const FCollectionNameType& InCollectionKey, const ECollectionRecursionFlags::Flags InRecursionMode, const FRecursiveWorkerFunc& InWorkerFunc) const
+{
+	if ((InRecursionMode & ECollectionRecursionFlags::Self) && InWorkerFunc(InCollectionKey, ECollectionRecursionFlags::Self) == ERecursiveWorkerFlowControl::Stop)
+	{
+		return;
+	}
+
+	if ((InRecursionMode & ECollectionRecursionFlags::Parents) && RecursionHelper_DoWorkOnParents(InCollectionKey, InWorkerFunc) == ERecursiveWorkerFlowControl::Stop)
+	{
+		return;
+	}
+
+	if ((InRecursionMode & ECollectionRecursionFlags::Children) && RecursionHelper_DoWorkOnChildren(InCollectionKey, InWorkerFunc) == ERecursiveWorkerFlowControl::Stop)
+	{
+		return;
+	}
+}
+
+FCollectionManager::ERecursiveWorkerFlowControl FCollectionManager::RecursionHelper_DoWorkOnParents(const FCollectionNameType& InCollectionKey, const FRecursiveWorkerFunc& InWorkerFunc) const
+{
+	const TSharedRef<FCollection>* const CollectionRefPtr = CachedCollections.Find(InCollectionKey);
+	if (CollectionRefPtr)
+	{
+		const FCollectionNameType* const ParentCollectionKeyPtr = CachedCollectionNamesFromGuids.Find((*CollectionRefPtr)->GetParentCollectionGuid());
+		if (ParentCollectionKeyPtr)
+		{
+			if (InWorkerFunc(*ParentCollectionKeyPtr, ECollectionRecursionFlags::Parents) == ERecursiveWorkerFlowControl::Stop || RecursionHelper_DoWorkOnParents(*ParentCollectionKeyPtr, InWorkerFunc) == ERecursiveWorkerFlowControl::Stop)
+			{
+				return ERecursiveWorkerFlowControl::Stop;
+			}
+		}
+	}
+
+	return ERecursiveWorkerFlowControl::Continue;
+}
+
+FCollectionManager::ERecursiveWorkerFlowControl FCollectionManager::RecursionHelper_DoWorkOnChildren(const FCollectionNameType& InCollectionKey, const FRecursiveWorkerFunc& InWorkerFunc) const
+{
+	const TSharedRef<FCollection>* const CollectionRefPtr = CachedCollections.Find(InCollectionKey);
+	if (CollectionRefPtr)
+	{
+		const TArray<FGuid>* const ChildCollectionGuids = CachedHierarchy.Find((*CollectionRefPtr)->GetCollectionGuid());
+		if (ChildCollectionGuids)
+		{
+			for (const FGuid& ChildCollectionGuid : *ChildCollectionGuids)
+			{
+				const FCollectionNameType* const ChildCollectionKeyPtr = CachedCollectionNamesFromGuids.Find(ChildCollectionGuid);
+				if (ChildCollectionKeyPtr)
+				{
+					if (InWorkerFunc(*ChildCollectionKeyPtr, ECollectionRecursionFlags::Children) == ERecursiveWorkerFlowControl::Stop || RecursionHelper_DoWorkOnChildren(*ChildCollectionKeyPtr, InWorkerFunc) == ERecursiveWorkerFlowControl::Stop)
+					{
+						return ERecursiveWorkerFlowControl::Stop;
+					}
+				}
+			}
+		}
+	}
+
+	return ERecursiveWorkerFlowControl::Continue;
+}
+
 #undef LOCTEXT_NAMESPACE
diff --git a/Engine/Source/Developer/CollectionManager/Private/CollectionManager.h b/Engine/Source/Developer/CollectionManager/Private/CollectionManager.h
--- a/Engine/Source/Developer/CollectionManager/Private/CollectionManager.h
+++ b/Engine/Source/Developer/CollectionManager/Private/CollectionManager.h
@@ -10,18 +10,24 @@ public:
 
 	// ICollectionManager implementation
 	virtual bool HasCollections() const override;
+	virtual void GetCollections(TArray<FCollectionNameType>& OutCollections) const override;
 	virtual void GetCollectionNames(ECollectionShareType::Type ShareType, TArray<FName>& CollectionNames) const override;
+	virtual void GetRootCollections(TArray<FCollectionNameType>& OutCollections) const override;
+	virtual void GetRootCollectionNames(ECollectionShareType::Type ShareType, TArray<FName>& CollectionNames) const override;
+	virtual void GetChildCollections(FName CollectionName, ECollectionShareType::Type ShareType, TArray<FCollectionNameType>& OutCollections) const override;
+	virtual void GetChildCollectionNames(FName CollectionName, ECollectionShareType::Type ShareType, ECollectionShareType::Type ChildShareType, TArray<FName>& CollectionNames) const override;
 	virtual bool CollectionExists(FName CollectionName, ECollectionShareType::Type ShareType) const override;
-	virtual bool GetAssetsInCollection(FName CollectionName, ECollectionShareType::Type ShareType, TArray<FName>& AssetPaths) const override;
-	virtual bool GetObjectsInCollection(FName CollectionName, ECollectionShareType::Type ShareType, TArray<FName>& ObjectPaths) const override;
-	virtual bool GetClassesInCollection(FName CollectionName, ECollectionShareType::Type ShareType, TArray<FName>& ClassPaths) const override;
-	virtual void GetCollectionsContainingObject(FName ObjectPath, ECollectionShareType::Type ShareType, TArray<FName>& OutCollectionNames) const override;
-	virtual void GetCollectionsContainingObject(FName ObjectPath, TArray<FCollectionNameType>& OutCollections) const override;
-	virtual void GetCollectionsContainingObjects(const TArray<FName>& ObjectPaths, TMap<FCollectionNameType, TArray<FName>>& OutCollectionsAndMatchedObjects) const override;
-	virtual FString GetCollectionsStringForObject(FName ObjectPath, ECollectionShareType::Type ShareType) const override;
+	virtual bool GetAssetsInCollection(FName CollectionName, ECollectionShareType::Type ShareType, TArray<FName>& AssetPaths, ECollectionRecursionFlags::Flags RecursionMode = ECollectionRecursionFlags::Self) const override;
+	virtual bool GetObjectsInCollection(FName CollectionName, ECollectionShareType::Type ShareType, TArray<FName>& ObjectPaths, ECollectionRecursionFlags::Flags RecursionMode = ECollectionRecursionFlags::Self) const override;
+	virtual bool GetClassesInCollection(FName CollectionName, ECollectionShareType::Type ShareType, TArray<FName>& ClassPaths, ECollectionRecursionFlags::Flags RecursionMode = ECollectionRecursionFlags::Self) const override;
+	virtual void GetCollectionsContainingObject(FName ObjectPath, ECollectionShareType::Type ShareType, TArray<FName>& OutCollectionNames, ECollectionRecursionFlags::Flags RecursionMode = ECollectionRecursionFlags::Self) const override;
+	virtual void GetCollectionsContainingObject(FName ObjectPath, TArray<FCollectionNameType>& OutCollections, ECollectionRecursionFlags::Flags RecursionMode = ECollectionRecursionFlags::Self) const override;
+	virtual void GetCollectionsContainingObjects(const TArray<FName>& ObjectPaths, TMap<FCollectionNameType, TArray<FName>>& OutCollectionsAndMatchedObjects, ECollectionRecursionFlags::Flags RecursionMode = ECollectionRecursionFlags::Self) const override;
+	virtual FString GetCollectionsStringForObject(FName ObjectPath, ECollectionShareType::Type ShareType, ECollectionRecursionFlags::Flags RecursionMode = ECollectionRecursionFlags::Self) const override;
 	virtual void CreateUniqueCollectionName(const FName& BaseName, ECollectionShareType::Type ShareType, FName& OutCollectionName) const override;
 	virtual bool CreateCollection(FName CollectionName, ECollectionShareType::Type ShareType) override;
 	virtual bool RenameCollection(FName CurrentCollectionName, ECollectionShareType::Type CurrentShareType, FName NewCollectionName, ECollectionShareType::Type NewShareType) override;
+	virtual bool ReparentCollection(FName CollectionName, ECollectionShareType::Type ShareType, FName ParentCollectionName, ECollectionShareType::Type ParentShareType) override;
 	virtual bool DestroyCollection(FName CollectionName, ECollectionShareType::Type ShareType) override;
 	virtual bool AddToCollection(FName CollectionName, ECollectionShareType::Type ShareType, FName ObjectPath) override;
 	virtual bool AddToCollection(FName CollectionName, ECollectionShareType::Type ShareType, const TArray<FName>& ObjectPaths, int32* OutNumAdded = nullptr) override;
@@ -29,7 +35,8 @@ public:
 	virtual bool RemoveFromCollection(FName CollectionName, ECollectionShareType::Type ShareType, const TArray<FName>& ObjectPaths, int32* OutNumRemoved = nullptr) override;
 	virtual bool EmptyCollection(FName CollectionName, ECollectionShareType::Type ShareType) override;
 	virtual bool IsCollectionEmpty(FName CollectionName, ECollectionShareType::Type ShareType) const override;
-	virtual bool IsObjectInCollection(FName ObjectPath, FName CollectionName, ECollectionShareType::Type ShareType) const override;
+	virtual bool IsObjectInCollection(FName ObjectPath, FName CollectionName, ECollectionShareType::Type ShareType, ECollectionRecursionFlags::Flags RecursionMode = ECollectionRecursionFlags::Self) const override;
+	virtual bool IsValidParentCollection(FName CollectionName, ECollectionShareType::Type ShareType, FName ParentCollectionName, ECollectionShareType::Type ParentShareType) const override;
 	virtual FText GetLastError() const override { return LastError; }
 
 	/** Event for when collections are created */
@@ -52,6 +59,10 @@ public:
 	DECLARE_DERIVED_EVENT( FCollectionManager, ICollectionManager::FCollectionRenamedEvent, FCollectionRenamedEvent );
 	virtual FCollectionRenamedEvent& OnCollectionRenamed() override { return CollectionRenamedEvent; }
 
+	/** Event for when collections are re-parented */
+	DECLARE_DERIVED_EVENT( FCollectionManager, ICollectionManager::FCollectionReparentedEvent, FCollectionReparentedEvent );
+	virtual FCollectionReparentedEvent& OnCollectionReparented() override { return CollectionReparentedEvent; }
+
 private:
 	/** Loads all collection files from disk */
 	void LoadCollections();
@@ -59,16 +70,55 @@ private:
 	/** Rebuild the entire cached objects map based on the current collection data */
 	void RebuildCachedObjects();
 
+	/** Rebuild the entire cached hierarchy map based on the current collection data */
+	void RebuildCachedHierarchy();
+
 	/** Returns true if the specified share type requires source control */
 	bool ShouldUseSCC(ECollectionShareType::Type ShareType) const;
 
+	/** Given a collection name and share type, work out the full filename for the collection to use on disk */
+	FString GetCollectionFilename(const FName& InCollectionName, const ECollectionShareType::Type InCollectionShareType) const;
+
 	/** Adds a collection to the lookup maps */
 	bool AddCollection(const TSharedRef<FCollection>& CollectionRef, ECollectionShareType::Type ShareType);
 
 	/** Removes a collection from the lookup maps */
 	bool RemoveCollection(const TSharedRef<FCollection>& CollectionRef, ECollectionShareType::Type ShareType);
 
+	enum class ERecursiveWorkerFlowControl : uint8
+	{
+		Stop,
+		Continue,
+	};
+
+	typedef TFunctionRef<ERecursiveWorkerFlowControl(const FCollectionNameType&, ECollectionRecursionFlags::Flag)> FRecursiveWorkerFunc;
+
+	void RecursionHelper_DoWork(const FCollectionNameType& InCollectionKey, const ECollectionRecursionFlags::Flags InRecursionMode, const FRecursiveWorkerFunc& InWorkerFunc) const;
+	ERecursiveWorkerFlowControl RecursionHelper_DoWorkOnParents(const FCollectionNameType& InCollectionKey, const FRecursiveWorkerFunc& InWorkerFunc) const;
+	ERecursiveWorkerFlowControl RecursionHelper_DoWorkOnChildren(const FCollectionNameType& InCollectionKey, const FRecursiveWorkerFunc& InWorkerFunc) const;
+
 private:
+	/** Collection info for a given object - gives the collection name, as well as the reason this object is considered to be part of this collection */
+	struct FObjectCollectionInfo
+	{
+		explicit FObjectCollectionInfo(const FCollectionNameType& InCollectionKey)
+			: CollectionKey(InCollectionKey)
+			, Reason(0)
+		{
+		}
+
+		FObjectCollectionInfo(const FCollectionNameType& InCollectionKey, const ECollectionRecursionFlags::Flags InReason)
+			: CollectionKey(InCollectionKey)
+			, Reason(InReason)
+		{
+		}
+
+		/** The key identifying the collection that contains this object */
+		FCollectionNameType CollectionKey;
+		/** The reason(s) why this collection contains this object - this can be tested against the recursion mode when getting the collections for an object */
+		ECollectionRecursionFlags::Flags Reason;
+	};
+
 	/** The folders that contain collections */
 	FString CollectionFolders[ECollectionShareType::CST_All];
 
@@ -78,8 +128,14 @@ private:
 	/** A map of collection names to FCollection objects */
 	TMap<FCollectionNameType, TSharedRef<FCollection>> CachedCollections;
 
-	/** A map of object paths to their associated collection names */
-	TMap<FName, TArray<FCollectionNameType>> CachedObjects;
+	/** A map of collection GUIDs to their associated collection names */
+	TMap<FGuid, FCollectionNameType> CachedCollectionNamesFromGuids;
+
+	/** A map of object paths to their associated collection info - only objects that are in collections will appear in here */
+	TMap<FName, TArray<FObjectCollectionInfo>> CachedObjects;
+
+	/** A map of parent collection GUIDs to their child collection GUIDs - only collections that have children will appear in here */
+	TMap<FGuid, TArray<FGuid>> CachedHierarchy;
 
 	/** The most recent error that occurred */
 	mutable FText LastError;
@@ -93,6 +149,9 @@ private:
 	/** Event for when collections are renamed */
 	FCollectionRenamedEvent CollectionRenamedEvent;
 
+	/** Event for when collections are re-parented */
+	FCollectionReparentedEvent CollectionReparentedEvent;
+
 	/** Event for when collections are created */
 	FCollectionCreatedEvent CollectionCreatedEvent;
 
diff --git a/Engine/Source/Developer/CollectionManager/Public/CollectionManagerTypes.h b/Engine/Source/Developer/CollectionManager/Public/CollectionManagerTypes.h
--- a/Engine/Source/Developer/CollectionManager/Public/CollectionManagerTypes.h
+++ b/Engine/Source/Developer/CollectionManager/Public/CollectionManagerTypes.h
@@ -80,16 +80,16 @@ struct ECollectionShareType
 		return FText::GetEmpty();
 	}
 
-	static FName GetIconStyleName(const Type InType)
+	static FName GetIconStyleName(const Type InType, const TCHAR* InSizeSuffix = TEXT(".Small"))
 	{
 		switch (InType)
 		{
 		case CST_Local:
-			return FName("ContentBrowser.Local");
+			return FName(*FString::Printf(TEXT("ContentBrowser.Local%s"), InSizeSuffix));
 		case CST_Private:
-			return FName("ContentBrowser.Private");
+			return FName(*FString::Printf(TEXT("ContentBrowser.Private%s"), InSizeSuffix));
 		case CST_Shared:
-			return FName("ContentBrowser.Shared");
+			return FName(*FString::Printf(TEXT("ContentBrowser.Shared%s"), InSizeSuffix));
 		default:
 			break;
 		}
@@ -97,6 +97,28 @@ struct ECollectionShareType
 	}
 };
 
+/** Controls how the collections manager will recurse when performing work against a given collection */
+struct ECollectionRecursionFlags
+{
+	typedef uint8 Flags;
+	enum Flag
+	{
+		/** Include the current collection when performing work */
+		Self = 1<<0,
+		/** Include the parent collections when performing work */
+		Parents = 1<<1,
+		/** Include the child collections when performing work */
+		Children = 1<<2,
+
+		/** Include parent collections in addition to the current collection */
+		SelfAndParents = Self | Parents,
+		/** Include child collections in addition to the current collection */
+		SelfAndChildren = Self | Children,
+		/** Include parent and child collections in addition to the current collection */
+		All = Self | Parents | Children,
+	};
+};
+
 /** A name/type pair to uniquely identify a collection */
 struct FCollectionNameType
 {
diff --git a/Engine/Source/Developer/CollectionManager/Public/ICollectionManager.h b/Engine/Source/Developer/CollectionManager/Public/ICollectionManager.h
--- a/Engine/Source/Developer/CollectionManager/Public/ICollectionManager.h
+++ b/Engine/Source/Developer/CollectionManager/Public/ICollectionManager.h
@@ -13,32 +13,47 @@ public:
 	/** Returns whether or not the collection manager contains any collections */
 	virtual bool HasCollections() const = 0;
 
+	/** Returns the list of collections */
+	virtual void GetCollections(TArray<FCollectionNameType>& OutCollections) const = 0;
+
 	/** Returns the list of collection names of the specified share type */
 	virtual void GetCollectionNames(ECollectionShareType::Type ShareType, TArray<FName>& CollectionNames) const = 0;
 
+	/** Returns the list of root-level collections */
+	virtual void GetRootCollections(TArray<FCollectionNameType>& OutCollections) const = 0;
+
+	/** Returns the list of root-level collection names of the specified share type */
+	virtual void GetRootCollectionNames(ECollectionShareType::Type ShareType, TArray<FName>& CollectionNames) const = 0;
+
+	/** Returns the list of child collections of the given collection */
+	virtual void GetChildCollections(FName CollectionName, ECollectionShareType::Type ShareType, TArray<FCollectionNameType>& OutCollections) const = 0;
+
+	/** Returns the list of child collections of the given collection that are also of the specified share type */
+	virtual void GetChildCollectionNames(FName CollectionName, ECollectionShareType::Type ShareType, ECollectionShareType::Type ChildShareType, TArray<FName>& CollectionNames) const = 0;
+
 	/** Returns true if the collection exists */
 	virtual bool CollectionExists(FName CollectionName, ECollectionShareType::Type ShareType) const = 0;
 
 	/** Returns a list of asset paths found in the specified collection and share type */
-	virtual bool GetAssetsInCollection(FName CollectionName, ECollectionShareType::Type ShareType, TArray<FName>& OutAssetPaths) const = 0;
+	virtual bool GetAssetsInCollection(FName CollectionName, ECollectionShareType::Type ShareType, TArray<FName>& OutAssetPaths, ECollectionRecursionFlags::Flags RecursionMode = ECollectionRecursionFlags::Self) const = 0;
 
 	/** Returns a list of class paths found in the specified collection and share type */
-	virtual bool GetClassesInCollection(FName CollectionName, ECollectionShareType::Type ShareType, TArray<FName>& ClassPaths) const = 0;
+	virtual bool GetClassesInCollection(FName CollectionName, ECollectionShareType::Type ShareType, TArray<FName>& ClassPaths, ECollectionRecursionFlags::Flags RecursionMode = ECollectionRecursionFlags::Self) const = 0;
 
 	/** Returns a list of object paths found in the specified collection and share type */
-	virtual bool GetObjectsInCollection(FName CollectionName, ECollectionShareType::Type ShareType, TArray<FName>& ObjectPaths) const = 0;
+	virtual bool GetObjectsInCollection(FName CollectionName, ECollectionShareType::Type ShareType, TArray<FName>& ObjectPaths, ECollectionRecursionFlags::Flags RecursionMode = ECollectionRecursionFlags::Self) const = 0;
 
 	/** Returns a list of collections in which the specified object exists of the specified share type */
-	virtual void GetCollectionsContainingObject(FName ObjectPath, ECollectionShareType::Type ShareType, TArray<FName>& OutCollectionNames) const = 0;
+	virtual void GetCollectionsContainingObject(FName ObjectPath, ECollectionShareType::Type ShareType, TArray<FName>& OutCollectionNames, ECollectionRecursionFlags::Flags RecursionMode = ECollectionRecursionFlags::Self) const = 0;
 
 	/** Returns a list of collections in which the specified object exists */
-	virtual void GetCollectionsContainingObject(FName ObjectPath, TArray<FCollectionNameType>& OutCollections) const = 0;
+	virtual void GetCollectionsContainingObject(FName ObjectPath, TArray<FCollectionNameType>& OutCollections, ECollectionRecursionFlags::Flags RecursionMode = ECollectionRecursionFlags::Self) const = 0;
 
 	/** Returns a list of collections in which any of the specified objects exist */
-	virtual void GetCollectionsContainingObjects(const TArray<FName>& ObjectPaths, TMap<FCollectionNameType, TArray<FName>>& OutCollectionsAndMatchedObjects) const = 0;
+	virtual void GetCollectionsContainingObjects(const TArray<FName>& ObjectPaths, TMap<FCollectionNameType, TArray<FName>>& OutCollectionsAndMatchedObjects, ECollectionRecursionFlags::Flags RecursionMode = ECollectionRecursionFlags::Self) const = 0;
 
 	/** Returns a string containing a comma separated list of collections in which the specified object exists of the specified share type */
-	virtual FString GetCollectionsStringForObject(FName ObjectPath, ECollectionShareType::Type ShareType) const = 0;
+	virtual FString GetCollectionsStringForObject(FName ObjectPath, ECollectionShareType::Type ShareType, ECollectionRecursionFlags::Flags RecursionMode = ECollectionRecursionFlags::Self) const = 0;
 
 	/** Creates a unique collection name for the given type taking the form BaseName+(unique number) */
 	virtual void CreateUniqueCollectionName(const FName& BaseName, ECollectionShareType::Type ShareType, FName& OutCollectionName) const = 0;
@@ -53,16 +68,27 @@ public:
 	virtual bool CreateCollection(FName CollectionName, ECollectionShareType::Type ShareType) = 0;
 
 	/**
-	  * Renames a collection to the asset registry. A .collection file will be added to disk and a .collection file will be removed.
+	  * Renames a collection. A .collection file will be added to disk and a .collection file will be removed.
 	  *
 	  * @param CurrentCollectionName The current name of the collection.
 	  * @param CurrentShareType The current way the collection is shared.
 	  * @param NewCollectionName The new name of the collection.
 	  * @param NewShareType The new way the collection is shared.
-	  * @return true if the add was successful. If false, GetLastError will return a human readable string description of the error.
+	  * @return true if the rename was successful. If false, GetLastError will return a human readable string description of the error.
 	  */
 	virtual bool RenameCollection(FName CurrentCollectionName, ECollectionShareType::Type CurrentShareType, FName NewCollectionName, ECollectionShareType::Type NewShareType) = 0;
 
+	/**
+	  * Re-parents a collection. The parent collection may be re-saved if it's too old to have a stable GUID.
+	  *
+	  * @param CollectionName The name of the collection to set the parent of.
+	  * @param ShareType The way the collection is shared.
+	  * @param ParentCollectionName The new parent of the collection, or None to have the collection become a root collection.
+	  * @param ParentShareType The way the new parent collection is shared.
+	  * @return true if the re-parent was successful. If false, GetLastError will return a human readable string description of the error.
+	  */
+	virtual bool ReparentCollection(FName CollectionName, ECollectionShareType::Type ShareType, FName ParentCollectionName, ECollectionShareType::Type ParentShareType) = 0;
+
 	/**
 	  * Removes a collection to the asset registry. A .collection file will be deleted from disk.
 	  *
@@ -120,7 +146,18 @@ public:
 	 * @param ShareType The way the collection is shared.
 	 * @return true if the object is in the collection.
 	 */
-	virtual bool IsObjectInCollection(FName ObjectPath, FName CollectionName, ECollectionShareType::Type ShareType) const = 0;
+	virtual bool IsObjectInCollection(FName ObjectPath, FName CollectionName, ECollectionShareType::Type ShareType, ECollectionRecursionFlags::Flags RecursionMode = ECollectionRecursionFlags::Self) const = 0;
+
+	/**
+	  * Check to see if the given collection is valid to be used as the parent of another collection. A collection may not be parented to itself, nor any of its current children.
+	  *
+	  * @param CollectionName The name of the collection to check the parent of.
+	  * @param ShareType The way the collection is shared.
+	  * @param ParentCollectionName The name of the parent collection to test against.
+	  * @param ParentShareType The way the new parent collection is shared.
+	  * @return true if the parent is considered valid for this collection. If false, GetLastError will return a human readable string describing why the parent is invalid.
+	  */
+	virtual bool IsValidParentCollection(FName CollectionName, ECollectionShareType::Type ShareType, FName ParentCollectionName, ECollectionShareType::Type ParentShareType) const = 0;
 
 	/** Returns the most recent error. */
 	virtual FText GetLastError() const = 0;
@@ -144,4 +181,8 @@ public:
 	/** Event for when collections are renamed */
 	DECLARE_EVENT_TwoParams( ICollectionManager, FCollectionRenamedEvent, const FCollectionNameType&, const FCollectionNameType& );
 	virtual FCollectionRenamedEvent& OnCollectionRenamed() = 0;
+
+	/** Event for when collections are re-parented (params: Collection, OldParent, NewParent) */
+	DECLARE_EVENT_ThreeParams( ICollectionManager, FCollectionReparentedEvent, const FCollectionNameType&, const TOptional<FCollectionNameType>&, const TOptional<FCollectionNameType>& );
+	virtual FCollectionReparentedEvent& OnCollectionReparented() = 0;
 };
diff --git a/Engine/Source/Editor/ContentBrowser/Private/AssetContextMenu.cpp b/Engine/Source/Editor/ContentBrowser/Private/AssetContextMenu.cpp
--- a/Engine/Source/Editor/ContentBrowser/Private/AssetContextMenu.cpp
+++ b/Engine/Source/Editor/ContentBrowser/Private/AssetContextMenu.cpp
@@ -815,26 +815,15 @@ bool FAssetContextMenu::AddCollectionMenuOptions(FMenuBuilder& MenuBuilder)
 		{
 			FCollectionManagerModule& CollectionManagerModule = FCollectionManagerModule::GetModule();
 
-			// Get collections of all types
 			TArray<FCollectionNameType> AvailableCollections;
-			for (int32 TypeIdx = 0; TypeIdx < ECollectionShareType::CST_All; ++TypeIdx)
-			{
-				const ECollectionShareType::Type CollectionType = ECollectionShareType::Type(TypeIdx);
-
-				// Never display system collections
-				if (CollectionType == ECollectionShareType::CST_System)
-				{
-					continue;
-				}
+			CollectionManagerModule.Get().GetRootCollections(AvailableCollections);
 
-				TArray<FName> CollectionNames;
-				CollectionManagerModule.Get().GetCollectionNames(CollectionType, CollectionNames);
+			CreateManageCollectionsSubMenu(SubMenuBuilder, QuickAssetManagement, MoveTemp(AvailableCollections));
+		}
 
-				for (const FName& CollectionName : CollectionNames)
-				{
-					AvailableCollections.Add(FCollectionNameType(CollectionName, CollectionType));
-				}
-			}
+		static void CreateManageCollectionsSubMenu(FMenuBuilder& SubMenuBuilder, TSharedRef<FCollectionAssetManagement> QuickAssetManagement, TArray<FCollectionNameType> AvailableCollections)
+		{
+			FCollectionManagerModule& CollectionManagerModule = FCollectionManagerModule::GetModule();
 
 			AvailableCollections.Sort([](const FCollectionNameType& One, const FCollectionNameType& Two) -> bool
 			{
@@ -843,18 +832,47 @@ bool FAssetContextMenu::AddCollectionMenuOptions(FMenuBuilder& MenuBuilder)
 
 			for (const FCollectionNameType& AvailableCollection : AvailableCollections)
 			{
-				SubMenuBuilder.AddMenuEntry(
-					FText::FromName(AvailableCollection.Name), 
-					ECollectionShareType::GetDescription(AvailableCollection.Type), 
-					FSlateIcon(FEditorStyle::GetStyleSetName(), ECollectionShareType::GetIconStyleName(AvailableCollection.Type)), 
-					FUIAction(
-						FExecuteAction::CreateStatic(&FManageCollectionsContextMenu::OnCollectionClicked, QuickAssetManagement, AvailableCollection),
-						FCanExecuteAction::CreateStatic(&FManageCollectionsContextMenu::IsCollectionEnabled, QuickAssetManagement, AvailableCollection),
-						FGetActionCheckState::CreateStatic(&FManageCollectionsContextMenu::GetCollectionCheckState, QuickAssetManagement, AvailableCollection)
-						), 
-					NAME_None, 
-					EUserInterfaceActionType::ToggleButton
-					);
+				// Never display system collections
+				if (AvailableCollection.Type == ECollectionShareType::CST_System)
+				{
+					continue;
+				}
+
+				TArray<FCollectionNameType> AvailableChildCollections;
+				CollectionManagerModule.Get().GetChildCollections(AvailableCollection.Name, AvailableCollection.Type, AvailableChildCollections);
+
+				if (AvailableChildCollections.Num() > 0)
+				{
+					SubMenuBuilder.AddSubMenu(
+						FText::FromName(AvailableCollection.Name), 
+						ECollectionShareType::GetDescription(AvailableCollection.Type), 
+						FNewMenuDelegate::CreateStatic(&FManageCollectionsContextMenu::CreateManageCollectionsSubMenu, QuickAssetManagement, AvailableChildCollections),
+						FUIAction(
+							FExecuteAction::CreateStatic(&FManageCollectionsContextMenu::OnCollectionClicked, QuickAssetManagement, AvailableCollection),
+							FCanExecuteAction::CreateStatic(&FManageCollectionsContextMenu::IsCollectionEnabled, QuickAssetManagement, AvailableCollection),
+							FGetActionCheckState::CreateStatic(&FManageCollectionsContextMenu::GetCollectionCheckState, QuickAssetManagement, AvailableCollection)
+							), 
+						NAME_None, 
+						EUserInterfaceActionType::ToggleButton,
+						false,
+						FSlateIcon(FEditorStyle::GetStyleSetName(), ECollectionShareType::GetIconStyleName(AvailableCollection.Type))
+						);
+				}
+				else
+				{
+					SubMenuBuilder.AddMenuEntry(
+						FText::FromName(AvailableCollection.Name), 
+						ECollectionShareType::GetDescription(AvailableCollection.Type), 
+						FSlateIcon(FEditorStyle::GetStyleSetName(), ECollectionShareType::GetIconStyleName(AvailableCollection.Type)), 
+						FUIAction(
+							FExecuteAction::CreateStatic(&FManageCollectionsContextMenu::OnCollectionClicked, QuickAssetManagement, AvailableCollection),
+							FCanExecuteAction::CreateStatic(&FManageCollectionsContextMenu::IsCollectionEnabled, QuickAssetManagement, AvailableCollection),
+							FGetActionCheckState::CreateStatic(&FManageCollectionsContextMenu::GetCollectionCheckState, QuickAssetManagement, AvailableCollection)
+							), 
+						NAME_None, 
+						EUserInterfaceActionType::ToggleButton
+						);
+				}
 			}
 		}
 
diff --git a/Engine/Source/Editor/ContentBrowser/Private/AssetViewWidgets.cpp b/Engine/Source/Editor/ContentBrowser/Private/AssetViewWidgets.cpp
--- a/Engine/Source/Editor/ContentBrowser/Private/AssetViewWidgets.cpp
+++ b/Engine/Source/Editor/ContentBrowser/Private/AssetViewWidgets.cpp
@@ -59,6 +59,150 @@ FReply FAssetViewModeUtils::OnViewModeKeyDown( const TSet< TSharedPtr<FAssetView
 }
 
 
+///////////////////////////////
+// FAssetViewModeUtils
+///////////////////////////////
+
+TSharedRef<SWidget> FAssetViewItemHelper::CreateListItemContents(SAssetListItem* const InListItem, const TSharedRef<SWidget>& InThumbnail, FName& OutItemShadowBorder)
+{
+	return CreateListTileItemContents(InListItem, InThumbnail, OutItemShadowBorder);
+}
+
+TSharedRef<SWidget> FAssetViewItemHelper::CreateTileItemContents(SAssetTileItem* const InTileItem, const TSharedRef<SWidget>& InThumbnail, FName& OutItemShadowBorder)
+{
+	return CreateListTileItemContents(InTileItem, InThumbnail, OutItemShadowBorder);
+}
+
+template <typename T>
+TSharedRef<SWidget> FAssetViewItemHelper::CreateListTileItemContents(T* const InTileOrListItem, const TSharedRef<SWidget>& InThumbnail, FName& OutItemShadowBorder)
+{
+	TSharedRef<SOverlay> ItemContentsOverlay = SNew(SOverlay);
+
+	if (InTileOrListItem->IsFolder())
+	{
+		OutItemShadowBorder = FName("NoBorder");
+
+		TSharedPtr<FAssetViewFolder> AssetFolderItem = StaticCastSharedPtr<FAssetViewFolder>(InTileOrListItem->AssetItem);
+
+		const bool bIsDeveloperFolder = AssetFolderItem->bDeveloperFolder;
+
+		ECollectionShareType::Type CollectionFolderShareType = ECollectionShareType::CST_All;
+		const bool bIsCollectionFolder = ContentBrowserUtils::IsCollectionPath(AssetFolderItem->FolderPath, nullptr, &CollectionFolderShareType);
+
+		const FSlateBrush* FolderBaseImage = bIsDeveloperFolder 
+			? FEditorStyle::GetBrush("ContentBrowser.ListViewDeveloperFolderIcon.Base") 
+			: FEditorStyle::GetBrush("ContentBrowser.ListViewFolderIcon.Base");
+
+		const FSlateBrush* FolderTintImage = bIsDeveloperFolder 
+			? FEditorStyle::GetBrush("ContentBrowser.ListViewDeveloperFolderIcon.Mask") 
+			: FEditorStyle::GetBrush("ContentBrowser.ListViewFolderIcon.Mask");
+
+		// Folder base
+		ItemContentsOverlay->AddSlot()
+		[
+			SNew(SImage)
+			.Image(FolderBaseImage)
+			.ColorAndOpacity(InTileOrListItem, &T::GetAssetColor)
+		];
+
+		if (bIsCollectionFolder)
+		{
+			FLinearColor IconColor = FLinearColor::White;
+			switch(CollectionFolderShareType)
+			{
+			case ECollectionShareType::CST_Local:
+				IconColor = FColor(196, 15, 24);
+				break;
+			case ECollectionShareType::CST_Private:
+				IconColor = FColor(192, 196, 0);
+				break;
+			case ECollectionShareType::CST_Shared:
+				IconColor = FColor(0, 136, 0);
+				break;
+			default:
+				break;
+			}
+
+			auto GetCollectionIconBoxSize = [InTileOrListItem]() -> FOptionalSize
+			{
+				return FOptionalSize(InTileOrListItem->GetThumbnailBoxSize().Get() * 0.3f);
+			};
+
+			auto GetCollectionIconBrush = [=]() -> const FSlateBrush*
+			{
+				const TCHAR* IconSizeSuffix = (GetCollectionIconBoxSize().Get() <= 16.0f) ? TEXT(".Small") : TEXT(".Large");
+				return FEditorStyle::GetBrush(ECollectionShareType::GetIconStyleName(CollectionFolderShareType, IconSizeSuffix));
+			};
+
+			// Collection share type
+			ItemContentsOverlay->AddSlot()
+			.HAlign(HAlign_Center)
+			.VAlign(VAlign_Center)
+			[
+				SNew(SBox)
+				.WidthOverride_Lambda(GetCollectionIconBoxSize)
+				.HeightOverride_Lambda(GetCollectionIconBoxSize)
+				[
+					SNew(SImage)
+					.Image_Lambda(GetCollectionIconBrush)
+					.ColorAndOpacity(IconColor)
+				]
+			];
+		}
+
+		// Folder tint
+		ItemContentsOverlay->AddSlot()
+		[
+			SNew(SImage)
+			.Image(FolderTintImage)
+		];
+	}
+	else
+	{
+		OutItemShadowBorder = FName("ContentBrowser.ThumbnailShadow");
+
+		// The actual thumbnail
+		ItemContentsOverlay->AddSlot()
+		[
+			InThumbnail
+		];
+
+		// Tools for thumbnail edit mode
+		ItemContentsOverlay->AddSlot()
+		[
+			SNew(SThumbnailEditModeTools, InTileOrListItem->AssetThumbnail)
+			.SmallView(true)
+			.Visibility(InTileOrListItem, &T::GetThumbnailEditModeUIVisibility)
+		];
+
+		// Source control state
+		ItemContentsOverlay->AddSlot()
+		.HAlign(HAlign_Right)
+		.VAlign(VAlign_Top)
+		[
+			SNew(SBox)
+			.WidthOverride(InTileOrListItem, &T::GetSCCImageSize)
+			.HeightOverride(InTileOrListItem, &T::GetSCCImageSize)
+			[
+				SNew(SImage)
+				.Image(InTileOrListItem, &T::GetSCCStateImage)
+			]
+		];
+
+		// Dirty state
+		ItemContentsOverlay->AddSlot()
+		.HAlign(HAlign_Left)
+		.VAlign(VAlign_Bottom)
+		[
+			SNew(SImage)
+			.Image(InTileOrListItem, &T::GetDirtyImage)
+		];
+	}
+
+	return ItemContentsOverlay;
+}
+
+
 ///////////////////////////////
 // Asset view modes
 ///////////////////////////////
@@ -188,19 +332,15 @@ TSharedPtr<IToolTip> SAssetViewItem::GetToolTip()
 	return ShouldAllowToolTip.Get() ? SCompoundWidget::GetToolTip() : NULL;
 }
 
-bool SAssetViewItem::ValidateDragDrop( const FGeometry& MyGeometry, const FDragDropEvent& DragDropEvent ) const
+bool SAssetViewItem::ValidateDragDrop( const FGeometry& MyGeometry, const FDragDropEvent& DragDropEvent, bool& OutIsKnownDragOperation ) const
 {
-	return IsFolder() && DragDropHandler::ValidateDragDropOnAssetFolder(MyGeometry, DragDropEvent, StaticCastSharedPtr<FAssetViewFolder>(AssetItem)->FolderPath);
+	OutIsKnownDragOperation = false;
+	return IsFolder() && DragDropHandler::ValidateDragDropOnAssetFolder(MyGeometry, DragDropEvent, StaticCastSharedPtr<FAssetViewFolder>(AssetItem)->FolderPath, OutIsKnownDragOperation);
 }
 
 void SAssetViewItem::OnDragEnter( const FGeometry& MyGeometry, const FDragDropEvent& DragDropEvent )
 {
-	bDraggedOver = false;
-
-	if (ValidateDragDrop(MyGeometry, DragDropEvent))
-	{
-		bDraggedOver = true;
-	}
+	ValidateDragDrop(MyGeometry, DragDropEvent, bDraggedOver); // updates bDraggedOver
 }
 	
 void SAssetViewItem::OnDragLeave( const FDragDropEvent& DragDropEvent )
@@ -227,23 +367,16 @@ void SAssetViewItem::OnDragLeave( const FDragDropEvent& DragDropEvent )
 
 FReply SAssetViewItem::OnDragOver( const FGeometry& MyGeometry, const FDragDropEvent& DragDropEvent )
 {
-	bDraggedOver = false;
-
-	if (ValidateDragDrop(MyGeometry, DragDropEvent))
-	{
-		bDraggedOver = true;
-		return FReply::Handled();
-	}
-
-	return FReply::Unhandled();
+	ValidateDragDrop(MyGeometry, DragDropEvent, bDraggedOver); // updates bDraggedOver
+	return (bDraggedOver) ? FReply::Handled() : FReply::Unhandled();
 }
 
 FReply SAssetViewItem::OnDrop( const FGeometry& MyGeometry, const FDragDropEvent& DragDropEvent )
 {
-	bDraggedOver = false;
-
-	if (ValidateDragDrop(MyGeometry, DragDropEvent))
+	if (ValidateDragDrop(MyGeometry, DragDropEvent, bDraggedOver)) // updates bDraggedOver
 	{
+		bDraggedOver = false;
+
 		check(AssetItem->GetType() == EAssetItemType::Folder);
 
 		TSharedPtr<FDragDropOperation> Operation = DragDropEvent.GetOperation();
@@ -272,6 +405,13 @@ FReply SAssetViewItem::OnDrop( const FGeometry& MyGeometry, const FDragDropEvent
 		}
 	}
 
+	if (bDraggedOver)
+	{
+		// We were able to handle this operation, but could not due to another error - still report this drop as handled so it doesn't fall through to other widgets
+		bDraggedOver = false;
+		return FReply::Handled();
+	}
+
 	return FReply::Unhandled();
 }
 
@@ -925,7 +1065,6 @@ void SAssetListItem::Construct( const FArguments& InArgs )
 	ItemHeight = InArgs._ItemHeight;
 
 	const float ThumbnailPadding = InArgs._ThumbnailPadding;
-	bool bIsDeveloperFolder = false;
 
 	TSharedPtr<SWidget> Thumbnail;
 	if ( AssetItem.IsValid() && AssetThumbnail.IsValid() )
@@ -945,6 +1084,9 @@ void SAssetListItem::Construct( const FArguments& InArgs )
 		Thumbnail = SNew(SImage) .Image( FEditorStyle::GetDefaultBrush() );
 	}
 
+	FName ItemShadowBorderName;
+	TSharedRef<SWidget> ItemContents = FAssetViewItemHelper::CreateListItemContents(this, Thumbnail.ToSharedRef(), ItemShadowBorderName);
+
 	ChildSlot
 	[
 		SNew(SBorder)
@@ -967,66 +1109,9 @@ void SAssetListItem::Construct( const FArguments& InArgs )
 					// Drop shadow border
 					SNew(SBorder)
 					.Padding(4.f)
-					.BorderImage( IsFolder() ? FEditorStyle::GetBrush("NoBorder") : FEditorStyle::GetBrush("ContentBrowser.ThumbnailShadow") )
+					.BorderImage(FEditorStyle::GetBrush(ItemShadowBorderName))
 					[
-						SNew(SOverlay)
-				
-						// Folder base
-						+SOverlay::Slot()
-						[
-							SNew(SImage)
-							.Visibility(IsFolder() ? EVisibility::Visible : EVisibility::Collapsed)
-							.Image( bIsDeveloperFolder ? FEditorStyle::GetBrush("ContentBrowser.ListViewDeveloperFolderIcon.Base") : FEditorStyle::GetBrush("ContentBrowser.ListViewFolderIcon.Base") )
-							.ColorAndOpacity(this, &SAssetViewItem::GetAssetColor)
-						]
-
-						// Folder tint
-						+SOverlay::Slot()
-						[
-							SNew(SImage)
-							.Visibility(IsFolder() ? EVisibility::Visible : EVisibility::Collapsed)
-							.Image( bIsDeveloperFolder ? FEditorStyle::GetBrush("ContentBrowser.ListViewDeveloperFolderIcon.Mask") : FEditorStyle::GetBrush("ContentBrowser.ListViewFolderIcon.Mask") )
-						]
-
-						// The actual thumbnail
-						+SOverlay::Slot()
-						[
-							SNew(SHorizontalBox)
-							.Visibility(IsFolder() ? EVisibility::Collapsed : EVisibility::Visible)
-							+SHorizontalBox::Slot()
-							[
-								Thumbnail.ToSharedRef()
-							]
-						]
-
-						+SOverlay::Slot()
-						[
-							SNew(SThumbnailEditModeTools, AssetThumbnail)
-							.SmallView(true)
-							.Visibility(this, &SAssetListItem::GetThumbnailEditModeUIVisibility)
-						]
-
-						// Source control state
-						+SOverlay::Slot()
-						.HAlign(HAlign_Right)
-						.VAlign(VAlign_Top)
-						[
-							SNew( SBox )
-							.WidthOverride( this, &SAssetListItem::GetSCCImageSize )
-							.HeightOverride( this, &SAssetListItem::GetSCCImageSize )
-							[
-								SNew(SImage).Image( this, &SAssetListItem::GetSCCStateImage )
-							]
-						]
-
-						// Dirty state
-						+SOverlay::Slot()
-						.HAlign(HAlign_Left)
-						.VAlign(VAlign_Bottom)
-						[
-							SNew(SImage)
-							.Image( this, &SAssetListItem::GetDirtyImage )
-						]
+						ItemContents
 					]
 				]
 			]
@@ -1155,15 +1240,8 @@ void SAssetTileItem::Construct( const FArguments& InArgs )
 		Thumbnail = SNew(SImage) .Image( FEditorStyle::GetDefaultBrush() );
 	}
 
-	bool bIsDeveloperFolder = false;
-
-	if(AssetItem.IsValid())
-	{
-		if(AssetItem->GetType() == EAssetItemType::Folder)
-		{
-			bIsDeveloperFolder = StaticCastSharedPtr<FAssetViewFolder>(AssetItem)->bDeveloperFolder;
-		}
-	}
+	FName ItemShadowBorderName;
+	TSharedRef<SWidget> ItemContents = FAssetViewItemHelper::CreateTileItemContents(this, Thumbnail.ToSharedRef(), ItemShadowBorderName);
 
 	ChildSlot
 	[
@@ -1188,69 +1266,9 @@ void SAssetTileItem::Construct( const FArguments& InArgs )
 					// Drop shadow border
 					SNew(SBorder)
 					.Padding(4.f)
-					.BorderImage(IsFolder() ? FEditorStyle::GetBrush("NoBorder") : FEditorStyle::GetBrush("ContentBrowser.ThumbnailShadow"))
+					.BorderImage(FEditorStyle::GetBrush(ItemShadowBorderName))
 					[
-						SNew(SOverlay)
-
-						// Folder base
-						+SOverlay::Slot()
-						[
-							SNew(SImage)
-							.Visibility(IsFolder() ? EVisibility::Visible : EVisibility::Collapsed)
-							.Image( bIsDeveloperFolder ? FEditorStyle::GetBrush("ContentBrowser.TileViewDeveloperFolderIcon.Base") : FEditorStyle::GetBrush("ContentBrowser.TileViewFolderIcon.Base") )
-							.ColorAndOpacity(this, &SAssetViewItem::GetAssetColor)
-						]
-
-						// Folder tint
-						+SOverlay::Slot()
-						[
-							SNew(SImage)
-							.Visibility(IsFolder() ? EVisibility::Visible : EVisibility::Collapsed)
-							.Image( bIsDeveloperFolder ? FEditorStyle::GetBrush("ContentBrowser.TileViewDeveloperFolderIcon.Mask") : FEditorStyle::GetBrush("ContentBrowser.TileViewFolderIcon.Mask") )
-						]
-
-						// The actual thumbnail
-						+SOverlay::Slot()
-						[
-							SNew(SHorizontalBox)
-							.Visibility(IsFolder() ? EVisibility::Collapsed : EVisibility::Visible)
-							+SHorizontalBox::Slot()
-							[
-								Thumbnail.ToSharedRef()
-							]
-						]
-
-						// Tools for thumbnail edit mode
-						+SOverlay::Slot()
-						[
-							SNew(SThumbnailEditModeTools, AssetThumbnail)
-							.Visibility(this, &SAssetTileItem::GetThumbnailEditModeUIVisibility)
-						]
-
-						// Source control state
-						+SOverlay::Slot()
-						.HAlign(HAlign_Right)
-						.VAlign(VAlign_Top)
-						.Padding(FMargin(0, 2, 2, 0))
-						[
-							SNew( SBox )
-							.WidthOverride( this, &SAssetTileItem::GetSCCImageSize )
-							.HeightOverride( this, &SAssetTileItem::GetSCCImageSize )
-							[
-								SNew(SImage)
-								.Image( this, &SAssetTileItem::GetSCCStateImage )
-							]
-						]
-
-						// Dirty state
-						+SOverlay::Slot()
-						.HAlign(HAlign_Left)
-						.VAlign(VAlign_Bottom)
-						.Padding(FMargin(2, 0, 0, 2))
-						[
-							SNew(SImage)
-							.Image( this, &SAssetTileItem::GetDirtyImage )
-						]
+						ItemContents
 					]
 				]
 			]
diff --git a/Engine/Source/Editor/ContentBrowser/Private/AssetViewWidgets.h b/Engine/Source/Editor/ContentBrowser/Private/AssetViewWidgets.h
--- a/Engine/Source/Editor/ContentBrowser/Private/AssetViewWidgets.h
+++ b/Engine/Source/Editor/ContentBrowser/Private/AssetViewWidgets.h
@@ -7,11 +7,25 @@ DECLARE_DELEGATE_FourParams( FOnRenameCommit, const TSharedPtr<FAssetViewItem>&
 DECLARE_DELEGATE_RetVal_FourParams( bool, FOnVerifyRenameCommit, const TSharedPtr<FAssetViewItem>& /*AssetItem*/, const FText& /*NewName*/, const FSlateRect& /*MessageAnchor*/, FText& /*OutErrorMessage*/)
 DECLARE_DELEGATE_OneParam( FOnItemDestroyed, const TSharedPtr<FAssetViewItem>& /*AssetItem*/);
 
+class SAssetListItem;
+class SAssetTileItem;
+
 namespace FAssetViewModeUtils 
 {
 	FReply OnViewModeKeyDown( const TSet< TSharedPtr<FAssetViewItem> >& SelectedItems, const FKeyEvent& InKeyEvent );
 }
 
+struct FAssetViewItemHelper
+{
+public:
+	static TSharedRef<SWidget> CreateListItemContents(SAssetListItem* const InListItem, const TSharedRef<SWidget>& InThumbnail, FName& OutItemShadowBorder);
+	static TSharedRef<SWidget> CreateTileItemContents(SAssetTileItem* const InTileItem, const TSharedRef<SWidget>& InThumbnail, FName& OutItemShadowBorder);
+
+private:
+	template <typename T>
+	static TSharedRef<SWidget> CreateListTileItemContents(T* const InTileOrListItem, const TSharedRef<SWidget>& InThumbnail, FName& OutItemShadowBorder);
+};
+
 /** The tile view mode of the asset view */
 class SAssetTileView : public STileView<TSharedPtr<FAssetViewItem>>
 {
@@ -126,7 +140,7 @@ public:
 
 protected:
 	/** Used by OnDragEnter, OnDragOver, and OnDrop to check and update the validity of the drag operation */
-	bool ValidateDragDrop( const FGeometry& MyGeometry, const FDragDropEvent& DragDropEvent ) const;
+	bool ValidateDragDrop( const FGeometry& MyGeometry, const FDragDropEvent& DragDropEvent, bool& OutIsKnownDragOperation ) const;
 
 	/** Handles starting a name change */
 	virtual void HandleBeginNameChange( const FText& OriginalText );
@@ -266,7 +280,7 @@ protected:
 	/** Delegate for when a list of files is dropped on this item (if it is a folder) from an external source */
 	FOnFilesDragDropped OnFilesDragDropped;
 
-	/** Whether an item is dragged over us or not */
+	/** True when a drag is over this item with a drag operation that we know how to handle. The operation itself may not be valid to drop. */
 	bool bDraggedOver;
 
 	/** Cached brush for the source control state */
@@ -277,6 +291,8 @@ protected:
 /** An item in the asset list view */
 class SAssetListItem : public SAssetViewItem
 {
+	friend struct FAssetViewItemHelper;
+
 public:
 	SLATE_BEGIN_ARGS( SAssetListItem )
 		: _ThumbnailPadding(0)
@@ -383,6 +399,8 @@ private:
 /** An item in the asset tile view */
 class SAssetTileItem : public SAssetViewItem
 {
+	friend struct FAssetViewItemHelper;
+
 public:
 	SLATE_BEGIN_ARGS( SAssetTileItem )
 		: _ThumbnailPadding(0)
diff --git a/Engine/Source/Editor/ContentBrowser/Private/CollectionContextMenu.cpp b/Engine/Source/Editor/ContentBrowser/Private/CollectionContextMenu.cpp
--- a/Engine/Source/Editor/ContentBrowser/Private/CollectionContextMenu.cpp
+++ b/Engine/Source/Editor/ContentBrowser/Private/CollectionContextMenu.cpp
@@ -26,7 +26,7 @@ void FCollectionContextMenu::BindCommands(TSharedPtr< FUICommandList > InCommand
 		));
 }
 
-TSharedPtr<SWidget> FCollectionContextMenu::MakeCollectionListContextMenu(TSharedPtr< FUICommandList > InCommandList)
+TSharedPtr<SWidget> FCollectionContextMenu::MakeCollectionTreeContextMenu(TSharedPtr< FUICommandList > InCommandList)
 {
 	// Get all menu extenders for this context menu from the content browser module
 	FContentBrowserModule& ContentBrowserModule = FModuleManager::GetModuleChecked<FContentBrowserModule>( TEXT("ContentBrowser") );
@@ -46,7 +46,7 @@ TSharedPtr<SWidget> FCollectionContextMenu::MakeCollectionListContextMenu(TShare
 
 	UpdateProjectSourceControl();
 
-	TArray<TSharedPtr<FCollectionItem>> CollectionList = CollectionView.Pin()->CollectionListPtr->GetSelectedItems();
+	TArray<TSharedPtr<FCollectionItem>> SelectedCollections = CollectionView.Pin()->CollectionTreePtr->GetSelectedItems();
 
 	MenuBuilder.BeginSection("CollectionOptions", LOCTEXT("CollectionListOptionsMenuHeading", "Collection Options"));
 	{
@@ -58,20 +58,20 @@ TSharedPtr<SWidget> FCollectionContextMenu::MakeCollectionListContextMenu(TShare
 			);
 
 		// Only add rename/delete if at least one collection is selected
-		if ( CollectionList.Num() )
+		if ( SelectedCollections.Num() )
 		{
 			bool bAnyManagedBySCC = false;
 		
-			for (int32 CollectionIdx = 0; CollectionIdx < CollectionList.Num(); ++CollectionIdx)
+			for (int32 CollectionIdx = 0; CollectionIdx < SelectedCollections.Num(); ++CollectionIdx)
 			{
-				if ( CollectionList[CollectionIdx]->CollectionType != ECollectionShareType::CST_Local )
+				if ( SelectedCollections[CollectionIdx]->CollectionType != ECollectionShareType::CST_Local )
 				{
 					bAnyManagedBySCC = true;
 					break;
 				}
 			}
 
-			if ( CollectionList.Num() == 1 )
+			if ( SelectedCollections.Num() == 1 )
 			{
 				// Share Type
 				MenuBuilder.AddSubMenu(
@@ -95,7 +95,7 @@ TSharedPtr<SWidget> FCollectionContextMenu::MakeCollectionListContextMenu(TShare
 					)
 				);
 
-			if ( CollectionList.Num() == 1 )
+			if ( SelectedCollections.Num() == 1 )
 			{
 				// If any colors have already been set, display color options as a sub menu
 				if ( CollectionViewUtils::HasCustomColors() )
@@ -275,11 +275,11 @@ void FCollectionContextMenu::UpdateProjectSourceControl()
 
 bool FCollectionContextMenu::CanRenameSelectedCollections() const
 {
-	TArray<TSharedPtr<FCollectionItem>> CollectionList = CollectionView.Pin()->CollectionListPtr->GetSelectedItems();
+	TArray<TSharedPtr<FCollectionItem>> SelectedCollections = CollectionView.Pin()->CollectionTreePtr->GetSelectedItems();
 	
-	if(CollectionList.Num() == 1)
+	if(SelectedCollections.Num() == 1)
 	{
-		return !(CollectionList[0]->CollectionType != ECollectionShareType::CST_Local) || (bProjectUnderSourceControl && ISourceControlModule::Get().IsEnabled() && ISourceControlModule::Get().GetProvider().IsAvailable());
+		return !(SelectedCollections[0]->CollectionType != ECollectionShareType::CST_Local) || (bProjectUnderSourceControl && ISourceControlModule::Get().IsEnabled() && ISourceControlModule::Get().GetProvider().IsAvailable());
 	}
 	
 	return false;
@@ -302,16 +302,16 @@ void FCollectionContextMenu::ExecuteSetCollectionShareType(ECollectionShareType:
 		return;
 	}
 
-	TArray<TSharedPtr<FCollectionItem>> CollectionList = CollectionView.Pin()->CollectionListPtr->GetSelectedItems();
+	TArray<TSharedPtr<FCollectionItem>> SelectedCollections = CollectionView.Pin()->CollectionTreePtr->GetSelectedItems();
 
-	if ( !ensure(CollectionList.Num() == 1) )
+	if ( !ensure(SelectedCollections.Num() == 1) )
 	{
 		return;
 	}
 
 	FCollectionManagerModule& CollectionManagerModule = FCollectionManagerModule::GetModule();
 
-	CollectionManagerModule.Get().RenameCollection(CollectionList[0]->CollectionName, CollectionList[0]->CollectionType, CollectionList[0]->CollectionName, CollectionType);
+	CollectionManagerModule.Get().RenameCollection(SelectedCollections[0]->CollectionName, SelectedCollections[0]->CollectionType, SelectedCollections[0]->CollectionName, CollectionType);
 }
 
 void FCollectionContextMenu::ExecuteRenameCollection()
@@ -321,14 +321,14 @@ void FCollectionContextMenu::ExecuteRenameCollection()
 		return;
 	}
 
-	TArray<TSharedPtr<FCollectionItem>> CollectionList = CollectionView.Pin()->CollectionListPtr->GetSelectedItems();
+	TArray<TSharedPtr<FCollectionItem>> SelectedCollections = CollectionView.Pin()->CollectionTreePtr->GetSelectedItems();
 
-	if ( !ensure(CollectionList.Num() == 1) )
+	if ( !ensure(SelectedCollections.Num() == 1) )
 	{
 		return;
 	}
 
-	CollectionView.Pin()->RenameCollectionItem(CollectionList[0]);
+	CollectionView.Pin()->RenameCollectionItem(SelectedCollections[0]);
 }
 
 void FCollectionContextMenu::ExecuteDestroyCollection()
@@ -338,19 +338,19 @@ void FCollectionContextMenu::ExecuteDestroyCollection()
 		return;
 	}
 
-	TArray<TSharedPtr<FCollectionItem>> CollectionList = CollectionView.Pin()->CollectionListPtr->GetSelectedItems();
+	TArray<TSharedPtr<FCollectionItem>> SelectedCollections = CollectionView.Pin()->CollectionTreePtr->GetSelectedItems();
 
 	FText Prompt;
-	if ( CollectionList.Num() == 1 )
+	if ( SelectedCollections.Num() == 1 )
 	{
-		Prompt = FText::Format(LOCTEXT("CollectionDestroyConfirm_Single", "Delete {0}?"), FText::FromName(CollectionList[0]->CollectionName));
+		Prompt = FText::Format(LOCTEXT("CollectionDestroyConfirm_Single", "Delete {0}?"), FText::FromName(SelectedCollections[0]->CollectionName));
 	}
 	else
 	{
-		Prompt = FText::Format(LOCTEXT("CollectionDestroyConfirm_Multiple", "Delete {0} Collections?"), FText::AsNumber(CollectionList.Num()));
+		Prompt = FText::Format(LOCTEXT("CollectionDestroyConfirm_Multiple", "Delete {0} Collections?"), FText::AsNumber(SelectedCollections.Num()));
 	}
 
-	FOnClicked OnYesClicked = FOnClicked::CreateSP( this, &FCollectionContextMenu::ExecuteDestroyCollectionConfirmed, CollectionList );
+	FOnClicked OnYesClicked = FOnClicked::CreateSP( this, &FCollectionContextMenu::ExecuteDestroyCollectionConfirmed, SelectedCollections );
 	ContentBrowserUtils::DisplayConfirmationPopup(
 		Prompt,
 		LOCTEXT("CollectionDestroyConfirm_Yes", "Delete"),
@@ -361,36 +361,7 @@ void FCollectionContextMenu::ExecuteDestroyCollection()
 
 FReply FCollectionContextMenu::ExecuteDestroyCollectionConfirmed(TArray<TSharedPtr<FCollectionItem>> CollectionList)
 {
-	FCollectionManagerModule& CollectionManagerModule = FCollectionManagerModule::GetModule();
-
-	TArray<TSharedPtr<FCollectionItem>> ItemsToRemove;
-	for (int32 CollectionIdx = 0; CollectionIdx < CollectionList.Num(); ++CollectionIdx)
-	{
-		const TSharedPtr<FCollectionItem>& CollectionItem = CollectionList[CollectionIdx];
-
-		if ( CollectionManagerModule.Get().DestroyCollection(CollectionItem->CollectionName, CollectionItem->CollectionType) )
-		{
-			// Refresh the list now that the collection has been removed
-			ItemsToRemove.Add(CollectionItem);
-		}
-		else
-		{
-			// Display a warning
-			const FVector2D& CursorPos = FSlateApplication::Get().GetCursorPos();
-			FSlateRect MessageAnchor(CursorPos.X, CursorPos.Y, CursorPos.X, CursorPos.Y);
-			ContentBrowserUtils::DisplayMessage(
-				FText::Format( LOCTEXT("CollectionDestroyFailed", "Failed to destroy collection. {0}"), CollectionManagerModule.Get().GetLastError() ),
-				MessageAnchor,
-				CollectionView.Pin()->CollectionListPtr.ToSharedRef()
-				);
-		}
-	}
-
-	if ( ItemsToRemove.Num() > 0 )
-	{
-		CollectionView.Pin()->RemoveCollectionItems(ItemsToRemove);
-	}
-
+	CollectionView.Pin()->DeleteCollectionItems(CollectionList);
 	return FReply::Handled();
 }
 
@@ -448,17 +419,17 @@ bool FCollectionContextMenu::CanExecuteSetCollectionShareType(ECollectionShareTy
 		return false;
 	}
 
-	TArray<TSharedPtr<FCollectionItem>> CollectionList = CollectionView.Pin()->CollectionListPtr->GetSelectedItems();
+	TArray<TSharedPtr<FCollectionItem>> SelectedCollections = CollectionView.Pin()->CollectionTreePtr->GetSelectedItems();
 
-	if ( !ensure(CollectionList.Num() == 1) )
+	if ( !ensure(SelectedCollections.Num() == 1) )
 	{
 		return false;
 	}
 
 	const bool bIsSourceControlAvailable = bProjectUnderSourceControl && ISourceControlModule::Get().IsEnabled() && ISourceControlModule::Get().GetProvider().IsAvailable();
-	const bool bIsCurrentTypeLocal = CollectionList[0]->CollectionType == ECollectionShareType::CST_Local;
+	const bool bIsCurrentTypeLocal = SelectedCollections[0]->CollectionType == ECollectionShareType::CST_Local;
 	const bool bIsNewTypeLocal = CollectionType == ECollectionShareType::CST_Local;
-	const bool bIsNewShareTypeDifferent = CollectionList[0]->CollectionType != CollectionType;
+	const bool bIsNewShareTypeDifferent = SelectedCollections[0]->CollectionType != CollectionType;
 
 	return bIsNewShareTypeDifferent && ((bIsCurrentTypeLocal && bIsNewTypeLocal) || bIsSourceControlAvailable);
 }
@@ -470,14 +441,14 @@ bool FCollectionContextMenu::IsSetCollectionShareTypeChecked(ECollectionShareTyp
 		return false;
 	}
 
-	TArray<TSharedPtr<FCollectionItem>> CollectionList = CollectionView.Pin()->CollectionListPtr->GetSelectedItems();
+	TArray<TSharedPtr<FCollectionItem>> SelectedCollections = CollectionView.Pin()->CollectionTreePtr->GetSelectedItems();
 
-	if ( !ensure(CollectionList.Num() == 1) )
+	if ( !ensure(SelectedCollections.Num() == 1) )
 	{
 		return false;
 	}
 
-	return CollectionList[0]->CollectionType == CollectionType;
+	return SelectedCollections[0]->CollectionType == CollectionType;
 }
 
 bool FCollectionContextMenu::CanExecuteRenameCollection() const
diff --git a/Engine/Source/Editor/ContentBrowser/Private/CollectionContextMenu.h b/Engine/Source/Editor/ContentBrowser/Private/CollectionContextMenu.h
--- a/Engine/Source/Editor/ContentBrowser/Private/CollectionContextMenu.h
+++ b/Engine/Source/Editor/ContentBrowser/Private/CollectionContextMenu.h
@@ -11,8 +11,8 @@ public:
 	/** Bind menu selection commands to the command list */
 	void BindCommands(TSharedPtr< FUICommandList > InCommandList);
 
-	/** Makes the collection list context menu widget */
-	TSharedPtr<SWidget> MakeCollectionListContextMenu(TSharedPtr< FUICommandList > InCommandList);
+	/** Makes the collection tree context menu widget */
+	TSharedPtr<SWidget> MakeCollectionTreeContextMenu(TSharedPtr< FUICommandList > InCommandList);
 
 	/** Makes the new collection submenu */
 	void MakeNewCollectionSubMenu(FMenuBuilder& MenuBuilder);
diff --git a/Engine/Source/Editor/ContentBrowser/Private/CollectionViewTypes.h b/Engine/Source/Editor/ContentBrowser/Private/CollectionViewTypes.h
--- a/Engine/Source/Editor/ContentBrowser/Private/CollectionViewTypes.h
+++ b/Engine/Source/Editor/ContentBrowser/Private/CollectionViewTypes.h
@@ -19,6 +19,12 @@ struct FCollectionItem
 	/** The type of the collection */
 	ECollectionShareType::Type CollectionType;
 
+	/** Pointer to our parent collection (if any) */
+	TWeakPtr<FCollectionItem> ParentCollection;
+
+	/** Array of pointers to our child collections (if any) */
+	TArray<TWeakPtr<FCollectionItem>> ChildCollections;
+
 	/** If true, will set up an inline rename after next ScrollIntoView */
 	bool bRenaming;
 
@@ -35,6 +41,8 @@ struct FCollectionItem
 	FCollectionItem(const FName& InCollectionName, const ECollectionShareType::Type& InCollectionType)
 		: CollectionName(InCollectionName)
 		, CollectionType(InCollectionType)
+		, ParentCollection()
+		, ChildCollections()
 		, bRenaming(false)
 		, bNewCollection(false)
 	{}
diff --git a/Engine/Source/Editor/ContentBrowser/Private/ContentBrowserUtils.cpp b/Engine/Source/Editor/ContentBrowser/Private/ContentBrowserUtils.cpp
--- a/Engine/Source/Editor/ContentBrowser/Private/ContentBrowserUtils.cpp
+++ b/Engine/Source/Editor/ContentBrowser/Private/ContentBrowserUtils.cpp
@@ -1274,6 +1274,32 @@ bool ContentBrowserUtils::IsClassPath(const FString& InPath)
 	return InPath.StartsWith(ClassesRootPrefix);
 }
 
+bool ContentBrowserUtils::IsCollectionPath(const FString& InPath, FName* OutCollectionName, ECollectionShareType::Type* OutCollectionShareType)
+{
+	static const FString CollectionsRootPrefix = TEXT("/Collections");
+	if (InPath.StartsWith(CollectionsRootPrefix))
+	{
+		TArray<FString> PathParts;
+		InPath.ParseIntoArray(PathParts, TEXT("/"));
+		check(PathParts.Num() > 2);
+
+		// The second part of the path is the share type name
+		if (OutCollectionShareType)
+		{
+			*OutCollectionShareType = ECollectionShareType::FromString(*PathParts[1]);
+		}
+
+		// The third part of the path is the collection name
+		if (OutCollectionName)
+		{
+			*OutCollectionName = FName(*PathParts[2]);
+		}
+
+		return true;
+	}
+	return false;
+}
+
 void ContentBrowserUtils::CountPathTypes(const TArray<FString>& InPaths, int32& OutNumAssetPaths, int32& OutNumClassPaths)
 {
 	OutNumAssetPaths = 0;
diff --git a/Engine/Source/Editor/ContentBrowser/Private/ContentBrowserUtils.h b/Engine/Source/Editor/ContentBrowser/Private/ContentBrowserUtils.h
--- a/Engine/Source/Editor/ContentBrowser/Private/ContentBrowserUtils.h
+++ b/Engine/Source/Editor/ContentBrowser/Private/ContentBrowserUtils.h
@@ -168,6 +168,9 @@ namespace ContentBrowserUtils
 	/** Check to see whether the given path is rooted against a class directory */
 	bool IsClassPath(const FString& InPath);
 
+	/** Check to see whether the given path is rooted against a collection directory, optionally extracting the collection name and share type from the path */
+	bool IsCollectionPath(const FString& InPath, FName* OutCollectionName = nullptr, ECollectionShareType::Type* OutCollectionShareType = nullptr);
+
 	/** Given an array of paths, work out how many are rooted against class roots, and how many are rooted against asset roots */
 	void CountPathTypes(const TArray<FString>& InPaths, int32& OutNumAssetPaths, int32& OutNumClassPaths);
 
diff --git a/Engine/Source/Editor/ContentBrowser/Private/DragDropHandler.cpp b/Engine/Source/Editor/ContentBrowser/Private/DragDropHandler.cpp
--- a/Engine/Source/Editor/ContentBrowser/Private/DragDropHandler.cpp
+++ b/Engine/Source/Editor/ContentBrowser/Private/DragDropHandler.cpp
@@ -9,8 +9,10 @@
 
 #define LOCTEXT_NAMESPACE "ContentBrowser"
 
-bool DragDropHandler::ValidateDragDropOnAssetFolder(const FGeometry& MyGeometry, const FDragDropEvent& DragDropEvent, const FString& TargetPath)
+bool DragDropHandler::ValidateDragDropOnAssetFolder(const FGeometry& MyGeometry, const FDragDropEvent& DragDropEvent, const FString& TargetPath, bool& OutIsKnownDragOperation)
 {
+	OutIsKnownDragOperation = false;
+
 	TSharedPtr<FDragDropOperation> Operation = DragDropEvent.GetOperation();
 	if (!Operation.IsValid())
 	{
@@ -26,6 +28,8 @@ bool DragDropHandler::ValidateDragDropOnAssetFolder(const FGeometry& MyGeometry,
 	{
 		TSharedPtr<FAssetDragDropOp> DragDropOp = StaticCastSharedPtr<FAssetDragDropOp>(Operation);
 
+		OutIsKnownDragOperation = true;
+
 		if (bIsAssetPath)
 		{
 			int32 NumAssetItems, NumClassItems;
@@ -62,6 +66,8 @@ bool DragDropHandler::ValidateDragDropOnAssetFolder(const FGeometry& MyGeometry,
 	{
 		TSharedPtr<FAssetPathDragDropOp> DragDropOp = StaticCastSharedPtr<FAssetPathDragDropOp>(Operation);
 
+		OutIsKnownDragOperation = true;
+
 		if (DragDropOp->PathNames.Num() == 1 && DragDropOp->PathNames[0] == TargetPath)
 		{
 			DragDropOp->SetToolTip(LOCTEXT("OnDragFoldersOverFolder_CannotSelfDrop", "Cannot move or copy a folder onto itself"), FEditorStyle::GetBrush(TEXT("Graph.ConnectorFeedback.Error")));
@@ -101,6 +107,7 @@ bool DragDropHandler::ValidateDragDropOnAssetFolder(const FGeometry& MyGeometry,
 	else if (Operation->IsOfType<FExternalDragOperation>())
 	{
 		TSharedPtr<FExternalDragOperation> DragDropOp = StaticCastSharedPtr<FExternalDragOperation>(Operation);
+		OutIsKnownDragOperation = true;
 		bIsValidDrag = DragDropOp->HasFiles() && bIsAssetPath;
 	}
 
diff --git a/Engine/Source/Editor/ContentBrowser/Private/DragDropHandler.h b/Engine/Source/Editor/ContentBrowser/Private/DragDropHandler.h
--- a/Engine/Source/Editor/ContentBrowser/Private/DragDropHandler.h
+++ b/Engine/Source/Editor/ContentBrowser/Private/DragDropHandler.h
@@ -9,7 +9,7 @@ namespace DragDropHandler
 	DECLARE_DELEGATE_TwoParams(FExecuteCopyOrMoveFolders, TArray<FString> /*PathNames*/, FString /*TargetPath*/);
 
 	/** Used by OnDragEnter, OnDragOver, and OnDrop to check and update the validity of a drag-drop operation on an asset folder in the Content Browser */
-	bool ValidateDragDropOnAssetFolder(const FGeometry& MyGeometry, const FDragDropEvent& DragDropEvent, const FString& TargetPath);
+	bool ValidateDragDropOnAssetFolder(const FGeometry& MyGeometry, const FDragDropEvent& DragDropEvent, const FString& TargetPath, bool& OutIsKnownDragOperation);
 
 	/** Handle assets being dropped onto an asset folder in the Content Browser - this drop should have been externally validated by ValidateDragDropOnAssetFolder */
 	void HandleAssetsDroppedOnAssetFolder(const TSharedRef<SWidget>& ParentWidget, const TArray<FAssetData>& AssetList, const FString& TargetPath, const FText& TargetDisplayName, FExecuteCopyOrMoveAssets CopyActionHandler, FExecuteCopyOrMoveAssets MoveActionHandler);
diff --git a/Engine/Source/Editor/ContentBrowser/Private/FrontendFilters.cpp b/Engine/Source/Editor/ContentBrowser/Private/FrontendFilters.cpp
--- a/Engine/Source/Editor/ContentBrowser/Private/FrontendFilters.cpp
+++ b/Engine/Source/Editor/ContentBrowser/Private/FrontendFilters.cpp
@@ -126,7 +126,7 @@ bool AssetDataTestComplexExpression(FAssetFilterType Asset, const FName& InKey,
 		FCollectionManagerModule& CollectionManagerModule = FCollectionManagerModule::GetModule();
 
 		TArray<FName> AssetCollectionNames;
-		CollectionManagerModule.Get().GetCollectionsContainingObject(Asset.ObjectPath, ECollectionShareType::CST_All, AssetCollectionNames);
+		CollectionManagerModule.Get().GetCollectionsContainingObject(Asset.ObjectPath, ECollectionShareType::CST_All, AssetCollectionNames, ECollectionRecursionFlags::SelfAndChildren);
 
 		bool bFoundMatch = false;
 		for (const FName& AssetCollectionName : AssetCollectionNames)
diff --git a/Engine/Source/Editor/ContentBrowser/Private/SAssetView.cpp b/Engine/Source/Editor/ContentBrowser/Private/SAssetView.cpp
--- a/Engine/Source/Editor/ContentBrowser/Private/SAssetView.cpp
+++ b/Engine/Source/Editor/ContentBrowser/Private/SAssetView.cpp
@@ -1376,11 +1376,14 @@ void SAssetView::RefreshSourceItems()
 
 		if ( bFilterAllowsClasses )
 		{
+			// Include objects from child collections if we're recursing
+			const ECollectionRecursionFlags::Flags CollectionRecursionMode = (Filter.bRecursivePaths) ? ECollectionRecursionFlags::SelfAndChildren : ECollectionRecursionFlags::Self;
+
 			TArray< FName > ClassPaths;
 			FCollectionManagerModule& CollectionManagerModule = FCollectionManagerModule::GetModule();
 			for (int32 Index = 0; Index < SourcesData.Collections.Num(); Index++)
 			{
-				CollectionManagerModule.Get().GetClassesInCollection( SourcesData.Collections[Index].Name, SourcesData.Collections[Index].Type, ClassPaths );
+				CollectionManagerModule.Get().GetClassesInCollection( SourcesData.Collections[Index].Name, SourcesData.Collections[Index].Type, ClassPaths, CollectionRecursionMode );
 			}
 
 			for (int32 Index = 0; Index < ClassPaths.Num(); Index++)
@@ -1665,21 +1668,25 @@ void SAssetView::RefreshFolders()
 
 	const bool bDisplayDev = GetDefault<UContentBrowserSettings>()->GetDisplayDevelopersFolder();
 	FAssetRegistryModule& AssetRegistryModule = FModuleManager::LoadModuleChecked<FAssetRegistryModule>(TEXT("AssetRegistry"));
-	for(const FName& PackagePath : AssetPathsToShow)
-	{		
+	{
 		TArray<FString> SubPaths;
-		AssetRegistryModule.Get().GetSubPaths(PackagePath.ToString(), SubPaths, false);
-		for(const FString& SubPath : SubPaths)
+		for(const FName& PackagePath : AssetPathsToShow)
 		{
-			// If this is a developer folder, and we don't want to show them try the next path
-			if(!bDisplayDev && ContentBrowserUtils::IsDevelopersFolder(SubPath))
-			{
-				continue;
-			}
+			SubPaths.Reset();
+			AssetRegistryModule.Get().GetSubPaths(PackagePath.ToString(), SubPaths, false);
 
-			if(!Folders.Contains(SubPath))
+			for(const FString& SubPath : SubPaths)
 			{
-				FoldersToAdd.Add(SubPath);
+				// If this is a developer folder, and we don't want to show them try the next path
+				if(!bDisplayDev && ContentBrowserUtils::IsDevelopersFolder(SubPath))
+				{
+					continue;
+				}
+
+				if(!Folders.Contains(SubPath))
+				{
+					FoldersToAdd.Add(SubPath);
+				}
 			}
 		}
 	}
@@ -1700,6 +1707,25 @@ void SAssetView::RefreshFolders()
 		FoldersToAdd.Append(MatchingFolders);
 	}
 
+	// Add folders for any child collections of the currently selected collections
+	if(SourcesData.Collections.Num() > 0)
+	{
+		FCollectionManagerModule& CollectionManagerModule = FCollectionManagerModule::GetModule();
+		
+		TArray<FCollectionNameType> ChildCollections;
+		for(const FCollectionNameType& Collection : SourcesData.Collections)
+		{
+			ChildCollections.Reset();
+			CollectionManagerModule.Get().GetChildCollections(Collection.Name, Collection.Type, ChildCollections);
+
+			for(const FCollectionNameType& ChildCollection : ChildCollections)
+			{
+				// Use "Collections" as the root of the path to avoid this being confused with other asset view folders - see ContentBrowserUtils::IsCollectionPath
+				FoldersToAdd.Add(FString::Printf(TEXT("/Collections/%s/%s"), ECollectionShareType::ToString(ChildCollection.Type), *ChildCollection.Name.ToString()));
+			}
+		}
+	}
+
 	if(FoldersToAdd.Num() > 0)
 	{
 		for(const FString& FolderPath : FoldersToAdd)
@@ -2174,11 +2200,14 @@ void SAssetView::RunAssetsThroughBackendFilter(TArray<FAssetData>& InOutAssetDat
 
 		if ( SourcesData.Collections.Num() > 0 )
 		{
+			// Include objects from child collections if we're recursing
+			const ECollectionRecursionFlags::Flags CollectionRecursionMode = (Filter.bRecursivePaths) ? ECollectionRecursionFlags::SelfAndChildren : ECollectionRecursionFlags::Self;
+
 			FCollectionManagerModule& CollectionManagerModule = FCollectionManagerModule::GetModule();
 			TArray< FName > CollectionObjectPaths;
-			for (int Index = 0; Index < SourcesData.Collections.Num(); Index++)
+			for (const FCollectionNameType& Collection : SourcesData.Collections)
 			{
-				CollectionManagerModule.Get().GetObjectsInCollection(SourcesData.Collections[Index].Name, SourcesData.Collections[Index].Type, CollectionObjectPaths);
+				CollectionManagerModule.Get().GetObjectsInCollection(Collection.Name, Collection.Type, CollectionObjectPaths, CollectionRecursionMode);
 			}
 
 			for ( int32 AssetDataIdx = InOutAssetDataList.Num() - 1; AssetDataIdx >= 0; --AssetDataIdx )
@@ -3151,6 +3180,12 @@ bool SAssetView::CanOpenContextMenu() const
 		return false;
 	}
 
+	if ( SelectedAssets.Num() == 0 && SourcesData.Collections.Num() > 0 )
+	{
+		// Don't allow a context menu when we're viewing a collection and have no assets selected
+		return false;
+	}
+
 	// Build a list of selected object paths
 	TArray<FString> ObjectPaths;
 	for(auto AssetIt = SelectedAssets.CreateConstIterator(); AssetIt; ++AssetIt)
diff --git a/Engine/Source/Editor/ContentBrowser/Private/SCollectionView.cpp b/Engine/Source/Editor/ContentBrowser/Private/SCollectionView.cpp
--- a/Engine/Source/Editor/ContentBrowser/Private/SCollectionView.cpp
+++ b/Engine/Source/Editor/ContentBrowser/Private/SCollectionView.cpp
@@ -5,6 +5,8 @@
 #include "CollectionAssetManagement.h"
 #include "CollectionViewTypes.h"
 #include "CollectionContextMenu.h"
+#include "DragAndDrop/AssetDragDropOp.h"
+#include "DragAndDrop/CollectionDragDropOp.h"
 #include "ObjectTools.h"
 #include "SourcesViewWidgets.h"
 #include "ContentBrowserModule.h"
@@ -62,10 +64,13 @@ void SCollectionView::Construct( const FArguments& InArgs )
 	OnCollectionSelected = InArgs._OnCollectionSelected;
 	bAllowCollectionButtons = InArgs._AllowCollectionButtons;
 	bAllowRightClickMenu = InArgs._AllowRightClickMenu;
+	bAllowCollectionDrag = InArgs._AllowCollectionDrag;
+	bDraggedOver = false;
 
 	FCollectionManagerModule& CollectionManagerModule = FCollectionManagerModule::GetModule();
 	CollectionManagerModule.Get().OnCollectionCreated().AddSP( this, &SCollectionView::HandleCollectionCreated );
 	CollectionManagerModule.Get().OnCollectionRenamed().AddSP( this, &SCollectionView::HandleCollectionRenamed );
+	CollectionManagerModule.Get().OnCollectionReparented().AddSP( this, &SCollectionView::HandleCollectionReparented );
 	CollectionManagerModule.Get().OnCollectionDestroyed().AddSP( this, &SCollectionView::HandleCollectionDestroyed );
 
 	Commands = TSharedPtr< FUICommandList >(new FUICommandList);
@@ -86,7 +91,7 @@ void SCollectionView::Construct( const FArguments& InArgs )
 	FOnContextMenuOpening CollectionListContextMenuOpening;
 	if ( InArgs._AllowContextMenu )
 	{
-		CollectionListContextMenuOpening = FOnContextMenuOpening::CreateSP( this, &SCollectionView::MakeCollectionListContextMenu );
+		CollectionListContextMenuOpening = FOnContextMenuOpening::CreateSP( this, &SCollectionView::MakeCollectionTreeContextMenu );
 	}
 
 	PreventSelectionChangedDelegateCount = 0;
@@ -139,20 +144,21 @@ void SCollectionView::Construct( const FArguments& InArgs )
 				SNew(SSeparator)
 			]
 
-			// Collections list
+			// Collections tree
 			+SVerticalBox::Slot()
 			.FillHeight(1.f)
 			[
-				SAssignNew(CollectionListPtr, SListView< TSharedPtr<FCollectionItem> >)
-				.ListItemsSource(&FilteredCollectionItems)
-				.OnGenerateRow( this, &SCollectionView::GenerateCollectionRow )
+				SAssignNew(CollectionTreePtr, STreeView< TSharedPtr<FCollectionItem> >)
+				.TreeItemsSource(&VisibleRootCollectionItems)
+				.OnGenerateRow(this, &SCollectionView::GenerateCollectionRow)
+				.OnGetChildren(this, &SCollectionView::GetCollectionItemChildren)
 				.ItemHeight(18)
 				.SelectionMode(ESelectionMode::Multi)
 				.OnSelectionChanged(this, &SCollectionView::CollectionSelectionChanged)
-				.OnContextMenuOpening( CollectionListContextMenuOpening )
+				.OnContextMenuOpening(CollectionListContextMenuOpening)
 				.OnItemScrolledIntoView(this, &SCollectionView::CollectionItemScrolledIntoView)
 				.ClearSelectionOnClick(false)
-				.Visibility(this, &SCollectionView::GetCollectionListVisibility)
+				.Visibility(this, &SCollectionView::GetCollectionTreeVisibility)
 			];
 
 	TSharedPtr< SWidget > Content;
@@ -192,7 +198,26 @@ void SCollectionView::Construct( const FArguments& InArgs )
 
 	ChildSlot
 	[
-		Content.ToSharedRef()
+		SNew(SOverlay)
+
+		// Main content
+		+SOverlay::Slot()
+		[
+			Content.ToSharedRef()
+		]
+
+		// Drop target overlay
+		+SOverlay::Slot()
+		[
+			SNew(SBorder)
+			.Padding(0)
+			.Visibility(EVisibility::HitTestInvisible)
+			.BorderImage(this, &SCollectionView::GetCollectionViewDropTargetBorder)
+			.BorderBackgroundColor(FLinearColor::Yellow)
+			[
+				SNullWidget::NullWidget
+			]
+		]
 	];
 
 	UpdateCollectionItems();
@@ -204,6 +229,30 @@ void SCollectionView::HandleCollectionCreated( const FCollectionNameType& Collec
 }
 
 void SCollectionView::HandleCollectionRenamed( const FCollectionNameType& OriginalCollection, const FCollectionNameType& NewCollection )
+{
+	// If the original collection was expanded, we want to pass that expansion state onto its new entry
+	bool bWasExpanded = false;
+	{
+		TSharedPtr<FCollectionItem> OriginalCollectionItem = AvailableCollections.FindRef(OriginalCollection);
+		if (OriginalCollectionItem.IsValid())
+		{
+			bWasExpanded = CollectionTreePtr->IsItemExpanded(OriginalCollectionItem);
+		}
+	}
+
+	UpdateCollectionItems();
+
+	if (bWasExpanded)
+	{
+		TSharedPtr<FCollectionItem> NewCollectionItem = AvailableCollections.FindRef(NewCollection);
+		if (NewCollectionItem.IsValid())
+		{
+			CollectionTreePtr->SetItemExpansion(NewCollectionItem, true);
+		}
+	}
+}
+
+void SCollectionView::HandleCollectionReparented( const FCollectionNameType& Collection, const TOptional<FCollectionNameType>& OldParent, const TOptional<FCollectionNameType>& NewParent )
 {
 	UpdateCollectionItems();
 }
@@ -215,49 +264,149 @@ void SCollectionView::HandleCollectionDestroyed( const FCollectionNameType& Coll
 
 void SCollectionView::UpdateCollectionItems()
 {
-	CollectionItems.Empty();
+	struct FGatherCollectionItems
+	{
+		FGatherCollectionItems()
+			: CollectionManagerModule(FCollectionManagerModule::GetModule())
+		{
+		}
 
-	// Load the collection manager module
-	FCollectionManagerModule& CollectionManagerModule = FCollectionManagerModule::GetModule();
+		void GatherCollectionItems(FAvailableCollectionsMap& OutAvailableCollections)
+		{
+			OutAvailableCollections.Reset();
 
-	// Get collections of all types
-	for (int32 TypeIdx = 0; TypeIdx < ECollectionShareType::CST_All; ++TypeIdx)
-	{
-		ECollectionShareType::Type CollectionType = ECollectionShareType::Type(TypeIdx);
+			TArray<FCollectionNameType> RootCollections;
+			CollectionManagerModule.Get().GetRootCollections(RootCollections);
 
-		//Never display system collections
-		if ( CollectionType == ECollectionShareType::CST_System )
+			ProcessGatheredCollectionsAndRecurse(RootCollections, nullptr, OutAvailableCollections);
+		}
+
+		void GatherChildCollectionItems(const TSharedPtr<FCollectionItem>& InParentCollectionItem, FAvailableCollectionsMap& OutAvailableCollections)
 		{
-			continue;
+			TArray<FCollectionNameType> ChildCollections;
+			CollectionManagerModule.Get().GetChildCollections(InParentCollectionItem->CollectionName, InParentCollectionItem->CollectionType, ChildCollections);
+
+			ProcessGatheredCollectionsAndRecurse(ChildCollections, InParentCollectionItem, OutAvailableCollections);
 		}
 
-		TArray<FName> CollectionNames;
-		CollectionManagerModule.Get().GetCollectionNames(CollectionType, CollectionNames);
-		
-		for (int32 CollectionIdx = 0; CollectionIdx < CollectionNames.Num(); ++CollectionIdx)
+		void ProcessGatheredCollectionsAndRecurse(const TArray<FCollectionNameType>& InCollections, const TSharedPtr<FCollectionItem>& InParentCollectionItem, FAvailableCollectionsMap& OutAvailableCollections)
+		{
+			for (const FCollectionNameType& Collection : InCollections)
+			{
+				// Never display system collections
+				if (Collection.Type == ECollectionShareType::CST_System)
+				{
+					continue;
+				}
+
+				TSharedPtr<FCollectionItem> CollectionItem = MakeShareable(new FCollectionItem(Collection.Name, Collection.Type));
+				OutAvailableCollections.Add(Collection, CollectionItem);
+
+				if (InParentCollectionItem.IsValid())
+				{
+					// Fixup the parent and child pointers
+					InParentCollectionItem->ChildCollections.Add(CollectionItem);
+					CollectionItem->ParentCollection = InParentCollectionItem;
+				}
+
+				// Recurse
+				GatherChildCollectionItems(CollectionItem, OutAvailableCollections);
+			}
+		}
+
+		FCollectionManagerModule& CollectionManagerModule;
+	};
+
+	// Backup the current selection and expansion state of our collections
+	// We're about to re-create the tree, so we'll need to re-apply this again afterwards
+	TArray<FCollectionNameType> SelectedCollections;
+	TArray<FCollectionNameType> ExpandedCollections;
+	{
+		const auto SelectedCollectionItems = CollectionTreePtr->GetSelectedItems();
+		SelectedCollections.Reserve(SelectedCollectionItems.Num());
+		for (const TSharedPtr<FCollectionItem>& SelectedCollectionItem : SelectedCollectionItems)
 		{
-			const FName& CollectionName = CollectionNames[CollectionIdx];
-			CollectionItems.Add( MakeShareable(new FCollectionItem(CollectionName, CollectionType)) );
+			SelectedCollections.Add(FCollectionNameType(SelectedCollectionItem->CollectionName, SelectedCollectionItem->CollectionType));
+		}
+	}
+	{
+		TSet<TSharedPtr<FCollectionItem>> ExpandedCollectionItems;
+		CollectionTreePtr->GetExpandedItems(ExpandedCollectionItems);
+		ExpandedCollections.Reserve(ExpandedCollectionItems.Num());
+		for (const TSharedPtr<FCollectionItem>& ExpandedCollectionItem : ExpandedCollectionItems)
+		{
+			ExpandedCollections.Add(FCollectionNameType(ExpandedCollectionItem->CollectionName, ExpandedCollectionItem->CollectionType));
 		}
 	}
 
-	CollectionItems.Sort( FCollectionItem::FCompareFCollectionItemByName() );
+	FGatherCollectionItems GatherCollectionItems;
+	GatherCollectionItems.GatherCollectionItems(AvailableCollections);
 
 	UpdateFilteredCollectionItems();
+
+	// Restore selection and expansion
+	SetSelectedCollections(SelectedCollections, false);
+	SetExpandedCollections(ExpandedCollections);
 }
 
 void SCollectionView::UpdateFilteredCollectionItems()
 {
-	FilteredCollectionItems.Reset();
-	for ( const auto& CollectionItem : CollectionItems )
+	VisibleCollections.Reset();
+	VisibleRootCollectionItems.Reset();
+
+	auto AddVisibleCollection = [&](const TSharedPtr<FCollectionItem>& InCollectionItem)
+	{
+		VisibleCollections.Add(FCollectionNameType(InCollectionItem->CollectionName, InCollectionItem->CollectionType));
+		if (!InCollectionItem->ParentCollection.IsValid())
+		{
+			VisibleRootCollectionItems.AddUnique(InCollectionItem);
+		}
+	};
+
+	auto AddVisibleCollectionRecursive = [&](const TSharedPtr<FCollectionItem>& InCollectionItem)
 	{
-		if ( CollectionItemTextFilter->PassesFilter(*CollectionItem) )
+		TSharedPtr<FCollectionItem> CollectionItemToAdd = InCollectionItem;
+		do
 		{
-			FilteredCollectionItems.Add(CollectionItem);
+			AddVisibleCollection(CollectionItemToAdd);
+			CollectionItemToAdd = CollectionItemToAdd->ParentCollection.Pin();
+		}
+		while(CollectionItemToAdd.IsValid());
+	};
+
+	// Do we have an active filter to test against?
+	if (CollectionItemTextFilter->GetRawFilterText().IsEmpty())
+	{
+		// No filter, just mark everything as visible
+		for (const auto& AvailableCollectionInfo : AvailableCollections)
+		{
+			AddVisibleCollection(AvailableCollectionInfo.Value);
+		}
+	}
+	else
+	{
+		TArray<TSharedRef<FCollectionItem>> CollectionsToExpandTo;
+
+		// Test everything against the filter - a visible child needs to make sure its parents are also marked as visible
+		for (const auto& AvailableCollectionInfo : AvailableCollections)
+		{
+			const TSharedPtr<FCollectionItem>& CollectionItem = AvailableCollectionInfo.Value;
+			if (CollectionItemTextFilter->PassesFilter(*CollectionItem))
+			{
+				AddVisibleCollectionRecursive(CollectionItem);
+				CollectionsToExpandTo.Add(CollectionItem.ToSharedRef());
+			}
+		}
+
+		// Make sure all matching items have their parents expanded so they can be seen
+		for (const TSharedRef<FCollectionItem>& CollectionItem : CollectionsToExpandTo)
+		{
+			ExpandParentItems(CollectionItem);
 		}
 	}
 
-	CollectionListPtr->RequestListRefresh();
+	VisibleRootCollectionItems.Sort(FCollectionItem::FCompareFCollectionItemByName());
+	CollectionTreePtr->RequestTreeRefresh();
 }
 
 void SCollectionView::SetCollectionsSearchFilterText( const FText& InSearchText )
@@ -271,49 +420,71 @@ FText SCollectionView::GetCollectionsSearchFilterText() const
 	return CollectionItemTextFilter->GetRawFilterText();
 }
 
-void SCollectionView::SetSelectedCollections(const TArray<FCollectionNameType>& CollectionsToSelect)
+void SCollectionView::SetSelectedCollections(const TArray<FCollectionNameType>& CollectionsToSelect, const bool bEnsureVisible)
 {
 	// Prevent the selection changed delegate since the invoking code requested it
 	FScopedPreventSelectionChangedDelegate DelegatePrevention( SharedThis(this) );
 
 	// Expand the collections area if we are indeed selecting at least one collection
-	if ( CollectionsToSelect.Num() > 0 && CollectionsExpandableAreaPtr.IsValid() )
+	if ( bEnsureVisible && CollectionsToSelect.Num() > 0 && CollectionsExpandableAreaPtr.IsValid() )
 	{
 		CollectionsExpandableAreaPtr->SetExpanded(true);
 	}
 
-	// Clear the selection to start, then add the selected paths as they are found
-	CollectionListPtr->ClearSelection();
+	// Clear the selection to start, then add the selected items as they are found
+	CollectionTreePtr->ClearSelection();
 
-	for ( auto CollectionIt = CollectionItems.CreateConstIterator(); CollectionIt; ++CollectionIt )
+	for (const FCollectionNameType& CollectionToSelect : CollectionsToSelect)
 	{
-		const TSharedPtr<FCollectionItem>& Item = *CollectionIt;
-		for ( auto SelectionIt = CollectionsToSelect.CreateConstIterator(); SelectionIt; ++SelectionIt )
+		TSharedPtr<FCollectionItem> CollectionItemToSelect = AvailableCollections.FindRef(CollectionToSelect);
+		if (CollectionItemToSelect.IsValid())
 		{
-			const FCollectionNameType& Selection = *SelectionIt;
-			if ( Item->CollectionName == Selection.Name && Item->CollectionType == Selection.Type )
+			if (bEnsureVisible)
+			{
+				ExpandParentItems(CollectionItemToSelect.ToSharedRef());
+				CollectionTreePtr->RequestScrollIntoView(CollectionItemToSelect);
+			}
+
+			CollectionTreePtr->SetItemSelection(CollectionItemToSelect, true);
+
+			// If the selected collection doesn't pass our current filter, we need to clear it
+			if (bEnsureVisible && !CollectionItemTextFilter->PassesFilter(FCollectionItem(CollectionItemToSelect->CollectionName, CollectionItemToSelect->CollectionType)))
 			{
-				CollectionListPtr->SetItemSelection(Item, true);
-				CollectionListPtr->RequestScrollIntoView(Item);
+				SearchBoxPtr->SetText(FText::GetEmpty());
 			}
 		}
 	}
 }
 
+void SCollectionView::SetExpandedCollections(const TArray<FCollectionNameType>& CollectionsToExpand)
+{
+	// Clear the expansion to start, then add the expanded items as they are found
+	CollectionTreePtr->ClearExpandedItems();
+
+	for (const FCollectionNameType& CollectionToExpand : CollectionsToExpand)
+	{
+		TSharedPtr<FCollectionItem> CollectionItemToExpand = AvailableCollections.FindRef(CollectionToExpand);
+		if (CollectionItemToExpand.IsValid())
+		{
+			CollectionTreePtr->SetItemExpansion(CollectionItemToExpand, true);
+		}
+	}
+}
+
 void SCollectionView::ClearSelection()
 {
 	// Prevent the selection changed delegate since the invoking code requested it
 	FScopedPreventSelectionChangedDelegate DelegatePrevention( SharedThis(this) );
 
 	// Clear the selection to start, then add the selected paths as they are found
-	CollectionListPtr->ClearSelection();
+	CollectionTreePtr->ClearSelection();
 }
 
 TArray<FCollectionNameType> SCollectionView::GetSelectedCollections() const
 {
 	TArray<FCollectionNameType> RetArray;
 
-	TArray<TSharedPtr<FCollectionItem>> Items = CollectionListPtr->GetSelectedItems();
+	TArray<TSharedPtr<FCollectionItem>> Items = CollectionTreePtr->GetSelectedItems();
 	for ( int32 ItemIdx = 0; ItemIdx < Items.Num(); ++ItemIdx )
 	{
 		const TSharedPtr<FCollectionItem>& Item = Items[ItemIdx];
@@ -336,79 +507,105 @@ void SCollectionView::ApplyHistoryData( const FHistoryData& History )
 	// Prevent the selection changed delegate because it would add more history when we are just setting a state
 	FScopedPreventSelectionChangedDelegate DelegatePrevention( SharedThis(this) );
 
-	CollectionListPtr->ClearSelection();
+	CollectionTreePtr->ClearSelection();
 	for ( auto HistoryIt = History.SourcesData.Collections.CreateConstIterator(); HistoryIt; ++HistoryIt)
 	{
-		const FName Name = (*HistoryIt).Name;
-		ECollectionShareType::Type Type = (*HistoryIt).Type;
-		for ( auto CollectionIt = CollectionItems.CreateConstIterator(); CollectionIt; ++CollectionIt)
+		TSharedPtr<FCollectionItem> CollectionHistoryItem = AvailableCollections.FindRef(FCollectionNameType((*HistoryIt).Name, (*HistoryIt).Type));
+		if (CollectionHistoryItem.IsValid())
 		{
-			if ( (*CollectionIt)->CollectionName == Name && (*CollectionIt)->CollectionType == Type )
-			{
-				CollectionListPtr->SetItemSelection(*CollectionIt, true);
-				break;
-			}
+			ExpandParentItems(CollectionHistoryItem.ToSharedRef());
+			CollectionTreePtr->RequestScrollIntoView(CollectionHistoryItem);
+			CollectionTreePtr->SetItemSelection(CollectionHistoryItem, true);
 		}
 	}
 }
 
 void SCollectionView::SaveSettings(const FString& IniFilename, const FString& IniSection, const FString& SettingsString) const
 {
-	FString SelectedCollectionsString;
-	TArray< TSharedPtr<FCollectionItem> > CollectionSelectedItems = CollectionListPtr->GetSelectedItems();
-	for ( auto CollectionIt = CollectionSelectedItems.CreateConstIterator(); CollectionIt; ++CollectionIt )
+	auto SaveCollectionsArrayToIni = [&](const FString& InSubKey, const TArray<TSharedPtr<FCollectionItem>>& InCollectionItems)
 	{
-		if ( SelectedCollectionsString.Len() > 0 )
+		FString CollectionsString;
+
+		for (const TSharedPtr<FCollectionItem>& CollectionItem : InCollectionItems)
 		{
-			SelectedCollectionsString += TEXT(",");
+			if (CollectionsString.Len() > 0)
+			{
+				CollectionsString += TEXT(",");
+			}
+
+			CollectionsString += CollectionItem->CollectionName.ToString();
+			CollectionsString += TEXT("?");
+			CollectionsString += FString::FromInt(CollectionItem->CollectionType);
 		}
 
-		const TSharedPtr<FCollectionItem>& Collection = *CollectionIt;
-		SelectedCollectionsString += Collection->CollectionName.ToString() + TEXT("?") + FString::FromInt(Collection->CollectionType);
-	}
+		GConfig->SetString(*IniSection, *(SettingsString + InSubKey), *CollectionsString, IniFilename);
+	};
+
 	const bool IsCollectionsExpanded = CollectionsExpandableAreaPtr.IsValid() ? CollectionsExpandableAreaPtr->IsExpanded() : true;
 	GConfig->SetBool(*IniSection, *(SettingsString + TEXT(".CollectionsExpanded")), IsCollectionsExpanded, IniFilename);
-	GConfig->SetString(*IniSection, *(SettingsString + TEXT(".SelectedCollections")), *SelectedCollectionsString, IniFilename);
+	SaveCollectionsArrayToIni(TEXT(".SelectedCollections"), CollectionTreePtr->GetSelectedItems());
+	{
+		TSet<TSharedPtr<FCollectionItem>> ExpandedCollectionItems;
+		CollectionTreePtr->GetExpandedItems(ExpandedCollectionItems);
+		SaveCollectionsArrayToIni(TEXT(".ExpandedCollections"), ExpandedCollectionItems.Array());
+	}
 }
 
 void SCollectionView::LoadSettings(const FString& IniFilename, const FString& IniSection, const FString& SettingsString)
 {
+	auto LoadCollectionsArrayFromIni = [&](const FString& InSubKey) -> TArray<FCollectionNameType>
+	{
+		TArray<FCollectionNameType> RetCollectionsArray;
+
+		FString CollectionsArrayString;
+		if (GConfig->GetString(*IniSection, *(SettingsString + InSubKey), CollectionsArrayString, IniFilename))
+		{
+			TArray<FString> CollectionStrings;
+			CollectionsArrayString.ParseIntoArray(CollectionStrings, TEXT(","), /*bCullEmpty*/true);
+
+			for (const FString& CollectionString : CollectionStrings)
+			{
+				FString CollectionName;
+				FString CollectionTypeString;
+				if (CollectionString.Split(TEXT("?"), &CollectionName, &CollectionTypeString))
+				{
+					const int32 CollectionType = FCString::Atoi(*CollectionTypeString);
+					if (CollectionType >= 0 && CollectionType < ECollectionShareType::CST_All)
+					{
+						RetCollectionsArray.Add(FCollectionNameType(FName(*CollectionName), ECollectionShareType::Type(CollectionType)));
+					}
+				}
+			}
+		}
+
+		return RetCollectionsArray;
+	};
+
 	// Collection expansion state
 	bool bCollectionsExpanded = false;
-	if ( CollectionsExpandableAreaPtr.IsValid() && GConfig->GetBool(*IniSection, *(SettingsString + TEXT(".CollectionsExpanded")), bCollectionsExpanded, IniFilename) )
+	if (CollectionsExpandableAreaPtr.IsValid() && GConfig->GetBool(*IniSection, *(SettingsString + TEXT(".CollectionsExpanded")), bCollectionsExpanded, IniFilename))
 	{
 		CollectionsExpandableAreaPtr->SetExpanded(bCollectionsExpanded);
 	}
 
 	// Selected Collections
-	FString SelectedCollectionsString;
-	TArray<FCollectionNameType> NewSelectedCollections;
-	if ( GConfig->GetString(*IniSection, *(SettingsString + TEXT(".SelectedCollections")), SelectedCollectionsString, IniFilename) )
+	TArray<FCollectionNameType> NewSelectedCollections = LoadCollectionsArrayFromIni(TEXT(".SelectedCollections"));
+	if (NewSelectedCollections.Num() > 0)
 	{
-		TArray<FString> NewSelectedCollectionStrings;
-		SelectedCollectionsString.ParseIntoArray(NewSelectedCollectionStrings, TEXT(","), /*bCullEmpty*/true);
+		SetSelectedCollections(NewSelectedCollections);
 
-		for ( auto CollectionIt = NewSelectedCollectionStrings.CreateConstIterator(); CollectionIt; ++CollectionIt )
+		const TArray<TSharedPtr<FCollectionItem>> SelectedCollectionItems = CollectionTreePtr->GetSelectedItems();
+		if (SelectedCollectionItems.Num() > 0)
 		{
-			FString CollectionName;
-			FString CollectionTypeString;
-			if ( (*CollectionIt).Split(TEXT("?"), &CollectionName, &CollectionTypeString) )
-			{
-				int32 CollectionType = FCString::Atoi(*CollectionTypeString);
-
-				if ( CollectionType >= 0 && CollectionType < ECollectionShareType::CST_All )
-				{
-					new(NewSelectedCollections) FCollectionNameType(FName(*CollectionName), ECollectionShareType::Type(CollectionType) );
-				}
-			}
+			CollectionSelectionChanged(SelectedCollectionItems[0], ESelectInfo::Direct);
 		}
 	}
 
-	if ( NewSelectedCollections.Num() > 0 )
+	// Expanded Collections
+	TArray<FCollectionNameType> NewExpandedCollections = LoadCollectionsArrayFromIni(TEXT(".ExpandedCollections"));
+	if (NewExpandedCollections.Num() > 0)
 	{
-		// Select the collections
-		SetSelectedCollections(NewSelectedCollections);
-		CollectionSelectionChanged( TSharedPtr<FCollectionItem>(), ESelectInfo::Direct );
+		SetExpandedCollections(NewExpandedCollections);
 	}
 }
 
@@ -421,6 +618,40 @@ FReply SCollectionView::OnKeyDown( const FGeometry& MyGeometry, const FKeyEvent&
 	return FReply::Unhandled();
 }
 
+void SCollectionView::OnDragEnter( const FGeometry& MyGeometry, const FDragDropEvent& DragDropEvent )
+{
+	ValidateDragDropOnCollectionTree(MyGeometry, DragDropEvent, bDraggedOver); // updates bDraggedOver
+}
+
+void SCollectionView::OnDragLeave( const FDragDropEvent& DragDropEvent )
+{
+	bDraggedOver = false;
+}
+
+FReply SCollectionView::OnDragOver( const FGeometry& MyGeometry, const FDragDropEvent& DragDropEvent )
+{
+	ValidateDragDropOnCollectionTree(MyGeometry, DragDropEvent, bDraggedOver); // updates bDraggedOver
+	return (bDraggedOver) ? FReply::Handled() : FReply::Unhandled();
+}
+
+FReply SCollectionView::OnDrop( const FGeometry& MyGeometry, const FDragDropEvent& DragDropEvent )
+{
+	if (ValidateDragDropOnCollectionTree(MyGeometry, DragDropEvent, bDraggedOver)) // updates bDraggedOver
+	{
+		bDraggedOver = false;
+		return HandleDragDropOnCollectionTree(MyGeometry, DragDropEvent);
+	}
+
+	if (bDraggedOver)
+	{
+		// We were able to handle this operation, but could not due to another error - still report this drop as handled so it doesn't fall through to other widgets
+		bDraggedOver = false;
+		return FReply::Handled();
+	}
+
+	return FReply::Unhandled();
+}
+
 bool SCollectionView::ShouldAllowSelectionChangedDelegate() const
 {
 	return PreventSelectionChangedDelegateCount == 0;
@@ -494,11 +725,10 @@ void SCollectionView::CreateCollectionItem( ECollectionShareType::Type Collectio
 		NewItem->bRenaming = true;
 		NewItem->bNewCollection = true;
 
-		CollectionItems.Add( NewItem );
-		CollectionItems.Sort( FCollectionItem::FCompareFCollectionItemByName() );
+		AvailableCollections.Add( FCollectionNameType(NewItem->CollectionName, NewItem->CollectionType), NewItem );
 		UpdateFilteredCollectionItems();
-		CollectionListPtr->RequestScrollIntoView(NewItem);
-		CollectionListPtr->SetSelection( NewItem );
+		CollectionTreePtr->RequestScrollIntoView(NewItem);
+		CollectionTreePtr->SetSelection( NewItem );
 	}
 }
 
@@ -507,54 +737,149 @@ void SCollectionView::RenameCollectionItem( const TSharedPtr<FCollectionItem>& I
 	if ( ensure(ItemToRename.IsValid()) )
 	{
 		ItemToRename->bRenaming = true;
-		CollectionListPtr->RequestScrollIntoView(ItemToRename);
+		CollectionTreePtr->RequestScrollIntoView(ItemToRename);
 	}
 }
 
-void SCollectionView::RemoveCollectionItems( const TArray<TSharedPtr<FCollectionItem>>& ItemsToRemove )
+void SCollectionView::DeleteCollectionItems( const TArray<TSharedPtr<FCollectionItem>>& ItemsToDelete )
 {
-	TArray<TSharedPtr<FCollectionItem>> SelectedItems = CollectionListPtr->GetSelectedItems();
+	if (ItemsToDelete.Num() == 0)
+	{
+		return;
+	}
+
+	// Before we delete anything (as this will trigger a tree update) we need to work out what our new selection should be in the case that 
+	// all of the selected items are removed
+	const TArray<TSharedPtr<FCollectionItem>> PreviouslySelectedItems = CollectionTreePtr->GetSelectedItems();
 
-	// Remove all the items, while keeping track of the number of selected items that were removed.
-	int32 LastSelectedItemIdx = INDEX_NONE;
-	int32 NumSelectedItemsRemoved = 0;
-	for (int32 RemoveIdx = 0; RemoveIdx < ItemsToRemove.Num(); ++RemoveIdx)
+	// Get the first selected item that will be deleted so we can find a suitable new selection
+	TSharedPtr<FCollectionItem> FirstSelectedItemDeleted;
+	for (const auto& ItemToDelete : ItemsToDelete)
 	{
-		const TSharedPtr<FCollectionItem>& ItemToRemove = ItemsToRemove[RemoveIdx];
+		if (PreviouslySelectedItems.Contains(ItemToDelete))
+		{
+			FirstSelectedItemDeleted = ItemToDelete;
+			break;
+		}
+	}
+
+	// Build up an array of potential new selections (in the case that we're deleting everything that's selected)
+	// Earlier items should be considered first, we base this list on the first selected item that will be deleted, and include previous siblings, and then all parents and roots
+	TArray<FCollectionNameType> PotentialNewSelections;
+	if (FirstSelectedItemDeleted.IsValid())
+	{
+		TSharedPtr<FCollectionItem> RootSelectedItemDeleted = FirstSelectedItemDeleted;
+		TSharedPtr<FCollectionItem> ParentCollectionItem = FirstSelectedItemDeleted->ParentCollection.Pin();
 
-		int32 ItemIdx = INDEX_NONE;
-		if ( CollectionItems.Find(ItemToRemove, ItemIdx) )
+		if (ParentCollectionItem.IsValid())
 		{
-			if ( SelectedItems.Contains(ItemToRemove) )
+			// Add all the siblings until we find the item that will be deleted
+			for (const auto& ChildItemWeakPtr : ParentCollectionItem->ChildCollections)
+			{
+				TSharedPtr<FCollectionItem> ChildItem = ChildItemWeakPtr.Pin();
+				if (ChildItem.IsValid())
+				{
+					if (ChildItem == FirstSelectedItemDeleted)
+					{
+						break;
+					}
+					
+					// We add siblings as the start, as the closest sibling should be the first match
+					PotentialNewSelections.Insert(FCollectionNameType(ChildItem->CollectionName, ChildItem->CollectionType), 0);
+				}
+			}
+
+			// Now add this parent, and all other parents too
+			do
 			{
-				NumSelectedItemsRemoved++;
-				LastSelectedItemIdx = ItemIdx;
+				PotentialNewSelections.Add(FCollectionNameType(ParentCollectionItem->CollectionName, ParentCollectionItem->CollectionType));
+				RootSelectedItemDeleted = ParentCollectionItem;
+				ParentCollectionItem = ParentCollectionItem->ParentCollection.Pin();
 			}
+			while (ParentCollectionItem.IsValid());
+		}
 
-			CollectionItems.RemoveAt(ItemIdx);
+		if (RootSelectedItemDeleted.IsValid())
+		{
+			// Add all the root level items before this one
+			const int32 InsertionPoint = PotentialNewSelections.Num();
+			for (const auto& RootItem : VisibleRootCollectionItems)
+			{
+				if (RootItem == RootSelectedItemDeleted)
+				{
+					break;
+				}
+				
+				// Add each root item at the insertion point, as the closest item should be a better match
+				PotentialNewSelections.Insert(FCollectionNameType(RootItem->CollectionName, RootItem->CollectionType), InsertionPoint);
+			}
 		}
 	}
 
-	FScopedPreventSelectionChangedDelegate DelegatePrevention(SharedThis(this));
-	CollectionListPtr->ClearSelection();
+	// Delete all given collections
+	int32 NumSelectedItemsDeleted = 0;
+	FCollectionManagerModule& CollectionManagerModule = FCollectionManagerModule::GetModule();
+	for (const TSharedPtr<FCollectionItem>& ItemToDelete : ItemsToDelete)
+	{
+		if (CollectionManagerModule.Get().DestroyCollection(ItemToDelete->CollectionName, ItemToDelete->CollectionType))
+		{
+			if (PreviouslySelectedItems.Contains(ItemToDelete))
+			{
+				++NumSelectedItemsDeleted;
+			}
+		}
+		else
+		{
+			// Display a warning
+			const FVector2D& CursorPos = FSlateApplication::Get().GetCursorPos();
+			FSlateRect MessageAnchor(CursorPos.X, CursorPos.Y, CursorPos.X, CursorPos.Y);
+			ContentBrowserUtils::DisplayMessage(
+				FText::Format( LOCTEXT("CollectionDestroyFailed", "Failed to destroy collection. {0}"), CollectionManagerModule.Get().GetLastError() ),
+				MessageAnchor,
+				CollectionTreePtr.ToSharedRef()
+				);
+		}
+	}
+
+	// DestroyCollection will have triggered a notification that will have updated the tree, we now need to apply a suitable selection...
 
-	// If we removed all the selected items and there is at least one other item, select it
-	if ( LastSelectedItemIdx > INDEX_NONE && NumSelectedItemsRemoved > 0 && NumSelectedItemsRemoved >= SelectedItems.Num() && CollectionItems.Num() > 1 )
+	// Did this delete change the list of selected items?
+	if (NumSelectedItemsDeleted > 0 || PreviouslySelectedItems.Num() == 0)
 	{
-		// The last selected item idx will refer to the next item in the list now that it has been removed.
-		// If the removed item was the last item in the list, select the new last item
-		const int32 NewSelectedItemIdx = FMath::Min(LastSelectedItemIdx, CollectionItems.Num() - 1);
-		
-		CollectionListPtr->SetSelection(CollectionItems[NewSelectedItemIdx]);
+		// If we removed everything that was selected, we need to try and find a suitable replacement...
+		if (NumSelectedItemsDeleted >= PreviouslySelectedItems.Num() && VisibleCollections.Num() > 1)
+		{
+			// Include the first visible item as an absolute last resort should everything else suitable have been removed from the tree
+			PotentialNewSelections.Add(*VisibleCollections.CreateConstIterator());
+
+			// Check the potential new selections array and try and select the first one that's still visible in the tree
+			TArray<FCollectionNameType> NewItemSelection;
+			for (const FCollectionNameType& PotentialNewSelection : PotentialNewSelections)
+			{
+				if (VisibleCollections.Contains(PotentialNewSelection))
+				{
+					NewItemSelection.Add(PotentialNewSelection);
+					break;
+				}
+			}
+
+			SetSelectedCollections(NewItemSelection, true);
+		}
+
+		// Broadcast the new selection
+		const TArray<TSharedPtr<FCollectionItem>> UpdatedSelectedItems = CollectionTreePtr->GetSelectedItems();
+		CollectionSelectionChanged((UpdatedSelectedItems.Num() > 0) ? UpdatedSelectedItems[0] : nullptr, ESelectInfo::Direct);
 	}
+}
 
-	// Refresh the list
-	UpdateFilteredCollectionItems();
+EVisibility SCollectionView::GetCollectionTreeVisibility() const
+{
+	return AvailableCollections.Num() > 0 ? EVisibility::Visible : EVisibility::Collapsed;
 }
 
-EVisibility SCollectionView::GetCollectionListVisibility() const
+const FSlateBrush* SCollectionView::GetCollectionViewDropTargetBorder() const
 {
-	return CollectionItems.Num() > 0 ? EVisibility::Visible : EVisibility::Collapsed;
+	return bDraggedOver ? FEditorStyle::GetBrush("ContentBrowser.CollectionTreeDragDropBorder") : FEditorStyle::GetBrush("NoBorder");
 }
 
 TSharedRef<ITableRow> SCollectionView::GenerateCollectionRow( TSharedPtr<FCollectionItem> CollectionItem, const TSharedRef<STableViewBase>& OwnerTable )
@@ -572,34 +897,255 @@ TSharedRef<ITableRow> SCollectionView::GenerateCollectionRow( TSharedPtr<FCollec
 		OnCollectionCheckStateChangedDelegate.BindSP(this, &SCollectionView::OnCollectionCheckStateChanged, CollectionItem);
 	}
 
-	TSharedPtr< STableRow< TSharedPtr<FTreeItem> > > TableRow = SNew( STableRow< TSharedPtr<FTreeItem> >, OwnerTable );
+	TSharedPtr< STableRow< TSharedPtr<FCollectionItem> > > TableRow = SNew( STableRow< TSharedPtr<FCollectionItem> >, OwnerTable )
+		.OnDragDetected(this, &SCollectionView::OnCollectionDragDetected);
+
 	TableRow->SetContent
 		(
-			SNew(SCollectionListItem)
+			SNew(SCollectionTreeItem)
 			.ParentWidget(SharedThis(this))
 			.CollectionItem(CollectionItem)
 			.OnNameChangeCommit(this, &SCollectionView::CollectionNameChangeCommit)
 			.OnVerifyRenameCommit(this, &SCollectionView::CollectionVerifyRenameCommit)
-			.OnAssetsDragDropped(this, &SCollectionView::CollectionAssetsDropped)
-			.IsSelected( TableRow.Get(), &STableRow< TSharedPtr<FTreeItem> >::IsSelectedExclusively )
-			.IsReadOnly(this, &SCollectionView::IsCollectionNotRenamable)
-			.HighlightText( this, &SCollectionView::GetCollectionsSearchFilterText)
-			.IsCheckBoxEnabled( IsCollectionCheckBoxEnabledAttribute )
-			.IsCollectionChecked( IsCollectionCheckedAttribute )
-			.OnCollectionCheckStateChanged( OnCollectionCheckStateChangedDelegate )
+			.OnValidateDragDrop(this, &SCollectionView::ValidateDragDropOnCollectionItem)
+			.OnHandleDragDrop(this, &SCollectionView::HandleDragDropOnCollectionItem)
+			.IsSelected(TableRow.Get(), &STableRow< TSharedPtr<FCollectionItem> >::IsSelectedExclusively)
+			.IsReadOnly(this, &SCollectionView::IsCollectionNameReadOnly)
+			.HighlightText(this, &SCollectionView::GetCollectionsSearchFilterText)
+			.IsCheckBoxEnabled(IsCollectionCheckBoxEnabledAttribute)
+			.IsCollectionChecked(IsCollectionCheckedAttribute)
+			.OnCollectionCheckStateChanged(OnCollectionCheckStateChangedDelegate)
 		);
 
 	return TableRow.ToSharedRef();
 }
 
-TSharedPtr<SWidget> SCollectionView::MakeCollectionListContextMenu()
+void SCollectionView::GetCollectionItemChildren( TSharedPtr<FCollectionItem> InParentItem, TArray< TSharedPtr<FCollectionItem> >& OutChildItems ) const
+{
+	for (const auto& ChildItemWeakPtr : InParentItem->ChildCollections)
+	{
+		TSharedPtr<FCollectionItem> ChildItem = ChildItemWeakPtr.Pin();
+		if (ChildItem.IsValid() && VisibleCollections.Contains(FCollectionNameType(ChildItem->CollectionName, ChildItem->CollectionType)))
+		{
+			OutChildItems.Add(ChildItem);
+		}
+	}
+	OutChildItems.Sort(FCollectionItem::FCompareFCollectionItemByName());
+}
+
+FReply SCollectionView::OnCollectionDragDetected(const FGeometry& Geometry, const FPointerEvent& MouseEvent)
+{
+	if (bAllowCollectionDrag && MouseEvent.IsMouseButtonDown(EKeys::LeftMouseButton))
+	{
+		const TArray<FCollectionNameType> SelectedCollections = GetSelectedCollections();
+		if (SelectedCollections.Num() > 0)
+		{
+			TSharedRef<FCollectionDragDropOp> DragDropOp = FCollectionDragDropOp::New(SelectedCollections);
+			CurrentCollectionDragDropOp = DragDropOp;
+			return FReply::Handled().BeginDragDrop(DragDropOp);
+		}
+	}
+
+	return FReply::Unhandled();
+}
+
+bool SCollectionView::ValidateDragDropOnCollectionTree(const FGeometry& Geometry, const FDragDropEvent& DragDropEvent, bool& OutIsKnownDragOperation)
+{
+	OutIsKnownDragOperation = false;
+
+	TSharedPtr<FDragDropOperation> Operation = DragDropEvent.GetOperation();
+	if (!Operation.IsValid())
+	{
+		return false;
+	}
+
+	if (Operation->IsOfType<FCollectionDragDropOp>())
+	{
+		OutIsKnownDragOperation = true;
+		return true;
+	}
+
+	return false;
+}
+
+FReply SCollectionView::HandleDragDropOnCollectionTree(const FGeometry& Geometry, const FDragDropEvent& DragDropEvent)
+{
+	// Should have already called ValidateDragDropOnCollectionTree prior to calling this...
+	TSharedPtr<FDragDropOperation> Operation = DragDropEvent.GetOperation();
+	check(Operation.IsValid());
+
+	FCollectionManagerModule& CollectionManagerModule = FCollectionManagerModule::GetModule();
+
+	if (Operation->IsOfType<FCollectionDragDropOp>())
+	{
+		TSharedPtr<FCollectionDragDropOp> DragDropOp = StaticCastSharedPtr<FCollectionDragDropOp>(Operation);
+		
+		// Reparent all of the collections in the drag drop so that they are root level items
+		for (const FCollectionNameType& NewChildCollection : DragDropOp->Collections)
+		{
+			if (!CollectionManagerModule.Get().ReparentCollection(
+					NewChildCollection.Name, NewChildCollection.Type,
+					NAME_None, ECollectionShareType::CST_All
+					))
+			{
+				ContentBrowserUtils::DisplayMessage(CollectionManagerModule.Get().GetLastError(), Geometry.GetClippingRect(), SharedThis(this));
+			}
+		}
+
+		return FReply::Handled();
+	}
+
+	return FReply::Unhandled();
+}
+
+bool SCollectionView::ValidateDragDropOnCollectionItem(TSharedRef<FCollectionItem> CollectionItem, const FGeometry& Geometry, const FDragDropEvent& DragDropEvent, bool& OutIsKnownDragOperation)
+{
+	OutIsKnownDragOperation = false;
+
+	TSharedPtr<FDragDropOperation> Operation = DragDropEvent.GetOperation();
+	if (!Operation.IsValid())
+	{
+		return false;
+	}
+
+	bool bIsValidDrag = false;
+	TOptional<EMouseCursor::Type> NewDragCursor;
+
+	if (Operation->IsOfType<FCollectionDragDropOp>())
+	{
+		TSharedPtr<FCollectionDragDropOp> DragDropOp = StaticCastSharedPtr<FCollectionDragDropOp>(Operation);
+
+		OutIsKnownDragOperation = true;
+
+		FCollectionManagerModule& CollectionManagerModule = FCollectionManagerModule::GetModule();
+		
+		bIsValidDrag = true;
+		for (const FCollectionNameType& PotentialChildCollection : DragDropOp->Collections)
+		{
+			bIsValidDrag = CollectionManagerModule.Get().IsValidParentCollection(
+				PotentialChildCollection.Name, PotentialChildCollection.Type,
+				CollectionItem->CollectionName, CollectionItem->CollectionType
+				);
+
+			if (!bIsValidDrag)
+			{
+				DragDropOp->SetToolTip(CollectionManagerModule.Get().GetLastError(), FEditorStyle::GetBrush(TEXT("Graph.ConnectorFeedback.Error")));
+				break;
+			}
+		}
+
+		// If we are dragging over a child collection item, then this view as a whole should not be marked as dragged over
+		bDraggedOver = false;
+	}
+	else if (Operation->IsOfType<FAssetDragDropOp>())
+	{
+		TSharedPtr<FAssetDragDropOp> DragDropOp = StaticCastSharedPtr<FAssetDragDropOp>(Operation);
+		OutIsKnownDragOperation = true;
+		bIsValidDrag = DragDropOp->AssetData.Num() > 0;
+	}
+
+	// Set the default slashed circle if this drag is invalid and a drag operation hasn't set NewDragCursor to something custom
+	if (!bIsValidDrag && !NewDragCursor.IsSet())
+	{
+		NewDragCursor = EMouseCursor::SlashedCircle;
+	}
+	Operation->SetCursorOverride(NewDragCursor);
+
+	return bIsValidDrag;
+}
+
+FReply SCollectionView::HandleDragDropOnCollectionItem(TSharedRef<FCollectionItem> CollectionItem, const FGeometry& Geometry, const FDragDropEvent& DragDropEvent)
+{
+	// Should have already called ValidateDragDropOnCollectionItem prior to calling this...
+	TSharedPtr<FDragDropOperation> Operation = DragDropEvent.GetOperation();
+	check(Operation.IsValid());
+
+	FCollectionManagerModule& CollectionManagerModule = FCollectionManagerModule::GetModule();
+
+	if (Operation->IsOfType<FCollectionDragDropOp>())
+	{
+		TSharedPtr<FCollectionDragDropOp> DragDropOp = StaticCastSharedPtr<FCollectionDragDropOp>(Operation);
+		
+		// Make sure our drop item is marked as expanded so that we'll be able to see the newly added children
+		CollectionTreePtr->SetItemExpansion(CollectionItem, true);
+
+		// Reparent all of the collections in the drag drop so that they are our immediate children
+		for (const FCollectionNameType& NewChildCollection : DragDropOp->Collections)
+		{
+			if (!CollectionManagerModule.Get().ReparentCollection(
+					NewChildCollection.Name, NewChildCollection.Type,
+					CollectionItem->CollectionName, CollectionItem->CollectionType
+					))
+			{
+				ContentBrowserUtils::DisplayMessage(CollectionManagerModule.Get().GetLastError(), Geometry.GetClippingRect(), SharedThis(this));
+			}
+		}
+
+		return FReply::Handled();
+	}
+	else if (Operation->IsOfType<FAssetDragDropOp>())
+	{
+		TSharedPtr<FAssetDragDropOp> DragDropOp = StaticCastSharedPtr<FAssetDragDropOp>(Operation);
+			
+		TArray<FName> ObjectPaths;
+		ObjectPaths.Reserve(DragDropOp->AssetData.Num());
+		for (const FAssetData& AssetData : DragDropOp->AssetData)
+		{
+			ObjectPaths.Add(AssetData.ObjectPath);
+		}
+
+		int32 NumAdded = 0;
+		FText Message;
+		if (CollectionManagerModule.Get().AddToCollection(CollectionItem->CollectionName, CollectionItem->CollectionType, ObjectPaths, &NumAdded))
+		{
+			if (DragDropOp->AssetData.Num() == 1)
+			{
+				FFormatNamedArguments Args;
+				Args.Add(TEXT("AssetName"), FText::FromName(DragDropOp->AssetData[0].AssetName));
+				Args.Add(TEXT("CollectionName"), FText::FromName(CollectionItem->CollectionName));
+				Message = FText::Format(LOCTEXT("CollectionAssetsAdded", "Added {AssetName} to {CollectionName}"), Args);
+			}
+			else
+			{
+				FFormatNamedArguments Args;
+				Args.Add(TEXT("Number"), NumAdded);
+				Args.Add(TEXT("CollectionName"), FText::FromName(CollectionItem->CollectionName));
+				Message = FText::Format(LOCTEXT("CollectionAssetsAdded", "Added {Number} asset(s) to {CollectionName}"), Args);
+			}
+		}
+		else
+		{
+			Message = CollectionManagerModule.Get().GetLastError();
+		}
+
+		// Added items to the collection or failed. Either way, display the message.
+		ContentBrowserUtils::DisplayMessage(Message, Geometry.GetClippingRect(), SharedThis(this));
+
+		return FReply::Handled();
+	}
+
+	return FReply::Unhandled();
+}
+
+void SCollectionView::ExpandParentItems(const TSharedRef<FCollectionItem>& InCollectionItem)
+{
+	for (TSharedPtr<FCollectionItem> CollectionItemToExpand = InCollectionItem->ParentCollection.Pin(); 
+		CollectionItemToExpand.IsValid(); 
+		CollectionItemToExpand = CollectionItemToExpand->ParentCollection.Pin()
+		)
+	{
+		CollectionTreePtr->SetItemExpansion(CollectionItemToExpand, true);
+	}
+}
+
+TSharedPtr<SWidget> SCollectionView::MakeCollectionTreeContextMenu()
 {
 	if ( !bAllowRightClickMenu )
 	{
 		return NULL;
 	}
 
-	return CollectionContextMenu->MakeCollectionListContextMenu(Commands);
+	return CollectionContextMenu->MakeCollectionTreeContextMenu(Commands);
 }
 
 bool SCollectionView::IsCollectionCheckBoxEnabled( TSharedPtr<FCollectionItem> CollectionItem ) const
@@ -651,40 +1197,6 @@ void SCollectionView::CollectionSelectionChanged( TSharedPtr< FCollectionItem >
 	}
 }
 
-void SCollectionView::CollectionAssetsDropped(const TArray<FAssetData>& AssetList, const TSharedPtr<FCollectionItem>& CollectionItem, FText& OutMessage)
-{
-	TArray<FName> ObjectPaths;
-
-	for ( auto AssetIt = AssetList.CreateConstIterator(); AssetIt; ++AssetIt )
-	{
-		ObjectPaths.Add((*AssetIt).ObjectPath);
-	}
-
-	FCollectionManagerModule& CollectionManagerModule = FCollectionManagerModule::GetModule();
-	int32 NumAdded = 0;
-	if ( CollectionManagerModule.Get().AddToCollection(CollectionItem->CollectionName, CollectionItem->CollectionType, ObjectPaths, &NumAdded) )
-	{
-		if ( AssetList.Num() == 1 )
-		{
-			FFormatNamedArguments Args;
-			Args.Add( TEXT("AssetName"), FText::FromName(AssetList[0].AssetName) );
-			Args.Add( TEXT("CollectionName"), FText::FromName(CollectionItem->CollectionName) );
-			OutMessage = FText::Format( LOCTEXT("CollectionAssetsAdded", "Added {AssetName} to {CollectionName}"), Args );
-		}
-		else
-		{
-			FFormatNamedArguments Args;
-			Args.Add( TEXT("Number"), NumAdded );
-			Args.Add( TEXT("CollectionName"), FText::FromName(CollectionItem->CollectionName) );
-			OutMessage = FText::Format( LOCTEXT("CollectionAssetsAdded", "Added {Number} asset(s) to {CollectionName}"), Args );
-		}
-	}
-	else
-	{
-		OutMessage = CollectionManagerModule.Get().GetLastError();
-	}
-}
-
 void SCollectionView::CollectionItemScrolledIntoView( TSharedPtr<FCollectionItem> CollectionItem, const TSharedPtr<ITableRow>& Widget )
 {
 	if ( CollectionItem->bRenaming && Widget.IsValid() && Widget->GetContent().IsValid() )
@@ -693,8 +1205,22 @@ void SCollectionView::CollectionItemScrolledIntoView( TSharedPtr<FCollectionItem
 	}
 }
 
-bool SCollectionView::IsCollectionNotRenamable() const
+bool SCollectionView::IsCollectionNameReadOnly() const
 {
+	// We can't rename collections while they're being dragged
+	TSharedPtr<FCollectionDragDropOp> DragDropOp = CurrentCollectionDragDropOp.Pin();
+	if (DragDropOp.IsValid())
+	{
+		TArray<TSharedPtr<FCollectionItem>> SelectedCollectionItems = CollectionTreePtr->GetSelectedItems();
+		for (const auto& SelectedCollectionItem : SelectedCollectionItems)
+		{
+			if (DragDropOp->Collections.Contains(FCollectionNameType(SelectedCollectionItem->CollectionName, SelectedCollectionItem->CollectionType)))
+			{
+				return true;
+			}
+		}
+	}
+
 	CollectionContextMenu->UpdateProjectSourceControl();
 	return !CollectionContextMenu->CanRenameSelectedCollections();
 }
@@ -702,7 +1228,7 @@ bool SCollectionView::IsCollectionNotRenamable() const
 bool SCollectionView::CollectionNameChangeCommit( const TSharedPtr< FCollectionItem >& CollectionItem, const FString& NewName, bool bChangeConfirmed, FText& OutWarningMessage )
 {
 	// There should only ever be one item selected when renaming
-	check(CollectionListPtr->GetNumItemsSelected() == 1);
+	check(CollectionTreePtr->GetNumItemsSelected() == 1);
 
 	FCollectionManagerModule& CollectionManagerModule = FCollectionManagerModule::GetModule();
 	ECollectionShareType::Type CollectionType = CollectionItem->CollectionType;
@@ -717,7 +1243,7 @@ bool SCollectionView::CollectionNameChangeCommit( const TSharedPtr< FCollectionI
 		// If we can canceled the name change when creating a new asset, we want to silently remove it
 		if ( !bChangeConfirmed )
 		{
-			CollectionItems.Remove(CollectionItem);
+			AvailableCollections.Remove(FCollectionNameType(CollectionItem->CollectionName, CollectionItem->CollectionType));
 			UpdateFilteredCollectionItems();
 			return false;
 		}
@@ -725,7 +1251,7 @@ bool SCollectionView::CollectionNameChangeCommit( const TSharedPtr< FCollectionI
 		if ( !CollectionManagerModule.Get().CreateCollection(NewNameFinal, CollectionType) )
 		{
 			// Failed to add the collection, remove it from the list
-			CollectionItems.Remove(CollectionItem);
+			AvailableCollections.Remove(FCollectionNameType(CollectionItem->CollectionName, CollectionItem->CollectionType));
 			UpdateFilteredCollectionItems();
 
 			OutWarningMessage = FText::Format( LOCTEXT("CreateCollectionFailed", "Failed to create the collection. {0}"), CollectionManagerModule.Get().GetLastError());
@@ -755,33 +1281,23 @@ bool SCollectionView::CollectionNameChangeCommit( const TSharedPtr< FCollectionI
 		}
 	}
 
-	// Make sure we are sorted
-	CollectionItems.Sort( FCollectionItem::FCompareFCollectionItemByName() );
-
 	// At this point CollectionItem is no longer a member of the CollectionItems list (as the list is repopulated by
 	// UpdateCollectionItems, which is called by a broadcast from CollectionManagerModule::RenameCollection, above).
 	// So search again for the item by name and type.
-
-	auto FindCollectionItemPredicate = [CollectionType, &NewNameFinal] (const TSharedPtr<FCollectionItem>& Item)
-	{
-		return ( Item->CollectionType == CollectionType &&
-				 Item->CollectionName == NewNameFinal );
-	};
-
-	auto NewCollectionItemPtr = CollectionItems.FindByPredicate( FindCollectionItemPredicate );
+	auto NewCollectionItemPtr = AvailableCollections.Find( FCollectionNameType(NewNameFinal, CollectionType) );
 
 	// Reselect the path to notify that the selection has changed
 	{
 		FScopedPreventSelectionChangedDelegate DelegatePrevention( SharedThis(this) );
-		CollectionListPtr->ClearSelection();
+		CollectionTreePtr->ClearSelection();
 	}
 
 	// Set the selection
 	if (NewCollectionItemPtr)
 	{
 		const auto& NewCollectionItem = *NewCollectionItemPtr;
-		CollectionListPtr->RequestScrollIntoView(NewCollectionItem);
-		CollectionListPtr->SetItemSelection(NewCollectionItem, true);
+		CollectionTreePtr->RequestScrollIntoView(NewCollectionItem);
+		CollectionTreePtr->SetItemSelection(NewCollectionItem, true);
 	}
 
 	return true;
diff --git a/Engine/Source/Editor/ContentBrowser/Private/SCollectionView.h b/Engine/Source/Editor/ContentBrowser/Private/SCollectionView.h
--- a/Engine/Source/Editor/ContentBrowser/Private/SCollectionView.h
+++ b/Engine/Source/Editor/ContentBrowser/Private/SCollectionView.h
@@ -22,6 +22,7 @@ public:
 		, _AllowRightClickMenu(true)
 		, _AllowCollapsing(true)
 		, _AllowContextMenu(true)
+		, _AllowCollectionDrag(false)
 		, _AllowQuickAssetManagement(false)
 		{}
 
@@ -34,6 +35,9 @@ public:
 		SLATE_ARGUMENT( bool, AllowCollapsing )
 		SLATE_ARGUMENT( bool, AllowContextMenu )
 
+		/** If true, the user will be able to drag collections from this view */
+		SLATE_ARGUMENT( bool, AllowCollectionDrag )
+
 		/** If true, check boxes that let you quickly add/remove the current selection from a collection will be displayed */
 		SLATE_ARGUMENT( bool, AllowQuickAssetManagement )
 
@@ -43,7 +47,10 @@ public:
 	void Construct( const FArguments& InArgs );
 
 	/** Selects the specified collections */
-	void SetSelectedCollections(const TArray<FCollectionNameType>& CollectionsToSelect);
+	void SetSelectedCollections(const TArray<FCollectionNameType>& CollectionsToSelect, const bool bEnsureVisible = true);
+
+	/** Expands the specified collections */
+	void SetExpandedCollections(const TArray<FCollectionNameType>& CollectionsToExpand);
 
 	/** Clears selection of all collections */
 	void ClearSelection();
@@ -64,6 +71,10 @@ public:
 	void LoadSettings(const FString& IniFilename, const FString& IniSection, const FString& SettingsString);
 
 	virtual FReply OnKeyDown( const FGeometry& MyGeometry, const FKeyEvent& InKeyEvent ) override;
+	virtual void OnDragEnter( const FGeometry& MyGeometry, const FDragDropEvent& DragDropEvent ) override;
+	virtual void OnDragLeave( const FDragDropEvent& DragDropEvent ) override;
+	virtual FReply OnDragOver( const FGeometry& MyGeometry, const FDragDropEvent& DragDropEvent ) override;
+	virtual FReply OnDrop( const FGeometry& MyGeometry, const FDragDropEvent& DragDropEvent ) override;
 
 private:
 
@@ -88,17 +99,41 @@ private:
 	/** Sets up an inline rename for the specified collection */
 	void RenameCollectionItem( const TSharedPtr<FCollectionItem>& ItemToRename );
 
-	/** Remove a collection item from the list */
-	void RemoveCollectionItems( const TArray<TSharedPtr<FCollectionItem>>& ItemsToRemove );
+	/** Delete the given collections */
+	void DeleteCollectionItems( const TArray<TSharedPtr<FCollectionItem>>& ItemsToDelete );
 
-	/** Returns the visibility of the collection list */
-	EVisibility GetCollectionListVisibility() const;
+	/** Returns the visibility of the collection tree */
+	EVisibility GetCollectionTreeVisibility() const;
 
-	/** Creates a list item for the collection list */
+	/** Get the border of the collection tree */
+	const FSlateBrush* GetCollectionViewDropTargetBorder() const;
+
+	/** Creates a list item for the collection tree */
 	TSharedRef<ITableRow> GenerateCollectionRow( TSharedPtr<FCollectionItem> CollectionItem, const TSharedRef<STableViewBase>& OwnerTable );
 
-	/** Makes the context menu for the collection list */
-	TSharedPtr<SWidget> MakeCollectionListContextMenu();
+	/** Get the tree view children for the given item */
+	void GetCollectionItemChildren( TSharedPtr<FCollectionItem> InParentItem, TArray< TSharedPtr<FCollectionItem> >& OutChildItems ) const;
+
+	/** Handle starting a drag and drop operation for the currently selected collections */
+	FReply OnCollectionDragDetected(const FGeometry& Geometry, const FPointerEvent& MouseEvent);
+
+	/** Validate a drag drop operation on our collection tree */
+	bool ValidateDragDropOnCollectionTree(const FGeometry& Geometry, const FDragDropEvent& DragDropEvent, bool& OutIsKnownDragOperation);
+
+	/** Handle dropping something on collection tree */
+	FReply HandleDragDropOnCollectionTree(const FGeometry& Geometry, const FDragDropEvent& DragDropEvent);
+
+	/** Validate a drag drop operation on one of our collection items */
+	bool ValidateDragDropOnCollectionItem(TSharedRef<FCollectionItem> CollectionItem, const FGeometry& Geometry, const FDragDropEvent& DragDropEvent, bool& OutIsKnownDragOperation);
+
+	/** Handle dropping something on one of our collections */
+	FReply HandleDragDropOnCollectionItem(TSharedRef<FCollectionItem> CollectionItem, const FGeometry& Geometry, const FDragDropEvent& DragDropEvent);
+
+	/** Recursively expand the parent items of this collection to ensure that it is visible */
+	void ExpandParentItems(const TSharedRef<FCollectionItem>& InCollectionItem);
+
+	/** Makes the context menu for the collection tree */
+	TSharedPtr<SWidget> MakeCollectionTreeContextMenu();
 
 	/** Whether the check box of the given collection item is currently enabled */
 	bool IsCollectionCheckBoxEnabled( TSharedPtr<FCollectionItem> CollectionItem ) const;
@@ -112,14 +147,11 @@ private:
 	/** Handler for collection list selection changes */
 	void CollectionSelectionChanged( TSharedPtr< FCollectionItem > CollectionItem, ESelectInfo::Type SelectInfo );
 
-	/** Handler for the user dropping assets on a collection */
-	void CollectionAssetsDropped(const TArray<FAssetData>& AssetList, const TSharedPtr<FCollectionItem>& CollectionItem, FText& OutMessage);
-
 	/** Handles focusing a collection item widget after it has been created with the intent to rename */
 	void CollectionItemScrolledIntoView( TSharedPtr<FCollectionItem> CollectionItem, const TSharedPtr<ITableRow>& Widget );
 
 	/** Checks whether the selected collection is not allowed to be renamed */
-	bool IsCollectionNotRenamable() const;
+	bool IsCollectionNameReadOnly() const;
 
 	/** Handler for when a name was given to a collection. Returns false if the rename or create failed and sets OutWarningMessage depicting what happened. */
 	bool CollectionNameChangeCommit( const TSharedPtr< FCollectionItem >& CollectionItem, const FString& NewName, bool bChangeConfirmed, FText& OutWarningMessage );
@@ -133,10 +165,13 @@ private:
 	/** Handles an on collection renamed event */
 	void HandleCollectionRenamed( const FCollectionNameType& OriginalCollection, const FCollectionNameType& NewCollection );
 
+	/** Handles an on collection reparented event */
+	void HandleCollectionReparented( const FCollectionNameType& Collection, const TOptional<FCollectionNameType>& OldParent, const TOptional<FCollectionNameType>& NewParent );
+
 	/** Handles an on collection destroyed event */
 	void HandleCollectionDestroyed( const FCollectionNameType& Collection );
 
-	/** Updates the collections shown in the list view */
+	/** Updates the collections shown in the tree view */
 	void UpdateCollectionItems();
 
 	/** Update the visible collections based on the active filter text */
@@ -173,14 +208,18 @@ private:
 	/** The collection list search box */
 	TSharedPtr< SSearchBox > SearchBoxPtr;
 
-	/** The collection list widget */
-	TSharedPtr< SListView< TSharedPtr<FCollectionItem>> > CollectionListPtr;
+	/** The collection tree widget */
+	TSharedPtr< STreeView< TSharedPtr<FCollectionItem> > > CollectionTreePtr;
 
-	/** The list of available collections */
-	TArray< TSharedPtr<FCollectionItem> > CollectionItems;
+	/** A map of collection keys to their associated collection items - this map contains all available collections, even those that aren't currently visible */
+	typedef TMap< FCollectionNameType, TSharedPtr<FCollectionItem> > FAvailableCollectionsMap;
+	FAvailableCollectionsMap AvailableCollections;
 
-	/** The list of visible collections based on the current filter */
-	TArray< TSharedPtr<FCollectionItem> > FilteredCollectionItems;
+	/** A set of collections that are currently visible, including parents that are only visible due to their children */
+	TSet< FCollectionNameType > VisibleCollections;
+
+	/** The list of root items to show in the collections tree - this will be filtered as required */
+	TArray< TSharedPtr<FCollectionItem> > VisibleRootCollectionItems;
 
 	/** The filter to apply to the available collections */
 	typedef TTextFilter<const FCollectionItem&> FCollectionItemTextFilter;
@@ -201,6 +240,12 @@ private:
 	/** If true, the user will be able to access the right click menu of a collection */
 	bool bAllowRightClickMenu;
 
+	/** If true, the user will be able to drag collections from this view */
+	bool bAllowCollectionDrag;
+
+	/** True when a drag is over this view with a valid operation for drop */
+	bool bDraggedOver;
+
 	/** If > 0, the selection changed delegate will not be called. Used to update the tree from an external source or in certain bulk operations. */
 	int32 PreventSelectionChangedDelegateCount;
 
@@ -209,4 +254,10 @@ private:
 
 	/** Handles the collection management for the currently selected assets (if available) */
 	TSharedPtr<FCollectionAssetManagement> QuickAssetManagement;
+
+	/**
+	 * This is set after this view has initiated a drag and drop for some collections in our tree.
+	 * We keep a weak pointer to this so we can tell if that drag and drop is still ongoing, and if so, what collections it affects 
+	 */
+	TWeakPtr<class FCollectionDragDropOp> CurrentCollectionDragDropOp;
 };
\ No newline at end of file
diff --git a/Engine/Source/Editor/ContentBrowser/Private/SContentBrowser.cpp b/Engine/Source/Editor/ContentBrowser/Private/SContentBrowser.cpp
--- a/Engine/Source/Editor/ContentBrowser/Private/SContentBrowser.cpp
+++ b/Engine/Source/Editor/ContentBrowser/Private/SContentBrowser.cpp
@@ -455,6 +455,7 @@ void SContentBrowser::Construct( const FArguments& InArgs, const FName& InInstan
 							SAssignNew(CollectionViewPtr, SCollectionView)
 							.OnCollectionSelected(this, &SContentBrowser::CollectionSelected)
 							.AddMetaData<FTagMetaData>(FTagMetaData(TEXT("ContentBrowserCollections")))
+							.AllowCollectionDrag(true)
 							.AllowQuickAssetManagement(true)
 						]
 					]
@@ -1087,12 +1088,28 @@ void SContentBrowser::SourcesChanged(const TArray<FString>& SelectedPaths, const
 
 void SContentBrowser::FolderEntered(const FString& FolderPath)
 {
-	// set the path view to the incoming path
-	TArray<FString> SelectedPaths;
-	SelectedPaths.Add(FolderPath);
-	PathViewPtr->SetSelectedPaths(SelectedPaths);
+	// Have we entered a sub-collection folder?
+	FName CollectionName;
+	ECollectionShareType::Type CollectionFolderShareType = ECollectionShareType::CST_All;
+	if (ContentBrowserUtils::IsCollectionPath(FolderPath, &CollectionName, &CollectionFolderShareType))
+	{
+		const FCollectionNameType SelectedCollection(CollectionName, CollectionFolderShareType);
 
-	PathSelected(FolderPath);
+		TArray<FCollectionNameType> Collections;
+		Collections.Add(SelectedCollection);
+		CollectionViewPtr->SetSelectedCollections(Collections);
+
+		CollectionSelected(SelectedCollection);
+	}
+	else
+	{
+		// set the path view to the incoming path
+		TArray<FString> SelectedPaths;
+		SelectedPaths.Add(FolderPath);
+		PathViewPtr->SetSelectedPaths(SelectedPaths);
+
+		PathSelected(FolderPath);
+	}
 }
 
 void SContentBrowser::PathSelected(const FString& FolderPath)
diff --git a/Engine/Source/Editor/ContentBrowser/Private/SourcesData.h b/Engine/Source/Editor/ContentBrowser/Private/SourcesData.h
--- a/Engine/Source/Editor/ContentBrowser/Private/SourcesData.h
+++ b/Engine/Source/Editor/ContentBrowser/Private/SourcesData.h
@@ -27,13 +27,14 @@ struct FSourcesData
 			// Collection manager module should already be loaded since it may cause a hitch on the first search
 			FCollectionManagerModule& CollectionManagerModule = FCollectionManagerModule::GetModule();
 
+			// Include objects from child collections if we're recursing
+			const ECollectionRecursionFlags::Flags CollectionRecursionMode = (Filter.bRecursivePaths) ? ECollectionRecursionFlags::SelfAndChildren : ECollectionRecursionFlags::Self;
+
 			for ( int32 CollectionIdx = 0; CollectionIdx < Collections.Num(); ++CollectionIdx )
 			{
 				// Find the collection
 				const FCollectionNameType& CollectionNameType = Collections[CollectionIdx];
-
-				CollectionManagerModule.Get().GetAssetsInCollection(CollectionNameType.Name, CollectionNameType.Type, ObjectPathsFromCollections);
-				CollectionManagerModule.Get().GetClassesInCollection(CollectionNameType.Name, CollectionNameType.Type, ObjectPathsFromCollections);
+				CollectionManagerModule.Get().GetObjectsInCollection(CollectionNameType.Name, CollectionNameType.Type, ObjectPathsFromCollections, CollectionRecursionMode);
 			}
 		}
 		Filter.ObjectPaths = ObjectPathsFromCollections;
diff --git a/Engine/Source/Editor/ContentBrowser/Private/SourcesViewWidgets.cpp b/Engine/Source/Editor/ContentBrowser/Private/SourcesViewWidgets.cpp
--- a/Engine/Source/Editor/ContentBrowser/Private/SourcesViewWidgets.cpp
+++ b/Engine/Source/Editor/ContentBrowser/Private/SourcesViewWidgets.cpp
@@ -7,6 +7,7 @@
 
 #include "DragAndDrop/AssetDragDropOp.h"
 #include "DragAndDrop/AssetPathDragDropOp.h"
+#include "DragAndDrop/CollectionDragDropOp.h"
 #include "DragDropHandler.h"
 #include "ContentBrowserUtils.h"
 #include "CollectionViewUtils.h"
@@ -97,20 +98,17 @@ SAssetTreeItem::~SAssetTreeItem()
 	}
 }
 
-bool SAssetTreeItem::ValidateDragDrop( const FGeometry& MyGeometry, const FDragDropEvent& DragDropEvent ) const
+bool SAssetTreeItem::ValidateDragDrop( const FGeometry& MyGeometry, const FDragDropEvent& DragDropEvent, bool& OutIsKnownDragOperation ) const
 {
+	OutIsKnownDragOperation = false;
+
 	TSharedPtr<FTreeItem> TreeItemPinned = TreeItem.Pin();
-	return TreeItemPinned.IsValid() && DragDropHandler::ValidateDragDropOnAssetFolder(MyGeometry, DragDropEvent, TreeItemPinned->FolderPath);
+	return TreeItemPinned.IsValid() && DragDropHandler::ValidateDragDropOnAssetFolder(MyGeometry, DragDropEvent, TreeItemPinned->FolderPath, OutIsKnownDragOperation);
 }
 
 void SAssetTreeItem::OnDragEnter( const FGeometry& MyGeometry, const FDragDropEvent& DragDropEvent )
 {
-	bDraggedOver = false;
-
-	if (ValidateDragDrop(MyGeometry, DragDropEvent))
-	{
-		bDraggedOver = true;
-	}
+	ValidateDragDrop(MyGeometry, DragDropEvent, bDraggedOver); // updates bDraggedOver
 }
 
 void SAssetTreeItem::OnDragLeave( const FDragDropEvent& DragDropEvent )
@@ -132,23 +130,16 @@ void SAssetTreeItem::OnDragLeave( const FDragDropEvent& DragDropEvent )
 
 FReply SAssetTreeItem::OnDragOver( const FGeometry& MyGeometry, const FDragDropEvent& DragDropEvent )
 {
-	bDraggedOver = false;
-
-	if (ValidateDragDrop(MyGeometry, DragDropEvent))
-	{
-		bDraggedOver = true;
-		return FReply::Handled();
-	}
-
-	return FReply::Unhandled();
+	ValidateDragDrop(MyGeometry, DragDropEvent, bDraggedOver); // updates bDraggedOver
+	return (bDraggedOver) ? FReply::Handled() : FReply::Unhandled();
 }
 
 FReply SAssetTreeItem::OnDrop( const FGeometry& MyGeometry, const FDragDropEvent& DragDropEvent )
 {
-	bDraggedOver = false;
-
-	if (ValidateDragDrop(MyGeometry, DragDropEvent))
+	if (ValidateDragDrop(MyGeometry, DragDropEvent, bDraggedOver)) // updates bDraggedOver
 	{
+		bDraggedOver = false;
+
 		TSharedPtr<FDragDropOperation> Operation = DragDropEvent.GetOperation();
 		if (!Operation.IsValid())
 		{
@@ -175,6 +166,13 @@ FReply SAssetTreeItem::OnDrop( const FGeometry& MyGeometry, const FDragDropEvent
 		}
 	}
 
+	if (bDraggedOver)
+	{
+		// We were able to handle this operation, but could not due to another error - still report this drop as handled so it doesn't fall through to other widgets
+		bDraggedOver = false;
+		return FReply::Handled();
+	}
+
 	return FReply::Unhandled();
 }
 
@@ -310,23 +308,24 @@ const FSlateBrush* SAssetTreeItem::GetBorderImage() const
 
 
 //////////////////////////
-// SCollectionListItem
+// SCollectionTreeItem
 //////////////////////////
 
-void SCollectionListItem::Construct( const FArguments& InArgs )
+void SCollectionTreeItem::Construct( const FArguments& InArgs )
 {
 	ParentWidget = InArgs._ParentWidget;
 	CollectionItem = InArgs._CollectionItem;
 	OnBeginNameChange = InArgs._OnBeginNameChange;
 	OnNameChangeCommit = InArgs._OnNameChangeCommit;
 	OnVerifyRenameCommit = InArgs._OnVerifyRenameCommit;
-	OnAssetsDragDropped = InArgs._OnAssetsDragDropped;
+	OnValidateDragDrop = InArgs._OnValidateDragDrop;
+	OnHandleDragDrop = InArgs._OnHandleDragDrop;
 	bDraggedOver = false;
 
 	ChildSlot
 	[
 		SNew(SBorder)
-		.BorderImage(this, &SCollectionListItem::GetBorderImage)
+		.BorderImage(this, &SCollectionTreeItem::GetBorderImage)
 		.Padding(0)
 		[
 			SNew(SHorizontalBox)
@@ -352,7 +351,7 @@ void SCollectionListItem::Construct( const FArguments& InArgs )
 				// Type Icon
 				SNew(SImage)
 				.Image( FEditorStyle::GetBrush( ECollectionShareType::GetIconStyleName(InArgs._CollectionItem->CollectionType) ) )
-				.ColorAndOpacity( this, &SCollectionListItem::GetCollectionColor )
+				.ColorAndOpacity( this, &SCollectionTreeItem::GetCollectionColor )
 			]
 
 			+SHorizontalBox::Slot()
@@ -360,12 +359,12 @@ void SCollectionListItem::Construct( const FArguments& InArgs )
 			.VAlign(VAlign_Center)
 			[
 				SAssignNew(InlineRenameWidget, SInlineEditableTextBlock)
-				.Text( this, &SCollectionListItem::GetNameText )
+				.Text( this, &SCollectionTreeItem::GetNameText )
 				.HighlightText( InArgs._HighlightText )
 				.Font( FEditorStyle::GetFontStyle("ContentBrowser.SourceListItemFont") )
-				.OnBeginTextEdit(this, &SCollectionListItem::HandleBeginNameChange)
-				.OnTextCommitted(this, &SCollectionListItem::HandleNameCommitted)
-				.OnVerifyTextChanged(this, &SCollectionListItem::HandleVerifyNameChanged)
+				.OnBeginTextEdit(this, &SCollectionTreeItem::HandleBeginNameChange)
+				.OnTextCommitted(this, &SCollectionTreeItem::HandleNameCommitted)
+				.OnVerifyTextChanged(this, &SCollectionTreeItem::HandleVerifyNameChanged)
 				.IsSelected( InArgs._IsSelected )
 				.IsReadOnly( InArgs._IsReadOnly )
 			]
@@ -379,7 +378,7 @@ void SCollectionListItem::Construct( const FArguments& InArgs )
 	}
 }
 
-SCollectionListItem::~SCollectionListItem()
+SCollectionTreeItem::~SCollectionTreeItem()
 {
 	if(InlineRenameWidget.IsValid())
 	{
@@ -387,67 +386,72 @@ SCollectionListItem::~SCollectionListItem()
 	}
 }
 
-void SCollectionListItem::Tick( const FGeometry& AllottedGeometry, const double InCurrentTime, const float InDeltaTime )
+void SCollectionTreeItem::Tick( const FGeometry& AllottedGeometry, const double InCurrentTime, const float InDeltaTime )
 {
 	// Cache this widget's geometry so it can pop up warnings over itself
 	CachedGeometry = AllottedGeometry;
 }
 
-void SCollectionListItem::OnDragEnter( const FGeometry& MyGeometry, const FDragDropEvent& DragDropEvent )
+bool SCollectionTreeItem::ValidateDragDrop( const FGeometry& MyGeometry, const FDragDropEvent& DragDropEvent, bool& OutIsKnownDragOperation ) const
 {
-	TSharedPtr<FAssetDragDropOp> DragDropOp = DragDropEvent.GetOperationAs<FAssetDragDropOp>();
-	if ( DragDropOp.IsValid() )
-	{
-		bool bCanDrop = DragDropOp->AssetData.Num() > 0;
+	OutIsKnownDragOperation = false;
 
-		if (bCanDrop)
-		{
-			bDraggedOver = true;
-		}
+	TSharedPtr<FCollectionItem> CollectionItemPtr = CollectionItem.Pin();
+	if (OnValidateDragDrop.IsBound() && CollectionItemPtr.IsValid())
+	{
+		return OnValidateDragDrop.Execute( CollectionItemPtr.ToSharedRef(), MyGeometry, DragDropEvent, OutIsKnownDragOperation );
 	}
+
+	return false;
 }
 
-void SCollectionListItem::OnDragLeave( const FDragDropEvent& DragDropEvent )
+void SCollectionTreeItem::OnDragEnter( const FGeometry& MyGeometry, const FDragDropEvent& DragDropEvent )
 {
-	bDraggedOver = false;
+	ValidateDragDrop(MyGeometry, DragDropEvent, bDraggedOver); // updates bDraggedOver
 }
 
-FReply SCollectionListItem::OnDragOver( const FGeometry& MyGeometry, const FDragDropEvent& DragDropEvent )
+void SCollectionTreeItem::OnDragLeave( const FDragDropEvent& DragDropEvent )
 {
-	TSharedPtr<FAssetDragDropOp> DragDropOp = DragDropEvent.GetOperationAs<FAssetDragDropOp>();
-	if (DragDropOp.IsValid())
+	TSharedPtr<FDragDropOperation> Operation = DragDropEvent.GetOperation();
+	if (Operation.IsValid())
 	{
-		bool bCanDrop = DragDropOp->AssetData.Num() > 0;
+		Operation->SetCursorOverride(TOptional<EMouseCursor::Type>());
 
-		if (bCanDrop)
+		if (Operation->IsOfType<FCollectionDragDropOp>() || Operation->IsOfType<FAssetDragDropOp>())
 		{
-			return FReply::Handled();
+			TSharedPtr<FDecoratedDragDropOp> DragDropOp = StaticCastSharedPtr<FDecoratedDragDropOp>(Operation);
+			DragDropOp->ResetToDefaultToolTip();
 		}
 	}
 
-	return FReply::Unhandled();
+	bDraggedOver = false;
 }
 
-FReply SCollectionListItem::OnDrop( const FGeometry& MyGeometry, const FDragDropEvent& DragDropEvent )
+FReply SCollectionTreeItem::OnDragOver( const FGeometry& MyGeometry, const FDragDropEvent& DragDropEvent )
 {
-	bDraggedOver = false;
+	ValidateDragDrop(MyGeometry, DragDropEvent, bDraggedOver); // updates bDraggedOver
+	return (bDraggedOver) ? FReply::Handled() : FReply::Unhandled();
+}
 
-	TSharedPtr<FAssetDragDropOp> DragDropOp = DragDropEvent.GetOperationAs<FAssetDragDropOp>();
-	if (DragDropOp.IsValid() && ensure(CollectionItem.IsValid()))
+FReply SCollectionTreeItem::OnDrop( const FGeometry& MyGeometry, const FDragDropEvent& DragDropEvent )
+{
+	if (ValidateDragDrop(MyGeometry, DragDropEvent, bDraggedOver) && OnHandleDragDrop.IsBound()) // updates bDraggedOver
 	{
-		FText Message;
-		OnAssetsDragDropped.ExecuteIfBound(DragDropOp->AssetData, CollectionItem.Pin(), Message);
-
-		// Added items to the collection or failed. Either way, display the message.
-		ContentBrowserUtils::DisplayMessage(Message, CachedGeometry.GetClippingRect(), SharedThis(this));
+		bDraggedOver = false;
+		return OnHandleDragDrop.Execute( CollectionItem.Pin().ToSharedRef(), MyGeometry, DragDropEvent );
+	}
 
+	if (bDraggedOver)
+	{
+		// We were able to handle this operation, but could not due to another error - still report this drop as handled so it doesn't fall through to other widgets
+		bDraggedOver = false;
 		return FReply::Handled();
 	}
 
 	return FReply::Unhandled();
 }
 
-void SCollectionListItem::HandleBeginNameChange( const FText& OldText )
+void SCollectionTreeItem::HandleBeginNameChange( const FText& OldText )
 {
 	TSharedPtr<FCollectionItem> CollectionItemPtr = CollectionItem.Pin();
 
@@ -461,7 +465,7 @@ void SCollectionListItem::HandleBeginNameChange( const FText& OldText )
 	}
 }
 
-void SCollectionListItem::HandleNameCommitted( const FText& NewText, ETextCommit::Type CommitInfo )
+void SCollectionTreeItem::HandleNameCommitted( const FText& NewText, ETextCommit::Type CommitInfo )
 {
 	TSharedPtr<FCollectionItem> CollectionItemPtr = CollectionItem.Pin();
 
@@ -485,7 +489,7 @@ void SCollectionListItem::HandleNameCommitted( const FText& NewText, ETextCommit
 	}
 }
 
-bool SCollectionListItem::HandleVerifyNameChanged( const FText& NewText, FText& OutErrorMessage )
+bool SCollectionTreeItem::HandleVerifyNameChanged( const FText& NewText, FText& OutErrorMessage )
 {
 	TSharedPtr<FCollectionItem> CollectionItemPtr = CollectionItem.Pin();
 
@@ -497,7 +501,7 @@ bool SCollectionListItem::HandleVerifyNameChanged( const FText& NewText, FText&
 	return true;
 }
 
-FText SCollectionListItem::GetNameText() const
+FText SCollectionTreeItem::GetNameText() const
 {
 	TSharedPtr<FCollectionItem> CollectionItemPtr = CollectionItem.Pin();
 
@@ -511,7 +515,7 @@ FText SCollectionListItem::GetNameText() const
 	}
 }
 
-FSlateColor SCollectionListItem::GetCollectionColor() const
+FSlateColor SCollectionTreeItem::GetCollectionColor() const
 {
 	TSharedPtr<FCollectionItem> CollectionItemPtr = CollectionItem.Pin();
 
@@ -527,7 +531,7 @@ FSlateColor SCollectionListItem::GetCollectionColor() const
 	return CollectionViewUtils::GetDefaultColor();
 }
 
-const FSlateBrush* SCollectionListItem::GetBorderImage() const
+const FSlateBrush* SCollectionTreeItem::GetBorderImage() const
 {
 	return bDraggedOver ? FEditorStyle::GetBrush("Menu.Background") : FEditorStyle::GetBrush("NoBorder");
 }
diff --git a/Engine/Source/Editor/ContentBrowser/Private/SourcesViewWidgets.h b/Engine/Source/Editor/ContentBrowser/Private/SourcesViewWidgets.h
--- a/Engine/Source/Editor/ContentBrowser/Private/SourcesViewWidgets.h
+++ b/Engine/Source/Editor/ContentBrowser/Private/SourcesViewWidgets.h
@@ -58,7 +58,7 @@ public:
 
 private:
 	/** Used by OnDragEnter, OnDragOver, and OnDrop to check and update the validity of the drag operation */
-	bool ValidateDragDrop( const FGeometry& MyGeometry, const FDragDropEvent& DragDropEvent ) const;
+	bool ValidateDragDrop( const FGeometry& MyGeometry, const FDragDropEvent& DragDropEvent, bool& OutIsKnownDragOperation ) const;
 
 	/** Handles verifying name changes */
 	bool VerifyNameChanged(const FText& InName, FText& OutError) const;
@@ -135,7 +135,7 @@ private:
 	const FSlateBrush* FolderClosedCodeBrush;
 	const FSlateBrush* FolderDeveloperBrush;
 
-	/** True when a drag is over this item with a valid operation for drop */
+	/** True when a drag is over this item with a drag operation that we know how to handle. The operation itself may not be valid to drop. */
 	bool bDraggedOver;
 
 	/** What type of stuff does this folder hold */
@@ -148,23 +148,26 @@ private:
 	FDelegateHandle EnterEditingModeDelegateHandle;
 };
 
-/** A single item in the collection list. */
-class SCollectionListItem : public SCompoundWidget
+/** A single item in the collection tree. */
+class SCollectionTreeItem : public SCompoundWidget
 {
 public:
-	/** Delegate for when a collection is renamed. If returning false, OutWarningMessage will be displayed over the collection. */
+	/** Delegates for when a collection is renamed. If returning false, OutWarningMessage will be displayed over the collection. */
 	DECLARE_DELEGATE_OneParam( FOnBeginNameChange, const TSharedPtr<FCollectionItem>& /*Item*/);
 	DECLARE_DELEGATE_RetVal_FourParams( bool, FOnNameChangeCommit, const TSharedPtr<FCollectionItem>& /*Item*/, const FString& /*NewName*/, bool /*bChangeConfirmed*/, FText& /*OutWarningMessage*/);
 	DECLARE_DELEGATE_RetVal_FourParams( bool, FOnVerifyRenameCommit, const TSharedPtr<FCollectionItem>& /*Item*/, const FString& /*NewName*/, const FSlateRect& /*MessageAnchor*/, FText& /*OutErrorMessage*/)
-	DECLARE_DELEGATE_ThreeParams( FOnAssetsDragDropped, const TArray<FAssetData>& /*AssetList*/, const TSharedPtr<FCollectionItem>& /*CollectionItem*/, FText& /*OutMessage*/);
+	
+	/** Delegates for when a collection item has something dropped into it */
+	DECLARE_DELEGATE_RetVal_FourParams( bool, FOnValidateDragDrop, TSharedRef<FCollectionItem> /*CollectionItem*/, const FGeometry& /*Geometry*/, const FDragDropEvent& /*DragDropEvent*/, bool& /*OutIsKnownDragOperation*/ );
+	DECLARE_DELEGATE_RetVal_ThreeParams( FReply, FOnHandleDragDrop, TSharedRef<FCollectionItem> /*CollectionItem*/, const FGeometry& /*Geometry*/, const FDragDropEvent& /*DragDropEvent*/ );
 
-	SLATE_BEGIN_ARGS( SCollectionListItem )
+	SLATE_BEGIN_ARGS( SCollectionTreeItem )
 		: _CollectionItem( TSharedPtr<FCollectionItem>() )
 		, _ParentWidget()
 	{}
 
-	/** Data for the collection this item represents */
-	SLATE_ARGUMENT( TSharedPtr<FCollectionItem>, CollectionItem )
+		/** Data for the collection this item represents */
+		SLATE_ARGUMENT( TSharedPtr<FCollectionItem>, CollectionItem )
 
 		/** The parent widget */
 		SLATE_ARGUMENT( TSharedPtr<SWidget>, ParentWidget )
@@ -178,8 +181,11 @@ public:
 		/** Delegate for when a collection name has been entered for an item to verify the name before commit */
 		SLATE_EVENT( FOnVerifyRenameCommit, OnVerifyRenameCommit )
 
-		/** Delegate for when assets are dropped on this folder */
-		SLATE_EVENT( FOnAssetsDragDropped, OnAssetsDragDropped )
+		/** Delegate to validate a drag drop operation on this collection item */
+		SLATE_EVENT( FOnValidateDragDrop, OnValidateDragDrop )
+
+		/** Delegate to handle a drag drop operation on this collection item */
+		SLATE_EVENT( FOnHandleDragDrop, OnHandleDragDrop )
 
 		/** Callback to check if the widget is selected, should only be hooked up if parent widget is handling selection or focus. */
 		SLATE_EVENT( FIsSelected, IsSelected )
@@ -204,7 +210,7 @@ public:
 	/** Constructs this widget with InArgs */
 	void Construct( const FArguments& InArgs );
 
-	~SCollectionListItem();
+	~SCollectionTreeItem();
 
 	virtual void Tick( const FGeometry& AllottedGeometry, const double InCurrentTime, const float InDeltaTime ) override;
 	virtual void OnDragEnter( const FGeometry& MyGeometry, const FDragDropEvent& DragDropEvent ) override;
@@ -213,6 +219,9 @@ public:
 	virtual FReply OnDrop( const FGeometry& MyGeometry, const FDragDropEvent& DragDropEvent ) override;
 
 private:
+	/** Used by OnDragEnter, OnDragOver, and OnDrop to check and update the validity of the drag operation */
+	bool ValidateDragDrop( const FGeometry& MyGeometry, const FDragDropEvent& DragDropEvent, bool& OutIsKnownDragOperation ) const;
+
 	/** Handles beginning a name change */
 	void HandleBeginNameChange( const FText& OldText );
 
@@ -241,11 +250,14 @@ private:
 	/** The name of the asset as an editable text box */
 	TSharedPtr<SEditableTextBox> EditableName;
 
-	/** True when a drag is over this item with a valid operation for drop */
+	/** True when a drag is over this item with a drag operation that we know how to handle. The operation itself may not be valid to drop. */
 	bool bDraggedOver;
 
-	/** Delegate for when a list of assets */
-	FOnAssetsDragDropped OnAssetsDragDropped;
+	/** Delegate to validate a drag drop operation on this collection item */
+	FOnValidateDragDrop OnValidateDragDrop;
+
+	/** Delegate to handle a drag drop operation on this collection item */
+	FOnHandleDragDrop OnHandleDragDrop;
 
 	/** The geometry as of the last frame. Used to open warning messages over the item */
 	FGeometry CachedGeometry;
diff --git a/Engine/Source/Editor/EditorStyle/Private/SlateEditorStyle.cpp b/Engine/Source/Editor/EditorStyle/Private/SlateEditorStyle.cpp
--- a/Engine/Source/Editor/EditorStyle/Private/SlateEditorStyle.cpp
+++ b/Engine/Source/Editor/EditorStyle/Private/SlateEditorStyle.cpp
@@ -5571,11 +5571,18 @@ void FSlateEditorStyle::FStyle::SetupContentBrowserStyle()
 		Set( "ContentBrowser.SCC_NotInDepot_Small", new IMAGE_BRUSH("ContentBrowser/SCC_ContentMissing", Icon16x16) );
 		Set( "ContentBrowser.ContentDirty", new IMAGE_BRUSH( "ContentBrowser/ContentDirty", Icon16x16) );
 		Set( "ContentBrowser.AssetDragDropTooltipBackground", new BOX_BRUSH( "Old/Menu_Background", FMargin(8.0f/64.0f) ) );
+		Set( "ContentBrowser.CollectionTreeDragDropBorder", new BOX_BRUSH( "Old/Window/ViewportDebugBorder", 0.8f ) );
 		Set( "ContentBrowser.PopupMessageIcon", new IMAGE_BRUSH( "Icons/alert", Icon32x32) );
 		Set( "ContentBrowser.NewFolderIcon", new IMAGE_BRUSH("Icons/icon_AddFolder_16x", Icon16x16 ) );
 		Set( "ContentBrowser.Local", new IMAGE_BRUSH( "ContentBrowser/Content_Local_12x", Icon12x12 ) );
+		Set( "ContentBrowser.Local.Small", new IMAGE_BRUSH( "ContentBrowser/Content_Local_16x", Icon16x16 ) );
+		Set( "ContentBrowser.Local.Large", new IMAGE_BRUSH( "ContentBrowser/Content_Local_64x", Icon64x64 ) );
 		Set( "ContentBrowser.Shared", new IMAGE_BRUSH( "ContentBrowser/Content_Shared_12x", Icon12x12 ) );
+		Set( "ContentBrowser.Shared.Small", new IMAGE_BRUSH( "ContentBrowser/Content_Shared_16x", Icon16x16 ) );
+		Set( "ContentBrowser.Shared.Large", new IMAGE_BRUSH( "ContentBrowser/Content_Shared_64x", Icon64x64 ) );
 		Set( "ContentBrowser.Private", new IMAGE_BRUSH( "ContentBrowser/Content_Private_12x", Icon12x12 ) );
+		Set( "ContentBrowser.Private.Small", new IMAGE_BRUSH( "ContentBrowser/Content_Private_16x", Icon16x16 ) );
+		Set( "ContentBrowser.Private.Large", new IMAGE_BRUSH( "ContentBrowser/Content_Private_64x", Icon64x64 ) );
 	}
 #endif // #if WITH_EDITOR
 }
diff --git a/Engine/Source/Editor/UnrealEd/Public/DragAndDrop/CollectionDragDropOp.h b/Engine/Source/Editor/UnrealEd/Public/DragAndDrop/CollectionDragDropOp.h
new file mode 100644
--- /dev/null
+++ b/Engine/Source/Editor/UnrealEd/Public/DragAndDrop/CollectionDragDropOp.h
@@ -0,0 +1,87 @@
+// Copyright 1998-2015 Epic Games, Inc. All Rights Reserved.
+
+#pragma once
+
+#include "DecoratedDragDropOp.h"
+#include "CollectionManagerTypes.h"
+
+class FCollectionDragDropOp : public FDecoratedDragDropOp
+{
+public:
+	DRAG_DROP_OPERATOR_TYPE(FCollectionDragDropOp, FDecoratedDragDropOp)
+
+	/** Data for the collections this item represents */
+	TArray<FCollectionNameType> Collections;
+
+	static TSharedRef<FCollectionDragDropOp> New(TArray<FCollectionNameType> InCollections)
+	{
+		TSharedRef<FCollectionDragDropOp> Operation = MakeShareable(new FCollectionDragDropOp);
+		
+		Operation->MouseCursor = EMouseCursor::GrabHandClosed;
+		Operation->Collections = MoveTemp(InCollections);
+		Operation->Construct();
+
+		return Operation;
+	}
+	
+public:
+	FText GetDecoratorText() const
+	{
+		if (CurrentHoverText.IsEmpty() && Collections.Num() > 0)
+		{
+			return (Collections.Num() == 1)
+				? FText::FromName(Collections[0].Name)
+				: FText::Format(NSLOCTEXT("ContentBrowser", "CollectionDragDropDescription", "{0} and {1} other(s)"), FText::FromName(Collections[0].Name), FText::AsNumber(Collections.Num() - 1));
+		}
+		
+		return CurrentHoverText;
+	}
+
+	virtual TSharedPtr<SWidget> GetDefaultDecorator() const override
+	{
+		const FSlateBrush* const CollectionIcon = (Collections.Num() > 0) ? FEditorStyle::GetBrush(ECollectionShareType::GetIconStyleName(Collections[0].Type)) : nullptr;
+
+		return 
+			SNew(SBorder)
+			.BorderImage(FEditorStyle::GetBrush("ContentBrowser.AssetDragDropTooltipBackground"))
+			.Content()
+			[
+				SNew(SHorizontalBox)
+
+				// Left slot is collection icon
+				+ SHorizontalBox::Slot()
+				.AutoWidth()
+				.VAlign(VAlign_Center)
+				[
+					SNew(SImage)
+					.Image(CollectionIcon)
+				]
+
+				// Right slot is for tooltip
+				+ SHorizontalBox::Slot()
+				.AutoWidth()
+				.VAlign(VAlign_Center)
+				[
+					SNew(SHorizontalBox)
+
+					+ SHorizontalBox::Slot()
+					.AutoWidth()
+					.Padding(3.0f)
+					.VAlign(VAlign_Center)
+					[
+						SNew(SImage) 
+						.Image(this, &FCollectionDragDropOp::GetIcon)
+					]
+
+					+SHorizontalBox::Slot()
+					.AutoWidth()
+					.Padding(0,0,3,0)
+					.VAlign(VAlign_Center)
+					[
+						SNew(STextBlock) 
+						.Text(this, &FCollectionDragDropOp::GetDecoratorText)
+					]
+				]
+			];
+	}
+};
