Optimized the Content Browser text filter, and added support for matching collection names in simple string searches

UETOOL-332 - Collections 2.0

Converted FFrontendFilter_Text to use FTextFilterExpressionEvaluator directly (rather than via TTextFilter), as this allows it better control over memory re-use (in this case, re-using an array of collection names, as well as some strings).

Added FTextFilterString, which is used internally by the compiled text filter, and also by the comparison functions. This stores the string given to it as uppercase, which allows us to quickly perform case-insensitive string comparisons since we can actually do normal case-sensitive comparisons instead as everything is uppercase.

diff --git a/Engine/Source/Editor/ContentBrowser/Private/FrontendFilters.cpp b/Engine/Source/Editor/ContentBrowser/Private/FrontendFilters.cpp
--- a/Engine/Source/Editor/ContentBrowser/Private/FrontendFilters.cpp
+++ b/Engine/Source/Editor/ContentBrowser/Private/FrontendFilters.cpp
@@ -12,209 +12,277 @@
 // FFrontendFilter_Text
 /////////////////////////////////////////
 
-void AssetDataToNameString(FAssetFilterType Asset, OUT TArray< FString >& Array)
-{
-	const FString FullPath = Asset.PackageName.ToString();
-	const TCHAR* Ptr = FullPath.GetCharArray().GetData();
-	if (Ptr == nullptr)
+class FFrontendFilter_TextFilterExpressionContext : public ITextFilterExpressionContext
+{
+public:
+	typedef TRemoveReference<FAssetFilterType>::Type* FAssetFilterTypePtr;
+
+	FFrontendFilter_TextFilterExpressionContext()
+		: AssetPtr(nullptr)
+		, bIncludeClassName(true)
+		, NameKeyName("Name")
+		, PathKeyName("Path")
+		, ClassKeyName("Class")
+		, TypeKeyName("Type")
+		, CollectionKeyName("Collection")
+		, TagKeyName("Tag")
 	{
-		return;
 	}
 
-	// Enter each piece of the path name, apart from the first, as a potential match
-	bool bIsFirst = true;
-	while (const TCHAR* Delimiter = FCString::Strchr(Ptr, '/'))
+	void SetAsset(FAssetFilterTypePtr InAsset)
 	{
-		const int32 Length = Delimiter - Ptr;
+		AssetPtr = InAsset;
+		
+		AssetPtr->PackageName.AppendString(AssetFullPath);
+		AssetPtr->GetExportTextName(AssetExportTextName);
 
-		if (Length > 0)
+		if (FCollectionManagerModule::IsModuleAvailable())
 		{
-			if (bIsFirst)
-			{
-				bIsFirst = false;
-			}
-			else
-			{
-				Array.Emplace(Length, Ptr);
-			}
+			FCollectionManagerModule& CollectionManagerModule = FCollectionManagerModule::GetModule();
+			CollectionManagerModule.Get().GetCollectionsContainingObject(AssetPtr->ObjectPath, ECollectionShareType::CST_All, AssetCollectionNames, ECollectionRecursionFlags::SelfAndChildren);
 		}
-
-		Ptr += (Length + 1);
 	}
 
-	if (*Ptr != 0)
+	void ClearAsset()
 	{
-		Array.Emplace(Ptr);
+		AssetPtr = nullptr;
+		AssetFullPath.Reset();
+		AssetExportTextName.Reset();
+		AssetCollectionNames.Reset();
 	}
-}
-
-void AssetDataToClassAndNameStrings(FAssetFilterType Asset, OUT TArray< FString >& Array)
-{
-	Array.Add(Asset.AssetClass.ToString());
-	Array.Add(Asset.GetExportTextName()); // only include this if we're searching the class name too, as the exported text contains the type in the string
-	AssetDataToNameString(Asset, Array);
-}
 
-bool AssetDataTestComplexExpression(FAssetFilterType Asset, const FName& InKey, const FString& InValue, ETextFilterComparisonOperation InComparisonOperation, ETextFilterTextComparisonMode InTextComparisonMode)
-{
-	static const FName NameKeyName = "Name";
-	static const FName PathKeyName = "Path";
-	static const FName ClassKeyName = "Class";
-	static const FName TypeKeyName = "Type";
-	static const FName CollectionKeyName = "Collection";
-	static const FName TagKeyName = "Tag";
-
-	// Special case for the asset name, as this isn't contained within the asset registry meta-data
-	if (InKey == NameKeyName)
+	void SetIncludeClassName(const bool InIncludeClassName)
 	{
-		// Names can only work with Equal or NotEqual type tests
-		if (InComparisonOperation != ETextFilterComparisonOperation::Equal && InComparisonOperation != ETextFilterComparisonOperation::NotEqual)
-		{
-			return false;
-		}
+		bIncludeClassName = InIncludeClassName;
+	}
 
-		const bool bIsMatch = TextFilterUtils::TestBasicStringExpression(Asset.AssetName.ToString(), InValue, InTextComparisonMode);
-		return (InComparisonOperation == ETextFilterComparisonOperation::Equal) ? bIsMatch : !bIsMatch;
+	bool GetIncludeClassName() const
+	{
+		return bIncludeClassName;
 	}
 
-	// Special case for the asset path, as this isn't contained within the asset registry meta-data
-	if (InKey == PathKeyName)
+	virtual bool TestBasicStringExpression(const FTextFilterString& InValue, const ETextFilterTextComparisonMode InTextComparisonMode) const override
 	{
-		// Paths can only work with Equal or NotEqual type tests
-		if (InComparisonOperation != ETextFilterComparisonOperation::Equal && InComparisonOperation != ETextFilterComparisonOperation::NotEqual)
+		const TCHAR* Ptr = AssetFullPath.GetCharArray().GetData();
+		if (Ptr)
 		{
-			return false;
+			// Test each piece of the path name, apart from the first
+			bool bIsFirst = true;
+			while (const TCHAR* Delimiter = FCString::Strchr(Ptr, '/'))
+			{
+				const int32 Length = Delimiter - Ptr;
+
+				if (Length > 0)
+				{
+					if (bIsFirst)
+					{
+						bIsFirst = false;
+					}
+					else
+					{
+						if (TextFilterUtils::TestBasicStringExpression(FString(Length, Ptr), InValue, InTextComparisonMode))
+						{
+							return true;
+						}
+					}
+				}
+
+				Ptr += (Length + 1);
+			}
+
+			if (*Ptr != 0)
+			{
+				if (TextFilterUtils::TestBasicStringExpression(Ptr, InValue, InTextComparisonMode))
+				{
+					return true;
+				}
+			}
 		}
 
-		// If the comparison mode is partial, then we only need to test the ObjectPath as that contains the other two as sub-strings
-		bool bIsMatch = false;
-		if (InTextComparisonMode == ETextFilterTextComparisonMode::Partial)
+		if (bIncludeClassName)
 		{
-			bIsMatch = TextFilterUtils::TestBasicStringExpression(Asset.ObjectPath.ToString(), InValue, InTextComparisonMode);
+			if (TextFilterUtils::TestBasicStringExpression(AssetPtr->AssetClass, InValue, InTextComparisonMode))
+			{
+				return true;
+			}
+
+			// Only test this if we're searching the class name too, as the exported text contains the type in the string
+			if (TextFilterUtils::TestBasicStringExpression(AssetExportTextName, InValue, InTextComparisonMode))
+			{
+				return true;
+			}
 		}
-		else
+
+		for (const FName& AssetCollectionName : AssetCollectionNames)
 		{
-			bIsMatch = TextFilterUtils::TestBasicStringExpression(Asset.ObjectPath.ToString(), InValue, InTextComparisonMode)
-				|| TextFilterUtils::TestBasicStringExpression(Asset.PackageName.ToString(), InValue, InTextComparisonMode)
-				|| TextFilterUtils::TestBasicStringExpression(Asset.PackagePath.ToString(), InValue, InTextComparisonMode);
+			if (TextFilterUtils::TestBasicStringExpression(AssetCollectionName, InValue, InTextComparisonMode))
+			{
+				return true;
+			}
 		}
-		return (InComparisonOperation == ETextFilterComparisonOperation::Equal) ? bIsMatch : !bIsMatch;
+
+		return false;
 	}
 
-	// Special case for the asset type, as this isn't contained within the asset registry meta-data
-	if (InKey == ClassKeyName || InKey == TypeKeyName)
+	virtual bool TestComplexExpression(const FName& InKey, const FTextFilterString& InValue, const ETextFilterComparisonOperation InComparisonOperation, const ETextFilterTextComparisonMode InTextComparisonMode) const override
 	{
-		// Class names can only work with Equal or NotEqual type tests
-		if (InComparisonOperation != ETextFilterComparisonOperation::Equal && InComparisonOperation != ETextFilterComparisonOperation::NotEqual)
+		// Special case for the asset name, as this isn't contained within the asset registry meta-data
+		if (InKey == NameKeyName)
 		{
-			return false;
-		}
+			// Names can only work with Equal or NotEqual type tests
+			if (InComparisonOperation != ETextFilterComparisonOperation::Equal && InComparisonOperation != ETextFilterComparisonOperation::NotEqual)
+			{
+				return false;
+			}
 
-		const bool bIsMatch = TextFilterUtils::TestBasicStringExpression(Asset.AssetClass.ToString(), InValue, InTextComparisonMode);
-		return (InComparisonOperation == ETextFilterComparisonOperation::Equal) ? bIsMatch : !bIsMatch;
-	}
+			const bool bIsMatch = TextFilterUtils::TestBasicStringExpression(AssetPtr->AssetName, InValue, InTextComparisonMode);
+			return (InComparisonOperation == ETextFilterComparisonOperation::Equal) ? bIsMatch : !bIsMatch;
+		}
 
-	// Special case for collections, as these aren't contained within the asset registry meta-data
-	if ((InKey == CollectionKeyName || InKey == TagKeyName) && FCollectionManagerModule::IsModuleAvailable())
-	{
-		// Collections can only work with Equal or NotEqual type tests
-		if (InComparisonOperation != ETextFilterComparisonOperation::Equal && InComparisonOperation != ETextFilterComparisonOperation::NotEqual)
+		// Special case for the asset path, as this isn't contained within the asset registry meta-data
+		if (InKey == PathKeyName)
 		{
-			return false;
+			// Paths can only work with Equal or NotEqual type tests
+			if (InComparisonOperation != ETextFilterComparisonOperation::Equal && InComparisonOperation != ETextFilterComparisonOperation::NotEqual)
+			{
+				return false;
+			}
+
+			// If the comparison mode is partial, then we only need to test the ObjectPath as that contains the other two as sub-strings
+			bool bIsMatch = false;
+			if (InTextComparisonMode == ETextFilterTextComparisonMode::Partial)
+			{
+				bIsMatch = TextFilterUtils::TestBasicStringExpression(AssetPtr->ObjectPath, InValue, InTextComparisonMode);
+			}
+			else
+			{
+				bIsMatch = TextFilterUtils::TestBasicStringExpression(AssetPtr->ObjectPath, InValue, InTextComparisonMode)
+					|| TextFilterUtils::TestBasicStringExpression(AssetPtr->PackageName, InValue, InTextComparisonMode)
+					|| TextFilterUtils::TestBasicStringExpression(AssetPtr->PackagePath, InValue, InTextComparisonMode);
+			}
+			return (InComparisonOperation == ETextFilterComparisonOperation::Equal) ? bIsMatch : !bIsMatch;
 		}
 
-		FCollectionManagerModule& CollectionManagerModule = FCollectionManagerModule::GetModule();
+		// Special case for the asset type, as this isn't contained within the asset registry meta-data
+		if (InKey == ClassKeyName || InKey == TypeKeyName)
+		{
+			// Class names can only work with Equal or NotEqual type tests
+			if (InComparisonOperation != ETextFilterComparisonOperation::Equal && InComparisonOperation != ETextFilterComparisonOperation::NotEqual)
+			{
+				return false;
+			}
 
-		TArray<FName> AssetCollectionNames;
-		CollectionManagerModule.Get().GetCollectionsContainingObject(Asset.ObjectPath, ECollectionShareType::CST_All, AssetCollectionNames, ECollectionRecursionFlags::SelfAndChildren);
+			const bool bIsMatch = TextFilterUtils::TestBasicStringExpression(AssetPtr->AssetClass, InValue, InTextComparisonMode);
+			return (InComparisonOperation == ETextFilterComparisonOperation::Equal) ? bIsMatch : !bIsMatch;
+		}
 
-		bool bFoundMatch = false;
-		for (const FName& AssetCollectionName : AssetCollectionNames)
+		// Special case for collections, as these aren't contained within the asset registry meta-data
+		if (InKey == CollectionKeyName || InKey == TagKeyName)
 		{
-			if (TextFilterUtils::TestBasicStringExpression(AssetCollectionName.ToString(), InValue, InTextComparisonMode))
+			// Collections can only work with Equal or NotEqual type tests
+			if (InComparisonOperation != ETextFilterComparisonOperation::Equal && InComparisonOperation != ETextFilterComparisonOperation::NotEqual)
 			{
-				bFoundMatch = true;
-				break;
+				return false;
 			}
+
+			bool bFoundMatch = false;
+			for (const FName& AssetCollectionName : AssetCollectionNames)
+			{
+				if (TextFilterUtils::TestBasicStringExpression(AssetCollectionName, InValue, InTextComparisonMode))
+				{
+					bFoundMatch = true;
+					break;
+				}
+			}
+
+			return (InComparisonOperation == ETextFilterComparisonOperation::Equal) ? bFoundMatch : !bFoundMatch;
 		}
 
-		return (InComparisonOperation == ETextFilterComparisonOperation::Equal) ? bFoundMatch : !bFoundMatch;
-	}
+		// Generic handling for anything in the asset meta-data
+		const FString* const MetaDataValue = AssetPtr->TagsAndValues.Find(InKey);
+		if (MetaDataValue)
+		{
+			return TextFilterUtils::TestComplexExpression(*MetaDataValue, InValue, InComparisonOperation, InTextComparisonMode);
+		}
 
-	// Generic handling for anything in the asset meta-data
-	const FString* const MetaDataValue = Asset.TagsAndValues.Find(InKey);
-	if (MetaDataValue)
-	{
-		return TextFilterUtils::TestComplexExpression(*MetaDataValue, InValue, InComparisonOperation, InTextComparisonMode);
+		return false;
 	}
 
-	return false;
-}
+private:
+	/** Pointer to the asset we're currently filtering */
+	FAssetFilterTypePtr AssetPtr;
+
+	/** Full path of the current asset */
+	FString AssetFullPath;
+
+	/** The export text name of the current asset */
+	FString AssetExportTextName;
+
+	/** Names of the collections that the current asset is in */
+	TArray<FName> AssetCollectionNames;
+
+	/** Are we supposed to include the class name in our basic string tests? */
+	bool bIncludeClassName;
+
+	/** Scratch string that can be used for temporary string conversions while avoiding re-allocations */
+	mutable FString TempString;
+
+	/** Keys used by TestComplexExpression */
+	const FName NameKeyName;
+	const FName PathKeyName;
+	const FName ClassKeyName;
+	const FName TypeKeyName;
+	const FName CollectionKeyName;
+	const FName TagKeyName;
+};
 
 FFrontendFilter_Text::FFrontendFilter_Text()
 	: FFrontendFilter(nullptr)
-	, TextFilter(
-		TTextFilter<FAssetFilterType>::FItemToStringArray::CreateStatic(&AssetDataToClassAndNameStrings), 
-		TTextFilter<FAssetFilterType>::FItemTestComplexExpression::CreateStatic(&AssetDataTestComplexExpression)
-		)
+	, TextFilterExpressionContext(MakeShareable(new FFrontendFilter_TextFilterExpressionContext()))
+	, TextFilterExpressionEvaluator(ETextFilterExpressionEvaluatorMode::Complex)
 {
-	SetIncludeClassName(true);
 }
 
 FFrontendFilter_Text::~FFrontendFilter_Text()
 {
-	TextFilter.OnChanged().RemoveAll(this);
 }
 
 bool FFrontendFilter_Text::PassesFilter(FAssetFilterType InItem) const
 {
-	return TextFilter.PassesFilter(InItem);
+	TextFilterExpressionContext->SetAsset(&InItem);
+	const bool bMatched = TextFilterExpressionEvaluator.TestTextFilter(*TextFilterExpressionContext);
+	TextFilterExpressionContext->ClearAsset();
+	return bMatched;
 }
 
 FText FFrontendFilter_Text::GetRawFilterText() const
 {
-	return TextFilter.GetRawFilterText();
+	return TextFilterExpressionEvaluator.GetFilterText();
 }
 
 void FFrontendFilter_Text::SetRawFilterText(const FText& InFilterText)
 {
-	return TextFilter.SetRawFilterText(InFilterText);
+	if (TextFilterExpressionEvaluator.SetFilterText(InFilterText))
+	{
+		// Will trigger a re-filter with the new text
+		BroadcastChangedEvent();
+	}
 }
 
 FText FFrontendFilter_Text::GetFilterErrorText() const
 {
-	return TextFilter.GetFilterErrorText();
+	return TextFilterExpressionEvaluator.GetFilterErrorText();
 }
 
-void FFrontendFilter_Text::SetIncludeClassName(bool bIncludeClassName)
+void FFrontendFilter_Text::SetIncludeClassName(const bool InIncludeClassName)
 {
-	// Preserve the existing text. The filter is going to get recreated here.
-	const FText ExistingText = TextFilter.GetRawFilterText();
-
-	if ( bIncludeClassName )
+	if (TextFilterExpressionContext->GetIncludeClassName() != InIncludeClassName)
 	{
-		TextFilter = TTextFilter<FAssetFilterType>(
-			TTextFilter<FAssetFilterType>::FItemToStringArray::CreateStatic(&AssetDataToClassAndNameStrings),
-			TTextFilter<FAssetFilterType>::FItemTestComplexExpression::CreateStatic(&AssetDataTestComplexExpression)
-			);
-	}
-	else
-	{
-		TextFilter = TTextFilter<FAssetFilterType>(
-			TTextFilter<FAssetFilterType>::FItemToStringArray::CreateStatic(&AssetDataToNameString),
-			TTextFilter<FAssetFilterType>::FItemTestComplexExpression::CreateStatic(&AssetDataTestComplexExpression)
-			);
-	}
-
-	// Apply the existing text before we re-assign the delegate since we want the text preservation to be opaque.
-	TextFilter.SetRawFilterText(ExistingText);
+		TextFilterExpressionContext->SetIncludeClassName(InIncludeClassName);
 
-	TextFilter.OnChanged().AddRaw(this, &FFrontendFilter_Text::HandleOnChangedEvent);
-}
-
-void FFrontendFilter_Text::HandleOnChangedEvent()
-{
-	BroadcastChangedEvent();
+		// Will trigger a re-filter with the new setting
+		BroadcastChangedEvent();
+	}
 }
 
 
diff --git a/Engine/Source/Editor/ContentBrowser/Private/FrontendFilters.h b/Engine/Source/Editor/ContentBrowser/Private/FrontendFilters.h
--- a/Engine/Source/Editor/ContentBrowser/Private/FrontendFilters.h
+++ b/Engine/Source/Editor/ContentBrowser/Private/FrontendFilters.h
@@ -3,6 +3,7 @@
 #pragma once
 
 #include "FrontendFilterBase.h"
+#include "TextFilterExpressionEvaluator.h"
 
 #define LOCTEXT_NAMESPACE "ContentBrowser"
 
@@ -32,14 +33,14 @@ public:
 	FText GetFilterErrorText() const;
 
 	/** If bIncludeClassName is true, the text filter will include an asset's class name in the search */
-	void SetIncludeClassName(bool bIncludeClassName);
+	void SetIncludeClassName(const bool InIncludeClassName);
 
 private:
+	/** Transient context data, used when calling PassesFilter. Kept around to minimize re-allocations between multiple calls to PassesFilter */
+	TSharedRef<class FFrontendFilter_TextFilterExpressionContext> TextFilterExpressionContext;
 
-	/** Handler for the internal text filter */
-	void HandleOnChangedEvent();
-
-	TTextFilter<FAssetFilterType> TextFilter;
+	/** Expression evaluator that can be used to perform complex text filter queries */
+	FTextFilterExpressionEvaluator TextFilterExpressionEvaluator;
 };
 
 /** A filter that displays only checked out assets */
diff --git a/Engine/Source/Editor/ContentBrowser/Private/SCollectionView.cpp b/Engine/Source/Editor/ContentBrowser/Private/SCollectionView.cpp
--- a/Engine/Source/Editor/ContentBrowser/Private/SCollectionView.cpp
+++ b/Engine/Source/Editor/ContentBrowser/Private/SCollectionView.cpp
@@ -23,7 +23,7 @@ void GetBasicStrings(const FCollectionItem& InCollection, TArray<FString>& OutBa
 	OutBasicStrings.Add(InCollection.CollectionName.ToString());
 }
 
-bool TestComplexExpression(const FCollectionItem& InCollection, const FName& InKey, const FString& InValue, ETextFilterComparisonOperation InComparisonOperation, ETextFilterTextComparisonMode InTextComparisonMode)
+bool TestComplexExpression(const FCollectionItem& InCollection, const FName& InKey, const FTextFilterString& InValue, ETextFilterComparisonOperation InComparisonOperation, ETextFilterTextComparisonMode InTextComparisonMode)
 {
 	static const FName NameKeyName = "Name";
 	static const FName TypeKeyName = "Type";
diff --git a/Engine/Source/Runtime/Core/Private/Misc/TextFilterExpressionEvaluator.cpp b/Engine/Source/Runtime/Core/Private/Misc/TextFilterExpressionEvaluator.cpp
--- a/Engine/Source/Runtime/Core/Private/Misc/TextFilterExpressionEvaluator.cpp
+++ b/Engine/Source/Runtime/Core/Private/Misc/TextFilterExpressionEvaluator.cpp
@@ -9,7 +9,7 @@ namespace TextFilterExpressionParser
 	class FTextToken
 	{
 	public:
-		FTextToken(FString InString, ETextFilterTextComparisonMode InTextComparisonMode)
+		FTextToken(FTextFilterString InString, ETextFilterTextComparisonMode InTextComparisonMode)
 			: String(MoveTemp(InString))
 			, TextComparisonMode(MoveTemp(InTextComparisonMode))
 		{
@@ -41,11 +41,16 @@ namespace TextFilterExpressionParser
 			return *this;
 		}
 
-		const FString& GetString() const
+		const FTextFilterString& GetString() const
 		{
 			return String;
 		}
 
+		FName GetStringAsName() const
+		{
+			return String.AsName();
+		}
+
 		ETextFilterTextComparisonMode GetTextComparisonMode() const
 		{
 			return TextComparisonMode;
@@ -57,7 +62,7 @@ namespace TextFilterExpressionParser
 		}
 
 	private:
-		FString String;
+		FTextFilterString String;
 		ETextFilterTextComparisonMode TextComparisonMode;
 	};
 }
@@ -603,12 +608,12 @@ void FTextFilterExpressionEvaluator::ConstructExpressionParser()
 	Grammar.DefineBinaryOperator<FAnd>(2);
 	Grammar.DefinePreUnaryOperator<FNot>();
 
-	JumpTable.MapBinary<FLessOrEqual>([](const FTextToken& A, const FTextToken& B, const ITextFilterExpressionContext* InContext)		{ return InContext->TestComplexExpression(FName(*A.GetString()), B.GetString(), ETextFilterComparisonOperation::LessOrEqual, B.GetTextComparisonMode()); });
-	JumpTable.MapBinary<FLess>([](const FTextToken& A, const FTextToken& B, const ITextFilterExpressionContext* InContext)				{ return InContext->TestComplexExpression(FName(*A.GetString()), B.GetString(), ETextFilterComparisonOperation::Less, B.GetTextComparisonMode()); });
-	JumpTable.MapBinary<FGreaterOrEqual>([](const FTextToken& A, const FTextToken& B, const ITextFilterExpressionContext* InContext)	{ return InContext->TestComplexExpression(FName(*A.GetString()), B.GetString(), ETextFilterComparisonOperation::GreaterOrEqual, B.GetTextComparisonMode()); });
-	JumpTable.MapBinary<FGreater>([](const FTextToken& A, const FTextToken& B, const ITextFilterExpressionContext* InContext)			{ return InContext->TestComplexExpression(FName(*A.GetString()), B.GetString(), ETextFilterComparisonOperation::Greater, B.GetTextComparisonMode()); });
-	JumpTable.MapBinary<FNotEqual>([](const FTextToken& A, const FTextToken& B, const ITextFilterExpressionContext* InContext)			{ return InContext->TestComplexExpression(FName(*A.GetString()), B.GetString(), ETextFilterComparisonOperation::NotEqual, B.GetTextComparisonMode()); });
-	JumpTable.MapBinary<FEqual>([](const FTextToken& A, const FTextToken& B, const ITextFilterExpressionContext* InContext)				{ return InContext->TestComplexExpression(FName(*A.GetString()), B.GetString(), ETextFilterComparisonOperation::Equal, B.GetTextComparisonMode()); });
+	JumpTable.MapBinary<FLessOrEqual>([](const FTextToken& A, const FTextToken& B, const ITextFilterExpressionContext* InContext)		{ return InContext->TestComplexExpression(A.GetStringAsName(), B.GetString(), ETextFilterComparisonOperation::LessOrEqual, B.GetTextComparisonMode()); });
+	JumpTable.MapBinary<FLess>([](const FTextToken& A, const FTextToken& B, const ITextFilterExpressionContext* InContext)				{ return InContext->TestComplexExpression(A.GetStringAsName(), B.GetString(), ETextFilterComparisonOperation::Less, B.GetTextComparisonMode()); });
+	JumpTable.MapBinary<FGreaterOrEqual>([](const FTextToken& A, const FTextToken& B, const ITextFilterExpressionContext* InContext)	{ return InContext->TestComplexExpression(A.GetStringAsName(), B.GetString(), ETextFilterComparisonOperation::GreaterOrEqual, B.GetTextComparisonMode()); });
+	JumpTable.MapBinary<FGreater>([](const FTextToken& A, const FTextToken& B, const ITextFilterExpressionContext* InContext)			{ return InContext->TestComplexExpression(A.GetStringAsName(), B.GetString(), ETextFilterComparisonOperation::Greater, B.GetTextComparisonMode()); });
+	JumpTable.MapBinary<FNotEqual>([](const FTextToken& A, const FTextToken& B, const ITextFilterExpressionContext* InContext)			{ return InContext->TestComplexExpression(A.GetStringAsName(), B.GetString(), ETextFilterComparisonOperation::NotEqual, B.GetTextComparisonMode()); });
+	JumpTable.MapBinary<FEqual>([](const FTextToken& A, const FTextToken& B, const ITextFilterExpressionContext* InContext)				{ return InContext->TestComplexExpression(A.GetStringAsName(), B.GetString(), ETextFilterComparisonOperation::Equal, B.GetTextComparisonMode()); });
 
 	JumpTable.MapBinary<FOr>([](const FTextToken& A, const FTextToken& B, const ITextFilterExpressionContext* InContext)				{ return InContext->TestBasicStringExpression(A.GetString(), A.GetTextComparisonMode()) || InContext->TestBasicStringExpression(B.GetString(), B.GetTextComparisonMode()); });
 	JumpTable.MapBinary<FOr>([](const FTextToken& A, bool B, const ITextFilterExpressionContext* InContext)								{ return InContext->TestBasicStringExpression(A.GetString(), A.GetTextComparisonMode()) || B; });
diff --git a/Engine/Source/Runtime/Core/Private/Misc/TextFilterTests.cpp b/Engine/Source/Runtime/Core/Private/Misc/TextFilterTests.cpp
--- a/Engine/Source/Runtime/Core/Private/Misc/TextFilterTests.cpp
+++ b/Engine/Source/Runtime/Core/Private/Misc/TextFilterTests.cpp
@@ -23,7 +23,7 @@ namespace TextFilterTests
 		}
 
 		// ITextFilterExpressionContext API - used for testing FTextFilterExpressionEvaluator
-		virtual bool TestBasicStringExpression(const FString& InValue, const ETextFilterTextComparisonMode InTextComparisonMode) const override
+		virtual bool TestBasicStringExpression(const FTextFilterString& InValue, const ETextFilterTextComparisonMode InTextComparisonMode) const override
 		{
 			for (const FString& BasicString : BasicStrings)
 			{
@@ -34,7 +34,7 @@ namespace TextFilterTests
 			}
 			return false;
 		}
-		virtual bool TestComplexExpression(const FName& InKey, const FString& InValue, const ETextFilterComparisonOperation InComparisonOperation, const ETextFilterTextComparisonMode InTextComparisonMode) const override
+		virtual bool TestComplexExpression(const FName& InKey, const FTextFilterString& InValue, const ETextFilterComparisonOperation InComparisonOperation, const ETextFilterTextComparisonMode InTextComparisonMode) const override
 		{
 			const FString* ItemValue = KeyValuePairs.Find(InKey);
 			if (ItemValue)
@@ -49,7 +49,7 @@ namespace TextFilterTests
 		{
 			OutStrings = InItem->BasicStrings;
 		}
-		static bool TestItemComplexExpression(const FTestFilterItem* InItem, const FName& InKey, const FString& InValue, ETextFilterComparisonOperation InComparisonOperation, ETextFilterTextComparisonMode InTextComparisonMode)
+		static bool TestItemComplexExpression(const FTestFilterItem* InItem, const FName& InKey, const FTextFilterString& InValue, ETextFilterComparisonOperation InComparisonOperation, ETextFilterTextComparisonMode InTextComparisonMode)
 		{
 			return InItem->TestComplexExpression(InKey, InValue, InComparisonOperation, InTextComparisonMode);
 		}
diff --git a/Engine/Source/Runtime/Core/Private/Misc/TextFilterUtils.cpp b/Engine/Source/Runtime/Core/Private/Misc/TextFilterUtils.cpp
--- a/Engine/Source/Runtime/Core/Private/Misc/TextFilterUtils.cpp
+++ b/Engine/Source/Runtime/Core/Private/Misc/TextFilterUtils.cpp
@@ -3,67 +3,287 @@
 #include "CorePrivatePCH.h"
 #include "TextFilterUtils.h"
 
-bool TextFilterUtils::TestBasicStringExpression(const FString& InValue1, const FString& InValue2, const ETextFilterTextComparisonMode InTextComparisonMode)
+/** ToUpper implementation that's optimized for ASCII characters */
+namespace FastToUpper
+{
+	// -32 = 'A' - 'a'
+	static const int32 ToUpperAdjustmentTable[] = 
+	{
+		// ASCII control characters
+		0,		// NUL
+		0,		// SOH
+		0,		// STX
+		0,		// ETX
+		0,		// EOT
+		0,		// ENQ
+		0,		// ACK
+		0,		// BEL
+		0,		// BS
+		0,		// HT
+		0,		// LF
+		0,		// VT
+		0,		// FF
+		0,		// CR
+		0,		// SO
+		0,		// SI
+		0,		// DLE
+		0,		// DC1
+		0,		// DC2
+		0,		// DC3
+		0,		// DC4
+		0,		// NAK
+		0,		// SYN
+		0,		// ETB
+		0,		// CAN
+		0,		// EM
+		0,		// SUB
+		0,		// ESC
+		0,		// FS
+		0,		// GS
+		0,		// RS
+		0,		// US
+
+		// ASCII printable characters
+		0,		//  
+		0,		// !
+		0,		// "
+		0,		// #
+		0,		// $
+		0,		// %
+		0,		// &
+		0,		// '
+		0,		// (
+		0,		// )
+		0,		// *
+		0,		// +
+		0,		// ,
+		0,		// -
+		0,		// .
+		0,		// /
+		0,		// 0
+		0,		// 1
+		0,		// 2
+		0,		// 3
+		0,		// 4
+		0,		// 5
+		0,		// 6
+		0,		// 7
+		0,		// 8
+		0,		// 9
+		0,		// :
+		0,		// ;
+		0,		// <
+		0,		// =
+		0,		// >
+		0,		// ?
+		0,		// @
+		0,		// A
+		0,		// B
+		0,		// C
+		0,		// D
+		0,		// E
+		0,		// F
+		0,		// G
+		0,		// H
+		0,		// I
+		0,		// J
+		0,		// K
+		0,		// L
+		0,		// M
+		0,		// N
+		0,		// O
+		0,		// P
+		0,		// Q
+		0,		// R
+		0,		// S
+		0,		// T
+		0,		// U
+		0,		// V
+		0,		// W
+		0,		// X
+		0,		// Y
+		0,		// Z
+		0,		// [
+		0,		// \ 
+		0,		// ]
+		0,		// ^
+		0,		// _
+		0,		// `
+		-32,	// a
+		-32,	// b
+		-32,	// c
+		-32,	// d
+		-32,	// e
+		-32,	// f
+		-32,	// g
+		-32,	// h
+		-32,	// i
+		-32,	// j
+		-32,	// k
+		-32,	// l
+		-32,	// m
+		-32,	// n
+		-32,	// o
+		-32,	// p
+		-32,	// q
+		-32,	// r
+		-32,	// s
+		-32,	// t
+		-32,	// u
+		-32,	// v
+		-32,	// w
+		-32,	// x
+		-32,	// y
+		-32,	// z
+		0,		// {
+		0,		// |
+		0,		// }
+		0,		// ~
+		0,		// 	
+	};
+	static const int32 ToUpperAdjustmentTableCount = ARRAY_COUNT(ToUpperAdjustmentTable);
+
+	FORCEINLINE TCHAR ToUpper(const TCHAR InChar)
+	{
+		if (InChar < ToUpperAdjustmentTableCount)
+		{
+			return InChar + ToUpperAdjustmentTable[(int32)InChar];
+		}
+		return FChar::ToUpper(InChar);
+	}
+} // namespace FastToUpper
+
+FTextFilterString::FTextFilterString()
+	: InternalString()
+{
+}
+
+FTextFilterString::FTextFilterString(const FTextFilterString& Other)
+	: InternalString(Other.InternalString)
+{
+}
+
+FTextFilterString::FTextFilterString(FTextFilterString&& Other)
+	: InternalString(MoveTemp(Other.InternalString))
+{
+}
+
+FTextFilterString::FTextFilterString(const FString& InString)
+	: InternalString(InString)
+{
+	UppercaseInternalString();
+}
+
+FTextFilterString::FTextFilterString(FString&& InString)
+	: InternalString(MoveTemp(InString))
+{
+	UppercaseInternalString();
+}
+
+FTextFilterString::FTextFilterString(const TCHAR* InString)
+	: InternalString(InString)
+{
+	UppercaseInternalString();
+}
+
+FTextFilterString::FTextFilterString(const FName& InName)
+	: InternalString()
+{
+	InName.AppendString(InternalString);
+	UppercaseInternalString();
+}
+
+FTextFilterString& FTextFilterString::operator=(const FTextFilterString& Other)
+{
+	InternalString = Other.InternalString;
+	return *this;
+}
+
+FTextFilterString& FTextFilterString::operator=(FTextFilterString&& Other)
+{
+	InternalString = MoveTemp(Other.InternalString);
+	return *this;
+}
+
+bool FTextFilterString::CompareText(const FTextFilterString& InOther, const ETextFilterTextComparisonMode InTextComparisonMode) const
 {
 	switch(InTextComparisonMode)
 	{
 	case ETextFilterTextComparisonMode::Exact:
-		return InValue1.Equals(InValue2, ESearchCase::IgnoreCase);
+		return FCString::Strcmp(*InternalString, *InOther.InternalString) == 0;
 	case ETextFilterTextComparisonMode::Partial:
-		return InValue1.Contains(InValue2, ESearchCase::IgnoreCase);
+		return FCString::Strstr(*InternalString, *InOther.InternalString) != nullptr;
 	case ETextFilterTextComparisonMode::StartsWith:
-		return InValue1.StartsWith(InValue2, ESearchCase::IgnoreCase);
+		return InOther.InternalString.Len() > 0 && FCString::Strncmp(*InternalString, *InOther.InternalString, InOther.InternalString.Len()) == 0;
 	case ETextFilterTextComparisonMode::EndsWith:
-		return InValue1.EndsWith(InValue2, ESearchCase::IgnoreCase);
+		return InOther.InternalString.Len() > 0 && InternalString.Len() >= InOther.InternalString.Len() 
+			&& FCString::Strncmp((*InternalString) + (InternalString.Len() - InOther.InternalString.Len()), *InOther.InternalString, InOther.InternalString.Len()) == 0;
 	default:
 		break;
 	}
+
 	return false;
 }
 
-bool TextFilterUtils::TestComplexExpression(const FString& InValue1, const FString& InValue2, const ETextFilterComparisonOperation InComparisonOperation, const ETextFilterTextComparisonMode InTextComparisonMode)
+bool FTextFilterString::CanCompareNumeric(const FTextFilterString& InOther) const
 {
-	const bool bNumericComparison = InValue1.IsNumeric() && InValue2.IsNumeric();
+	return InternalString.IsNumeric() && InOther.InternalString.IsNumeric();
+}
 
-	int32 ComparisonSign = 0;
-	if (bNumericComparison)
-	{
-		const double Difference = FCString::Atod(*InValue1) - FCString::Atod(*InValue2);
-		ComparisonSign = (int32)FMath::Sign(Difference);
-	}
-	else if (InTextComparisonMode == ETextFilterTextComparisonMode::Exact)
-	{
-		ComparisonSign = InValue1.Compare(InValue2, ESearchCase::IgnoreCase);
-	}
-	else
-	{
-		ComparisonSign = (InValue1.Contains(InValue2, ESearchCase::IgnoreCase)) ? 0 : 1;
-	}
+bool FTextFilterString::CompareNumeric(const FTextFilterString& InOther, const ETextFilterComparisonOperation InComparisonOperation) const
+{
+	const double OurNumericValue = FCString::Atod(*InternalString);
+	const double OtherNumericValue = FCString::Atod(*InOther.InternalString);
+	const double Difference = OurNumericValue - OtherNumericValue;
 
-	bool bComparisonPassed = false;
+	const int32 ComparisonSign = (int32)FMath::Sign(Difference);
 	switch (InComparisonOperation)
 	{
 	case ETextFilterComparisonOperation::Equal:
-		bComparisonPassed = ComparisonSign == 0;
-		break;
+		return ComparisonSign == 0;
 	case ETextFilterComparisonOperation::NotEqual:
-		bComparisonPassed = ComparisonSign != 0;
-		break;
+		return ComparisonSign != 0;
 	case ETextFilterComparisonOperation::Less:
-		bComparisonPassed = ComparisonSign < 0;
-		break;
+		return ComparisonSign < 0;
 	case ETextFilterComparisonOperation::LessOrEqual:
-		bComparisonPassed = ComparisonSign <= 0;
-		break;
+		return ComparisonSign <= 0;
 	case ETextFilterComparisonOperation::Greater:
-		bComparisonPassed = ComparisonSign > 0;
-		break;
+		return ComparisonSign > 0;
 	case ETextFilterComparisonOperation::GreaterOrEqual:
-		bComparisonPassed = ComparisonSign >= 0;
-		break;
+		return ComparisonSign >= 0;
 	default:
 		check(false);
 	};
 
-	return bComparisonPassed;
+	return false;
+}
+
+void FTextFilterString::UppercaseInternalString()
+{
+	for (TCHAR* CharPtr = const_cast<TCHAR*>(*InternalString); *CharPtr; ++CharPtr)
+	{
+		*CharPtr = FastToUpper::ToUpper(*CharPtr);
+	}
+}
+
+bool TextFilterUtils::TestBasicStringExpression(const FTextFilterString& InValue1, const FTextFilterString& InValue2, const ETextFilterTextComparisonMode InTextComparisonMode)
+{
+	return InValue1.CompareText(InValue2, InTextComparisonMode);
+}
+
+bool TextFilterUtils::TestComplexExpression(const FTextFilterString& InValue1, const FTextFilterString& InValue2, const ETextFilterComparisonOperation InComparisonOperation, const ETextFilterTextComparisonMode InTextComparisonMode)
+{
+	if (InValue1.CanCompareNumeric(InValue2))
+	{
+		return InValue1.CompareNumeric(InValue2, InComparisonOperation);
+	}
+
+	// Text can only work with Equal or NotEqual type tests
+	if (InComparisonOperation == ETextFilterComparisonOperation::Equal || InComparisonOperation == ETextFilterComparisonOperation::NotEqual)
+	{
+		const bool bIsMatch = InValue1.CompareText(InValue2, InTextComparisonMode);
+		return (InComparisonOperation == ETextFilterComparisonOperation::Equal) ? bIsMatch : !bIsMatch;
+	}
+
+	return false;
 }
diff --git a/Engine/Source/Runtime/Core/Public/Misc/TextFilter.h b/Engine/Source/Runtime/Core/Public/Misc/TextFilter.h
--- a/Engine/Source/Runtime/Core/Public/Misc/TextFilter.h
+++ b/Engine/Source/Runtime/Core/Public/Misc/TextFilter.h
@@ -16,7 +16,7 @@ public:
 	DECLARE_DELEGATE_TwoParams( FItemToStringArray, ItemType, OUT TArray< FString >& );
 
 	/** Defines a function signature used to test a complex expression for an Item */
-	DECLARE_DELEGATE_RetVal_FiveParams( bool, FItemTestComplexExpression, ItemType, const FName& /*InKey*/, const FString& /*InValue*/, ETextFilterComparisonOperation /*InComparisonOperation*/, ETextFilterTextComparisonMode /*InTextComparisonMode*/ );
+	DECLARE_DELEGATE_RetVal_FiveParams( bool, FItemTestComplexExpression, ItemType, const FName& /*InKey*/, const FTextFilterString& /*InValue*/, ETextFilterComparisonOperation /*InComparisonOperation*/, ETextFilterTextComparisonMode /*InTextComparisonMode*/ );
 
 	/** 
 	 *	TTextFilter Constructor
@@ -107,7 +107,7 @@ private:
 			ItemBasicStrings.Reset();
 		}
 
-		virtual bool TestBasicStringExpression(const FString& InValue, const ETextFilterTextComparisonMode InTextComparisonMode) const override
+		virtual bool TestBasicStringExpression(const FTextFilterString& InValue, const ETextFilterTextComparisonMode InTextComparisonMode) const override
 		{
 			for (const FString& BasicString : ItemBasicStrings)
 			{
@@ -119,7 +119,7 @@ private:
 			return false;
 		}
 
-		virtual bool TestComplexExpression(const FName& InKey, const FString& InValue, const ETextFilterComparisonOperation InComparisonOperation, const ETextFilterTextComparisonMode InTextComparisonMode) const override
+		virtual bool TestComplexExpression(const FName& InKey, const FTextFilterString& InValue, const ETextFilterComparisonOperation InComparisonOperation, const ETextFilterTextComparisonMode InTextComparisonMode) const override
 		{
 			if (ItemPtr && TestComplexExpressionDelegate.IsBound())
 			{
diff --git a/Engine/Source/Runtime/Core/Public/Misc/TextFilterExpressionEvaluator.h b/Engine/Source/Runtime/Core/Public/Misc/TextFilterExpressionEvaluator.h
--- a/Engine/Source/Runtime/Core/Public/Misc/TextFilterExpressionEvaluator.h
+++ b/Engine/Source/Runtime/Core/Public/Misc/TextFilterExpressionEvaluator.h
@@ -25,10 +25,10 @@ class ITextFilterExpressionContext
 {
 public:
 	/** Test the given value against the strings extracted from the current item */
-	virtual bool TestBasicStringExpression(const FString& InValue, const ETextFilterTextComparisonMode InTextComparisonMode) const = 0;
+	virtual bool TestBasicStringExpression(const FTextFilterString& InValue, const ETextFilterTextComparisonMode InTextComparisonMode) const = 0;
 		
 	/** Perform a complex expression test for the current item */
-	virtual bool TestComplexExpression(const FName& InKey, const FString& InValue, const ETextFilterComparisonOperation InComparisonOperation, const ETextFilterTextComparisonMode InTextComparisonMode) const = 0;
+	virtual bool TestComplexExpression(const FName& InKey, const FTextFilterString& InValue, const ETextFilterComparisonOperation InComparisonOperation, const ETextFilterTextComparisonMode InTextComparisonMode) const = 0;
 
 protected:
 	virtual ~ITextFilterExpressionContext() {}
diff --git a/Engine/Source/Runtime/Core/Public/Misc/TextFilterUtils.h b/Engine/Source/Runtime/Core/Public/Misc/TextFilterUtils.h
--- a/Engine/Source/Runtime/Core/Public/Misc/TextFilterUtils.h
+++ b/Engine/Source/Runtime/Core/Public/Misc/TextFilterUtils.h
@@ -21,11 +21,66 @@ enum class ETextFilterTextComparisonMode : uint8
 	EndsWith,
 };
 
+/**
+ * String used by the text filter.
+ * The given string will be stored as uppercase since filter text always performs case-insensitive string comparisons, so this will minimize ToUpper calls.
+ */
+class CORE_API FTextFilterString
+{
+public:
+	/** Default constructor */
+	FTextFilterString();
+
+	/** Move and copy constructors */
+	FTextFilterString(const FTextFilterString& Other);
+	FTextFilterString(FTextFilterString&& Other);
+
+	/** Construct from a string */
+	FTextFilterString(const FString& InString);
+	FTextFilterString(FString&& InString);
+	FTextFilterString(const TCHAR* InString);
+
+	/** Construct from a name */
+	FTextFilterString(const FName& InName);
+
+	/** Move and copy assignment */
+	FTextFilterString& operator=(const FTextFilterString& Other);
+	FTextFilterString& operator=(FTextFilterString&& Other);
+
+	/** Compare this string against the other, using the text comparison mode provided */
+	bool CompareText(const FTextFilterString& InOther, const ETextFilterTextComparisonMode InTextComparisonMode) const;
+
+	/** Are the two given strings able to be compared numberically? */
+	bool CanCompareNumeric(const FTextFilterString& InOther) const;
+
+	/** Compare this string against the other, converting them to numbers and using the comparison operator provided - you should have tested CanCompareNumeric first! */
+	bool CompareNumeric(const FTextFilterString& InOther, const ETextFilterComparisonOperation InComparisonOperation) const;
+
+	/** Get the internal uppercase string of this filter string */
+	FORCEINLINE const FString& AsString() const
+	{
+		return InternalString;
+	}
+
+	/** Get the internal uppercase string of this filter string as an FName */
+	FORCEINLINE FName AsName() const
+	{
+		return FName(*InternalString);
+	}
+
+private:
+	/** Inline convert our internal string to uppercase */
+	void UppercaseInternalString();
+
+	/** The uppercase string to use for comparisons */
+	FString InternalString;
+};
+
 namespace TextFilterUtils
 {
 	/** Utility function to perform a basic string test for the given values */
-	CORE_API bool TestBasicStringExpression(const FString& InValue1, const FString& InValue2, const ETextFilterTextComparisonMode InTextComparisonMode);
+	CORE_API bool TestBasicStringExpression(const FTextFilterString& InValue1, const FTextFilterString& InValue2, const ETextFilterTextComparisonMode InTextComparisonMode);
 
 	/** Utility function to perform a complex expression test for the given values */
-	CORE_API bool TestComplexExpression(const FString& InValue1, const FString& InValue2, const ETextFilterComparisonOperation InComparisonOperation, const ETextFilterTextComparisonMode InTextComparisonMode);
+	CORE_API bool TestComplexExpression(const FTextFilterString& InValue1, const FTextFilterString& InValue2, const ETextFilterComparisonOperation InComparisonOperation, const ETextFilterTextComparisonMode InTextComparisonMode);
 }
