Removed unnecessary file movement popup when moving folder. Swapped ShouldPromptToLoadAssets with GetUnloadedAssets (which was still needed).

UE-16599

diff --git a/Engine/Source/Editor/ContentBrowser/Private/ContentBrowserUtils.cpp b/Engine/Source/Editor/ContentBrowser/Private/ContentBrowserUtils.cpp
--- a/Engine/Source/Editor/ContentBrowser/Private/ContentBrowserUtils.cpp
+++ b/Engine/Source/Editor/ContentBrowser/Private/ContentBrowserUtils.cpp
@@ -308,21 +308,9 @@ bool ContentBrowserUtils::LoadAssetsIfNeeded(const TArray<FString>& ObjectPaths,
 			}
 		}
 	}
-
-	// if we are allowed to prompt the user to load and we have enough assets that requires prompting then we should 
-	// prompt and load assets if the user said it was ok
-	bool bShouldLoadAssets = true;
-	if( bAllowedToPromptToLoadAssets && ShouldPromptToLoadAssets(ObjectPaths, UnloadedObjectPaths) )
-	{
-		bShouldLoadAssets = PromptToLoadAssets(ObjectPaths);
-	}
-
-
-	// Ask for confirmation if the user is attempting to load a large number of assets
-	if (bShouldLoadAssets == false)
-	{
-		return false;
-	}
+ 
+	// prompt and load assets
+	GetUnloadedAssets(ObjectPaths, UnloadedObjectPaths);
 
 	// Make sure all selected objects are loaded, where possible
 	if ( UnloadedObjectPaths.Num() > 0 )
@@ -385,12 +373,11 @@ bool ContentBrowserUtils::LoadAssetsIfNeeded(const TArray<FString>& ObjectPaths,
 	return true;
 }
 
-bool ContentBrowserUtils::ShouldPromptToLoadAssets(const TArray<FString>& ObjectPaths, TArray<FString>& OutUnloadedObjects)
+void ContentBrowserUtils::GetUnloadedAssets(const TArray<FString>& ObjectPaths, TArray<FString>& OutUnloadedObjects)
 {
 	OutUnloadedObjects.Empty();
 
-	bool bShouldPrompt = false;
-	// Build a list of unloaded assets
+	// Build a list of unloaded assets and check if there are any parent folders
 	for (int32 PathIdx = 0; PathIdx < ObjectPaths.Num(); ++PathIdx)
 	{
 		const FString& ObjectPath = ObjectPaths[PathIdx];
@@ -402,15 +389,6 @@ bool ContentBrowserUtils::ShouldPromptToLoadAssets(const TArray<FString>& Object
 			OutUnloadedObjects.Add(ObjectPath);
 		}
 	}
-
-	// Get the maximum objects to load before displaying a warning
-	// Ask for confirmation if the user is attempting to load a large number of assets
-	if (OutUnloadedObjects.Num() > GetDefault<UContentBrowserSettings>()->NumObjectsToLoadBeforeWarning)
-	{
-		bShouldPrompt = true;
-	}
-
-	return bShouldPrompt;
 }
 
 bool ContentBrowserUtils::PromptToLoadAssets(const TArray<FString>& UnloadedObjects)
@@ -817,14 +795,7 @@ bool ContentBrowserUtils::PrepareFoldersForDragDrop(const TArray<FString>& Sourc
 	}
 
 	TArray<FString> UnloadedObjects;
-	if(ShouldPromptToLoadAssets(ObjectPathsToWarnAbout, UnloadedObjects))
-	{
-		const bool bShouldLoadAssets = PromptToLoadAssets(UnloadedObjects);
-		if (!bShouldLoadAssets)
-		{
-			return false;
-		}
-	}
+	GetUnloadedAssets(ObjectPathsToWarnAbout, UnloadedObjects);
 
 	GWarn->BeginSlowTask(LOCTEXT("FolderDragDrop_Loading", "Loading folders"), true);
 
diff --git a/Engine/Source/Editor/ContentBrowser/Private/ContentBrowserUtils.h b/Engine/Source/Editor/ContentBrowser/Private/ContentBrowserUtils.h
--- a/Engine/Source/Editor/ContentBrowser/Private/ContentBrowserUtils.h
+++ b/Engine/Source/Editor/ContentBrowser/Private/ContentBrowserUtils.h
@@ -26,7 +26,7 @@ namespace ContentBrowserUtils
 	bool OpenEditorForAsset(const TArray<UObject*>& Assets);
 
 	/**
-	  * Makes sure the specified assets are loaded into memory.
+	  * Makes sure the specified assets are loaded into memory.6666
 	  * 
 	  * @param ObjectPaths The paths to the objects to load.
 	  * @param LoadedObjects The returned list of objects that were already loaded or loaded by this method.
@@ -35,13 +35,13 @@ namespace ContentBrowserUtils
 	bool LoadAssetsIfNeeded(const TArray<FString>& ObjectPaths, TArray<UObject*>& LoadedObjects, bool bAllowedToPromptToLoadAssets = true, bool bLoadRedirects = false);
 
 	/**
-	 * Determines if enough assets are unloaded that we should prompt the user to load them instead of loading them automatically
+	 * Determines the unloaded assets that need loading
 	 *
 	 * @param ObjectPaths		Paths to assets that may need to be loaded
 	 * @param OutUnloadedObjects	List of the unloaded object paths
 	 * @return true if the user should be prompted to load assets
 	 */
-	bool ShouldPromptToLoadAssets(const TArray<FString>& ObjectPaths, TArray<FString>& OutUnloadedObjects);
+	void GetUnloadedAssets(const TArray<FString>& ObjectPaths, TArray<FString>& OutUnloadedObjects);
 
 	/**
 	 * Prompts the user to load the list of unloaded objects
diff --git a/Engine/Source/Editor/ContentBrowser/Private/SAssetView.cpp b/Engine/Source/Editor/ContentBrowser/Private/SAssetView.cpp
--- a/Engine/Source/Editor/ContentBrowser/Private/SAssetView.cpp
+++ b/Engine/Source/Editor/ContentBrowser/Private/SAssetView.cpp
@@ -3194,37 +3194,20 @@ bool SAssetView::CanOpenContextMenu() const
 	}
 
 	bool bLoadSuccessful = true;
-	bool bShouldPromptToLoadAssets = false;
 
 	if ( bPreloadAssetsForContextMenu )
 	{
-		// Should the user be asked to load unloaded assets
+		// Get and load assets that are unloaded
 		TArray<FString> UnloadedObjects;
-		bShouldPromptToLoadAssets = ContentBrowserUtils::ShouldPromptToLoadAssets(ObjectPaths, UnloadedObjects);
+		ContentBrowserUtils::GetUnloadedAssets(ObjectPaths, UnloadedObjects);
 
-		bool bShouldLoadAssets = false;
-		if ( bShouldPromptToLoadAssets )
-		{
-			// The user should be prompted to loaded assets
-			bShouldLoadAssets = ContentBrowserUtils::PromptToLoadAssets(UnloadedObjects);
-		}
-		else
-		{
-			// The user should not be prompted to load assets but assets should still be loaded
-			bShouldLoadAssets = true;
-		}
-
-		if ( bShouldLoadAssets )
-		{
-			// Load assets that are unloaded
-			TArray<UObject*> LoadedObjects;
-			const bool bAllowedToPrompt = false;
-			bLoadSuccessful = ContentBrowserUtils::LoadAssetsIfNeeded(ObjectPaths, LoadedObjects, bAllowedToPrompt);
-		}
+		TArray<UObject*> LoadedObjects;
+		const bool bAllowedToPrompt = false;
+		bLoadSuccessful = ContentBrowserUtils::LoadAssetsIfNeeded(ObjectPaths, LoadedObjects, bAllowedToPrompt);
 	}
 
-	// Do not show the context menu if we prompted the user to load assets or if the load failed
-	return !bShouldPromptToLoadAssets && bLoadSuccessful;
+	// Do not show the context menu if the load failed
+	return bLoadSuccessful;
 }
 
 void SAssetView::OnListMouseButtonDoubleClick(TSharedPtr<FAssetViewItem> AssetItem)
