Made asset view tooltips refresh each time they become visible

UE-17564 - Collections are not removed from Asset tooltip when it's removed from collection

This allows the list of collections that an object is in to be updated, as this isn't stored as part of the asset data.

The main set of tags (which is the bulk of the tooltip text) is still cached and only updated when the asset data is changed.

diff --git a/Engine/Source/Editor/ContentBrowser/Private/AssetViewWidgets.cpp b/Engine/Source/Editor/ContentBrowser/Private/AssetViewWidgets.cpp
--- a/Engine/Source/Editor/ContentBrowser/Private/AssetViewWidgets.cpp
+++ b/Engine/Source/Editor/ContentBrowser/Private/AssetViewWidgets.cpp
@@ -203,6 +203,57 @@ TSharedRef<SWidget> FAssetViewItemHelper::CreateListTileItemContents(T* const In
 }
 
 
+///////////////////////////////
+// Asset view item tool tip
+///////////////////////////////
+
+class SAssetViewItemToolTip : public SToolTip
+{
+public:
+	SLATE_BEGIN_ARGS(SAssetViewItemToolTip)
+		: _AssetViewItem()
+	{ }
+
+		SLATE_ARGUMENT(TSharedPtr<SAssetViewItem>, AssetViewItem)
+
+	SLATE_END_ARGS()
+
+	void Construct(const FArguments& InArgs)
+	{
+		AssetViewItem = InArgs._AssetViewItem;
+
+		SToolTip::Construct(
+			SToolTip::FArguments()
+			.TextMargin(1.0f)
+			.BorderImage(FEditorStyle::GetBrush("ContentBrowser.TileViewTooltip.ToolTipBorder"))
+			);
+	}
+
+	// IToolTip interface
+	virtual bool IsEmpty() const override
+	{
+		return !AssetViewItem.IsValid();
+	}
+
+	virtual void OnOpening() override
+	{
+		TSharedPtr<SAssetViewItem> AssetViewItemPin = AssetViewItem.Pin();
+		if (AssetViewItemPin.IsValid())
+		{
+			SetContentWidget(AssetViewItemPin->CreateToolTipWidget());
+		}
+	}
+
+	virtual void OnClosed() override
+	{
+		SetContentWidget(SNullWidget::NullWidget);
+	}
+
+private:
+	TWeakPtr<SAssetViewItem> AssetViewItem;
+};
+
+
 ///////////////////////////////
 // Asset view modes
 ///////////////////////////////
@@ -293,6 +344,9 @@ void SAssetViewItem::Construct( const FArguments& InArgs )
 	AssetDirtyBrush = FEditorStyle::GetBrush("ContentBrowser.ContentDirty");
 	SCCStateBrush = nullptr;
 
+	// Set our tooltip - this will refresh each time it's opened to make sure it's up-to-date
+	SetToolTip(SNew(SAssetViewItemToolTip).AssetViewItem(SharedThis(this)));
+
 	// refresh SCC state icon
 	HandleSourceControlStateChanged();
 
@@ -451,6 +505,8 @@ void SAssetViewItem::OnAssetDataChanged()
 	{
 		InlineRenameWidget->SetText( GetNameText() );
 	}
+
+	CacheToolTipTags();
 }
 
 void SAssetViewItem::DirtyStateChanged()
@@ -509,7 +565,7 @@ EVisibility SAssetViewItem::GetThumbnailEditModeUIVisibility() const
 }
 
 BEGIN_SLATE_FUNCTION_BUILD_OPTIMIZATION
-TSharedRef<SToolTip> SAssetViewItem::CreateToolTipWidget() const
+TSharedRef<SWidget> SAssetViewItem::CreateToolTipWidget() const
 {
 	if ( AssetItem.IsValid() )
 	{
@@ -521,13 +577,11 @@ TSharedRef<SToolTip> SAssetViewItem::CreateToolTipWidget() const
 		else if(AssetItem->GetType() != EAssetItemType::Folder)
 		{
 			const FAssetData& AssetData = StaticCastSharedPtr<FAssetViewAsset>(AssetItem)->Data;
-			UClass* AssetClass = FindObject<UClass>(ANY_PACKAGE, *AssetData.AssetClass.ToString());
 
 			// The tooltip contains the name, class, path, and asset registry tags
 			const FText NameText = FText::FromName( AssetData.AssetName );
 			const FText ClassText = FText::Format( LOCTEXT("ClassName", "({0})"), GetAssetClassText() );
 
-
 			// Create a box to hold every line of info in the body of the tooltip
 			TSharedRef<SVerticalBox> InfoBox = SNew(SVerticalBox);
 
@@ -545,87 +599,10 @@ TSharedRef<SToolTip> SAssetViewItem::CreateToolTipWidget() const
 				}
 			}
 
-			// If we are using a loaded class, find all the hidden tags so we don't display them
-			TMap<FName, UObject::FAssetRegistryTagMetadata> MetadataMap;
-			TSet<FName> ShownTags;
-			if ( AssetClass != NULL && AssetClass->GetDefaultObject() != NULL )
-			{
-				AssetClass->GetDefaultObject()->GetAssetRegistryTagMetadata(MetadataMap);
-
-				TArray<UObject::FAssetRegistryTag> Tags;
-				AssetClass->GetDefaultObject()->GetAssetRegistryTags(Tags);
-
-				for ( auto TagIt = Tags.CreateConstIterator(); TagIt; ++TagIt )
-				{
-					if (TagIt->Type != UObject::FAssetRegistryTag::TT_Hidden )
-					{
-						ShownTags.Add(TagIt->Name);
-					}
-				}
-			}
-
-			// If an asset class could not be loaded we cannot determine hidden tags so display no tags.  
-			if( AssetClass != NULL )
+			// Add tags
+			for (const auto& ToolTipTagItem : CachedToolTipTags)
 			{
-				// Add all asset registry tags and values
-				for ( auto TagIt = AssetData.TagsAndValues.CreateConstIterator(); TagIt; ++TagIt )
-				{
-					// Skip tags that are set to be hidden
-					if ( ShownTags.Contains(TagIt.Key()) )
-					{
-						const UObject::FAssetRegistryTagMetadata* Metadata = MetadataMap.Find(TagIt.Key());
-
-						const bool bImportant = (Metadata != nullptr && !Metadata->ImportantValue.IsEmpty() && Metadata->ImportantValue == TagIt.Value());
-
-						// Since all we have at this point is a string, we can't be very smart here.
-						// We need to strip some noise off class paths in some cases, but can't load the asset to inspect its UPROPERTYs manually due to performance concerns.
-						FString ValueString = TagIt.Value();
-						const TCHAR StringToRemove[] = TEXT("Class'/Script/");
-						if (ValueString.StartsWith(StringToRemove) && ValueString.EndsWith(TEXT("'")))
-						{
-							// Remove the class path for native classes, and also remove Engine. for engine classes
-							const int32 SizeOfPrefix = ARRAY_COUNT(StringToRemove);
-							ValueString = ValueString.Mid(SizeOfPrefix - 1, ValueString.Len() - SizeOfPrefix).Replace(TEXT("Engine."), TEXT(""));
-						}
-						
-						// Check for DisplayName metadata
-						FText DisplayName;
-						if (UProperty* Field = FindField<UProperty>(AssetClass, TagIt.Key()))
-						{
-							DisplayName = Field->GetDisplayNameText();
-
-							// Strip off enum prefixes if they exist
-							if (UByteProperty* ByteProperty = Cast<UByteProperty>(Field))
-							{
-								if (ByteProperty->Enum)
-								{
-									const FString EnumPrefix = ByteProperty->Enum->GenerateEnumPrefix();
-									if (EnumPrefix.Len() && ValueString.StartsWith(EnumPrefix))
-									{
-										ValueString = ValueString.RightChop(EnumPrefix.Len() + 1);	// +1 to skip over the underscore
-									}
-								}
-
-								ValueString = FName::NameToDisplayString(ValueString, false);
-							}
-						}
-						else
-						{
-							// We have no type information by this point, so no idea if it's a bool :(
-							const bool bIsBool = false;
-							DisplayName = FText::FromString(FName::NameToDisplayString(TagIt.Key().ToString(), bIsBool));
-						}
-					
-						// Add suffix to the value string, if one is defined for this tag
-						if (Metadata != nullptr && !Metadata->Suffix.IsEmpty())
-						{
-							ValueString += TEXT(" ");
-							ValueString += Metadata->Suffix.ToString();
-						}
-
-						AddToToolTipInfoBox(InfoBox, DisplayName, FText::FromString(ValueString), bImportant);
-					}
-				}
+				AddToToolTipInfoBox(InfoBox, ToolTipTagItem.Key, ToolTipTagItem.Value, ToolTipTagItem.bImportant);
 			}
 
 			TSharedRef<SVerticalBox> OverallTooltipVBox = SNew(SVerticalBox);
@@ -709,16 +686,11 @@ TSharedRef<SToolTip> SAssetViewItem::CreateToolTipWidget() const
 				];
 
 
-			return SNew(SToolTip)
-				.TextMargin(1)
-				.BorderImage( FEditorStyle::GetBrush("ContentBrowser.TileViewTooltip.ToolTipBorder") )
+			return SNew(SBorder)
+				.Padding(6)
+				.BorderImage( FEditorStyle::GetBrush("ContentBrowser.TileViewTooltip.NonContentBorder") )
 				[
-					SNew(SBorder)
-					.Padding(6)
-					.BorderImage( FEditorStyle::GetBrush("ContentBrowser.TileViewTooltip.NonContentBorder") )
-					[
-						OverallTooltipVBox
-					]
+					OverallTooltipVBox
 				];
 		}
 		else
@@ -731,61 +703,56 @@ TSharedRef<SToolTip> SAssetViewItem::CreateToolTipWidget() const
 
 			AddToToolTipInfoBox( InfoBox, LOCTEXT("TileViewTooltipPath", "Path"), FText::FromString(FolderPath), false );
 
-			return SNew(SToolTip)
-				.TextMargin(1)
-				.BorderImage( FEditorStyle::GetBrush("ContentBrowser.TileViewTooltip.ToolTipBorder") )
+			return SNew(SBorder)
+				.Padding(6)
+				.BorderImage( FEditorStyle::GetBrush("ContentBrowser.TileViewTooltip.NonContentBorder") )
 				[
-					SNew(SBorder)
-					.Padding(6)
-					.BorderImage( FEditorStyle::GetBrush("ContentBrowser.TileViewTooltip.NonContentBorder") )
-					[
-						SNew(SVerticalBox)
+					SNew(SVerticalBox)
 
-						+SVerticalBox::Slot()
-						.AutoHeight()
-						.Padding(0, 0, 0, 4)
+					+SVerticalBox::Slot()
+					.AutoHeight()
+					.Padding(0, 0, 0, 4)
+					[
+						SNew(SBorder)
+						.Padding(6)
+						.BorderImage( FEditorStyle::GetBrush("ContentBrowser.TileViewTooltip.ContentBorder") )
 						[
-							SNew(SBorder)
-							.Padding(6)
-							.BorderImage( FEditorStyle::GetBrush("ContentBrowser.TileViewTooltip.ContentBorder") )
+							SNew(SVerticalBox)
+
+							+SVerticalBox::Slot()
+							.AutoHeight()
 							[
-								SNew(SVerticalBox)
+								SNew(SHorizontalBox)
+
+								+SHorizontalBox::Slot()
+								.AutoWidth()
+								.VAlign(VAlign_Center)
+								.Padding(0, 0, 4, 0)
+								[
+									SNew(STextBlock)
+									.Text( FolderName )
+									.Font( FEditorStyle::GetFontStyle("ContentBrowser.TileViewTooltip.NameFont") )
+								]
 
-								+SVerticalBox::Slot()
-								.AutoHeight()
+								+SHorizontalBox::Slot()
+								.AutoWidth()
+								.VAlign(VAlign_Center)
 								[
-									SNew(SHorizontalBox)
-
-									+SHorizontalBox::Slot()
-									.AutoWidth()
-									.VAlign(VAlign_Center)
-									.Padding(0, 0, 4, 0)
-									[
-										SNew(STextBlock)
-										.Text( FolderName )
-										.Font( FEditorStyle::GetFontStyle("ContentBrowser.TileViewTooltip.NameFont") )
-									]
-
-									+SHorizontalBox::Slot()
-									.AutoWidth()
-									.VAlign(VAlign_Center)
-									[
-										SNew(STextBlock) 
-										.Text( LOCTEXT("FolderNameBracketed", "(Folder)") )
-									]
+									SNew(STextBlock) 
+									.Text( LOCTEXT("FolderNameBracketed", "(Folder)") )
 								]
 							]
 						]
+					]
 
-						+SVerticalBox::Slot()
-						.AutoHeight()
+					+SVerticalBox::Slot()
+					.AutoHeight()
+					[
+						SNew(SBorder)
+						.Padding(6)
+						.BorderImage( FEditorStyle::GetBrush("ContentBrowser.TileViewTooltip.ContentBorder") )
 						[
-							SNew(SBorder)
-							.Padding(6)
-							.BorderImage( FEditorStyle::GetBrush("ContentBrowser.TileViewTooltip.ContentBorder") )
-							[
-								InfoBox
-							]
+							InfoBox
 						]
 					]
 				];
@@ -794,7 +761,7 @@ TSharedRef<SToolTip> SAssetViewItem::CreateToolTipWidget() const
 	else
 	{
 		// Return an empty tooltip since the asset item wasn't valid
-		return SNew(SToolTip);
+		return SNullWidget::NullWidget;
 	}
 }
 END_SLATE_FUNCTION_BUILD_OPTIMIZATION
@@ -927,6 +894,100 @@ void SAssetViewItem::CachePackageName()
 	}
 }
 
+void SAssetViewItem::CacheToolTipTags()
+{
+	CachedToolTipTags.Reset();
+
+	if(AssetItem->GetType() != EAssetItemType::Folder)
+	{
+		const FAssetData& AssetData = StaticCastSharedPtr<FAssetViewAsset>(AssetItem)->Data;
+		UClass* AssetClass = FindObject<UClass>(ANY_PACKAGE, *AssetData.AssetClass.ToString());
+
+		// If we are using a loaded class, find all the hidden tags so we don't display them
+		TMap<FName, UObject::FAssetRegistryTagMetadata> MetadataMap;
+		TSet<FName> ShownTags;
+		if ( AssetClass != NULL && AssetClass->GetDefaultObject() != NULL )
+		{
+			AssetClass->GetDefaultObject()->GetAssetRegistryTagMetadata(MetadataMap);
+
+			TArray<UObject::FAssetRegistryTag> Tags;
+			AssetClass->GetDefaultObject()->GetAssetRegistryTags(Tags);
+
+			for ( auto TagIt = Tags.CreateConstIterator(); TagIt; ++TagIt )
+			{
+				if (TagIt->Type != UObject::FAssetRegistryTag::TT_Hidden )
+				{
+					ShownTags.Add(TagIt->Name);
+				}
+			}
+		}
+
+		// If an asset class could not be loaded we cannot determine hidden tags so display no tags.
+		if( AssetClass != NULL )
+		{
+			// Add all asset registry tags and values
+			for (const auto& KeyValue: AssetData.TagsAndValues)
+			{
+				// Skip tags that are set to be hidden
+				if ( ShownTags.Contains(KeyValue.Key) )
+				{
+					const UObject::FAssetRegistryTagMetadata* Metadata = MetadataMap.Find(KeyValue.Key);
+
+					const bool bImportant = (Metadata != nullptr && !Metadata->ImportantValue.IsEmpty() && Metadata->ImportantValue == KeyValue.Value);
+
+					// Since all we have at this point is a string, we can't be very smart here.
+					// We need to strip some noise off class paths in some cases, but can't load the asset to inspect its UPROPERTYs manually due to performance concerns.
+					FString ValueString = KeyValue.Value;
+					const TCHAR StringToRemove[] = TEXT("Class'/Script/");
+					if (ValueString.StartsWith(StringToRemove) && ValueString.EndsWith(TEXT("'")))
+					{
+						// Remove the class path for native classes, and also remove Engine. for engine classes
+						const int32 SizeOfPrefix = ARRAY_COUNT(StringToRemove);
+						ValueString = ValueString.Mid(SizeOfPrefix - 1, ValueString.Len() - SizeOfPrefix).Replace(TEXT("Engine."), TEXT(""));
+					}
+						
+					// Check for DisplayName metadata
+					FText DisplayName;
+					if (UProperty* Field = FindField<UProperty>(AssetClass, KeyValue.Key))
+					{
+						DisplayName = Field->GetDisplayNameText();
+
+						// Strip off enum prefixes if they exist
+						if (UByteProperty* ByteProperty = Cast<UByteProperty>(Field))
+						{
+							if (ByteProperty->Enum)
+							{
+								const FString EnumPrefix = ByteProperty->Enum->GenerateEnumPrefix();
+								if (EnumPrefix.Len() && ValueString.StartsWith(EnumPrefix))
+								{
+									ValueString = ValueString.RightChop(EnumPrefix.Len() + 1);	// +1 to skip over the underscore
+								}
+							}
+
+							ValueString = FName::NameToDisplayString(ValueString, false);
+						}
+					}
+					else
+					{
+						// We have no type information by this point, so no idea if it's a bool :(
+						const bool bIsBool = false;
+						DisplayName = FText::FromString(FName::NameToDisplayString(KeyValue.Key.ToString(), bIsBool));
+					}
+					
+					// Add suffix to the value string, if one is defined for this tag
+					if (Metadata != nullptr && !Metadata->Suffix.IsEmpty())
+					{
+						ValueString += TEXT(" ");
+						ValueString += Metadata->Suffix.ToString();
+					}
+
+					CachedToolTipTags.Add(FToolTipTagItem(DisplayName, FText::FromString(ValueString), bImportant));
+				}
+			}
+		}
+	}
+}
+
 const FSlateBrush* SAssetViewItem::GetBorderImage() const
 {
 	return bDraggedOver ? FEditorStyle::GetBrush("Menu.Background") : FEditorStyle::GetBrush("NoBorder");
@@ -1152,8 +1213,6 @@ void SAssetListItem::Construct( const FArguments& InArgs )
 		]
 	];
 
-	SetToolTip( CreateToolTipWidget() );
-
 	if(AssetItem.IsValid())
 	{
 		AssetItem->RenamedRequestEvent.BindSP( InlineRenameWidget.Get(), &SInlineEditableTextBlock::EnterEditingMode );
@@ -1174,8 +1233,6 @@ void SAssetListItem::OnAssetDataChanged()
 	{
 		ClassText->SetText(GetAssetClassText());
 	}
-
-	SetToolTip( CreateToolTipWidget() );
 }
 
 FOptionalSize SAssetListItem::GetThumbnailBoxSize() const
@@ -1295,8 +1352,6 @@ void SAssetTileItem::Construct( const FArguments& InArgs )
 	
 	];
 
-	SetToolTip( CreateToolTipWidget() );
-
 	if(AssetItem.IsValid())
 	{
 		AssetItem->RenamedRequestEvent.BindSP( InlineRenameWidget.Get(), &SInlineEditableTextBlock::EnterEditingMode );
@@ -1312,8 +1367,6 @@ END_SLATE_FUNCTION_BUILD_OPTIMIZATION
 void SAssetTileItem::OnAssetDataChanged()
 {
 	SAssetViewItem::OnAssetDataChanged();
-
-	SetToolTip( CreateToolTipWidget() );
 }
 
 FOptionalSize SAssetTileItem::GetThumbnailBoxSize() const
@@ -1384,6 +1437,16 @@ public:
 		];
 	}
 
+	virtual TSharedPtr<IToolTip> GetToolTip() override
+	{
+		if ( OwnerAssetColumnItem.IsValid() )
+		{
+			return OwnerAssetColumnItem.Pin()->GetToolTip();
+		}
+
+		return nullptr;
+	}
+
 	/** Forward the event to the view item that this name box belongs to */
 	virtual void OnToolTipClosing() override
 	{
@@ -1514,7 +1577,6 @@ TSharedRef<SWidget> SAssetColumnItem::GenerateWidgetForColumn( const FName& Colu
 
 		return SNew( SAssetColumnItemNameBox, SharedThis(this) )
 			.Padding( ColumnItemPadding )
-			.ToolTip( CreateToolTipWidget() )
 			[
 				Content.ToSharedRef()
 			];
diff --git a/Engine/Source/Editor/ContentBrowser/Private/AssetViewWidgets.h b/Engine/Source/Editor/ContentBrowser/Private/AssetViewWidgets.h
--- a/Engine/Source/Editor/ContentBrowser/Private/AssetViewWidgets.h
+++ b/Engine/Source/Editor/ContentBrowser/Private/AssetViewWidgets.h
@@ -53,6 +53,8 @@ public:
 /** A base class for all asset view items */
 class SAssetViewItem : public SCompoundWidget
 {
+	friend class SAssetViewItemToolTip;
+
 public:
 	DECLARE_DELEGATE_TwoParams( FOnAssetsDragDropped, const TArray<FAssetData>& /*AssetList*/, const FString& /*DestinationPath*/);
 	DECLARE_DELEGATE_TwoParams( FOnPathsDragDropped, const TArray<FString>& /*PathNames*/, const FString& /*DestinationPath*/);
@@ -170,7 +172,7 @@ protected:
 	EVisibility GetThumbnailEditModeUIVisibility() const;
 
 	/** Creates a tooltip widget for this item */
-	TSharedRef<SToolTip> CreateToolTipWidget() const;
+	TSharedRef<SWidget> CreateToolTipWidget() const;
 
 	/** Gets the visibility of the checked out by other text block in the tooltip */
 	EVisibility GetCheckedOutByOtherTextVisibility() const;
@@ -196,6 +198,9 @@ protected:
 	/** Cache the package name from the asset we are representing */
 	void CachePackageName();
 
+	/** Cache the tags that should appear in the tooltip for this item */
+	void CacheToolTipTags();
+
 	/** Whether this item is a folder */
 	bool IsFolder() const;
 
@@ -207,7 +212,22 @@ protected:
 
 	/** Returns the width at which the name label will wrap the name */
 	virtual float GetNameTextWrapWidth() const { return 0.0f; }
+
 protected:
+	/** Data for a cached tag used in the tooltip for this item */
+	struct FToolTipTagItem
+	{
+		FToolTipTagItem(FText InKey, FText InValue, const bool InImportant)
+			: Key(MoveTemp(InKey))
+			, Value(MoveTemp(InValue))
+			, bImportant(InImportant)
+		{
+		}
+
+		FText Key;
+		FText Value;
+		bool bImportant;
+	};
 
 	TSharedPtr< SInlineEditableTextBlock > InlineRenameWidget;
 
@@ -226,6 +246,9 @@ protected:
 	/** The cached filename of the package containing the asset that this item represents */
 	FString CachedPackageFileName;
 
+	/** The cached tags that should appear in the tooltip for this item */
+	TArray<FToolTipTagItem> CachedToolTipTags;
+
 	/** Delegate for when an asset name has entered a rename state */
 	FOnRenameBegin OnRenameBegin;
 
