Can now use the "Show Developers Folder" option in the save/load menu

UE-6312 - Cannot toggle developer folders in Save Level As view options

diff --git a/Engine/Source/Editor/ContentBrowser/Private/SAssetDialog.cpp b/Engine/Source/Editor/ContentBrowser/Private/SAssetDialog.cpp
--- a/Engine/Source/Editor/ContentBrowser/Private/SAssetDialog.cpp
+++ b/Engine/Source/Editor/ContentBrowser/Private/SAssetDialog.cpp
@@ -50,6 +50,7 @@ void SAssetDialog::Construct(const FArguments& InArgs, const FSharedAssetDialogC
 	AssetPickerConfig.GetCurrentSelectionDelegates.Add(&GetCurrentSelectionDelegate);
 	AssetPickerConfig.SaveSettingsName = TEXT("AssetDialog");
 	AssetPickerConfig.bCanShowFolders = true;
+	AssetPickerConfig.bCanShowDevelopersFolder = true;
 	AssetPickerConfig.OnFolderEntered = FOnPathSelected::CreateSP(this, &SAssetDialog::HandleAssetViewFolderEntered);
 
 	SetCurrentlySelectedPath(DefaultPath);
