Fixed the "Show Collections" option appearing multiple times if multiple Content Browsers were open

UE-17880 - Show Collections shows up twice when box is checked and multiple content browsers

This used to be added via an extender, however the extender was bound once each time a Content Browser was opened, meaning it was possible to have the entry added multiple times if multiple Content Browsers were open.

This menu item is now handled by the SAssetView (the same as the other view options), and is also now disabled on dialogs where it doesn't make sense to change the "Show Collections" settings (eg, when using the load or save picker dialogs, as they don't contain a collections view).

diff --git a/Engine/Source/Editor/ContentBrowser/Private/SAssetPicker.cpp b/Engine/Source/Editor/ContentBrowser/Private/SAssetPicker.cpp
--- a/Engine/Source/Editor/ContentBrowser/Private/SAssetPicker.cpp
+++ b/Engine/Source/Editor/ContentBrowser/Private/SAssetPicker.cpp
@@ -223,6 +223,7 @@ void SAssetPicker::Construct( const FArguments& InArgs )
 		.FilterRecursivelyWithBackendFilter( false )
 		.CanShowRealTimeThumbnails( InArgs._AssetPickerConfig.bCanShowRealTimeThumbnails )
 		.CanShowDevelopersFolder( InArgs._AssetPickerConfig.bCanShowDevelopersFolder )
+		.CanShowCollections( false )
 		.PreloadAssetsForContextMenu( InArgs._AssetPickerConfig.bPreloadAssetsForContextMenu )
 		.HighlightedText( HighlightText )
 		.ThumbnailLabel( ThumbnailLabel )
diff --git a/Engine/Source/Editor/ContentBrowser/Private/SAssetView.cpp b/Engine/Source/Editor/ContentBrowser/Private/SAssetView.cpp
--- a/Engine/Source/Editor/ContentBrowser/Private/SAssetView.cpp
+++ b/Engine/Source/Editor/ContentBrowser/Private/SAssetView.cpp
@@ -147,6 +147,8 @@ void SAssetView::Construct( const FArguments& InArgs )
 
 	bCanShowDevelopersFolder = InArgs._CanShowDevelopersFolder;
 
+	bCanShowCollections = InArgs._CanShowCollections;
+
 	bPreloadAssetsForContextMenu = InArgs._PreloadAssetsForContextMenu;
 
 	SelectionMode = InArgs._SelectionMode;
@@ -2381,6 +2383,19 @@ TSharedRef<SWidget> SAssetView::GetViewButtonContent()
 			NAME_None,
 			EUserInterfaceActionType::ToggleButton
 			);
+
+		MenuBuilder.AddMenuEntry(
+			LOCTEXT("ShowCollectionOption", "Show Collections"),
+			LOCTEXT("ShowCollectionOptionToolTip", "Show the collections list in the view."),
+			FSlateIcon(),
+			FUIAction(
+			FExecuteAction::CreateSP( this, &SAssetView::ToggleShowCollections ),
+			FCanExecuteAction::CreateSP( this, &SAssetView::IsToggleShowCollectionsAllowed ),
+			FIsActionChecked::CreateSP( this, &SAssetView::IsShowingCollections )
+			),
+			NAME_None,
+			EUserInterfaceActionType::ToggleButton
+			);
 	}
 	MenuBuilder.EndSection();
 
@@ -2534,6 +2549,23 @@ bool SAssetView::IsShowingDevelopersFolder() const
 	return GetDefault<UContentBrowserSettings>()->GetDisplayDevelopersFolder();
 }
 
+void SAssetView::ToggleShowCollections()
+{
+	const bool bDisplayCollections = GetDefault<UContentBrowserSettings>()->GetDisplayCollections();
+	GetMutableDefault<UContentBrowserSettings>()->SetDisplayCollections( !bDisplayCollections );
+	GetMutableDefault<UContentBrowserSettings>()->PostEditChange();
+}
+
+bool SAssetView::IsToggleShowCollectionsAllowed() const
+{
+	return bCanShowCollections;
+}
+
+bool SAssetView::IsShowingCollections() const
+{
+	return GetDefault<UContentBrowserSettings>()->GetDisplayCollections();
+}
+
 void SAssetView::SetCurrentViewType(EAssetViewType::Type NewType)
 {
 	if ( ensure(NewType != EAssetViewType::MAX) && NewType != CurrentViewType )
diff --git a/Engine/Source/Editor/ContentBrowser/Private/SAssetView.h b/Engine/Source/Editor/ContentBrowser/Private/SAssetView.h
--- a/Engine/Source/Editor/ContentBrowser/Private/SAssetView.h
+++ b/Engine/Source/Editor/ContentBrowser/Private/SAssetView.h
@@ -26,6 +26,7 @@ public:
 		, _FilterRecursivelyWithBackendFilter(true)
 		, _CanShowRealTimeThumbnails(false)
 		, _CanShowDevelopersFolder(false)
+		, _CanShowCollections(false)
 		, _PreloadAssetsForContextMenu(true)
 		, _SelectionMode( ESelectionMode::Multi )
 		, _AllowDragging(true)
@@ -120,6 +121,9 @@ public:
 		/** Indicates if the 'Show Developers' option should be enabled or disabled */
 		SLATE_ARGUMENT( bool, CanShowDevelopersFolder )
 
+		/** Indicates if the 'Show Collections' option should be enabled or disabled */
+		SLATE_ARGUMENT( bool, CanShowCollections )
+
 		/** Indicates if the context menu is going to load the assets, and if so to preload before the context menu is shown, and warn about the pending load. */
 		SLATE_ARGUMENT( bool, PreloadAssetsForContextMenu )
 
@@ -388,6 +392,15 @@ private:
 	/** @return true when we are showing the developers folder */
 	bool IsShowingDevelopersFolder() const;
 
+	/** Toggle whether collections should be shown or not */
+	void ToggleShowCollections();
+
+	/** Whether or not it's possible to toggle collections */
+	bool IsToggleShowCollectionsAllowed() const;
+
+	/** @return true when we are showing collections */
+	bool IsShowingCollections() const;
+
 	/** Sets the view type and updates lists accordingly */
 	void SetCurrentViewType(EAssetViewType::Type NewType);
 
@@ -767,6 +780,9 @@ private:
 	/** Indicates if the 'Show Developers' option should be enabled or disabled */
 	bool bCanShowDevelopersFolder;
 
+	/** Indicates if the 'Show Collections' option should be enabled or disabled */
+	bool bCanShowCollections;
+
 	/** Indicates if the context menu is going to load the assets, and if so to preload before the context menu is shown, and warn about the pending load. */
 	bool bPreloadAssetsForContextMenu;
 
diff --git a/Engine/Source/Editor/ContentBrowser/Private/SContentBrowser.cpp b/Engine/Source/Editor/ContentBrowser/Private/SContentBrowser.cpp
--- a/Engine/Source/Editor/ContentBrowser/Private/SContentBrowser.cpp
+++ b/Engine/Source/Editor/ContentBrowser/Private/SContentBrowser.cpp
@@ -590,6 +590,7 @@ void SContentBrowser::Construct( const FArguments& InArgs, const FName& InInstan
 						.CanShowFolders(true)
 						.CanShowRealTimeThumbnails(true)
 						.CanShowDevelopersFolder(true)
+						.CanShowCollections(true)
 						.AddMetaData<FTagMetaData>(FTagMetaData(TEXT("ContentBrowserAssets")))
 					]
 				]
@@ -632,21 +633,6 @@ void SContentBrowser::Construct( const FArguments& InArgs, const FName& InInstan
 	// We want to be able to search the feature packs in the super search so we need the module loaded 
 	IAddContentDialogModule& AddContentDialogModule = FModuleManager::LoadModuleChecked<IAddContentDialogModule>("AddContentDialog");
 
-	const TWeakPtr<SContentBrowser> WeakThis = SharedThis(this);
-
-	FContentBrowserModule& ContentBrowserModule = FModuleManager::GetModuleChecked<FContentBrowserModule>( TEXT("ContentBrowser") );
-	ContentBrowserModule.GetAllAssetViewViewMenuExtenders().Add( 
-		FContentBrowserMenuExtender::CreateLambda( [=]()
-		{
-			TSharedRef<FExtender> Extender = MakeShareable(new FExtender);
-			if( WeakThis.IsValid() )
-			{
-				Extender->AddMenuExtension(FName("Folders"), EExtensionHook::After, nullptr, FMenuExtensionDelegate::CreateSP(WeakThis.Pin().ToSharedRef(), &SContentBrowser::ExtendAssetViewMenu));
-			}
-
-			return Extender;
-		})
-	);
 	// Update the breadcrumb trail path
 	UpdatePath();
 }
@@ -679,36 +665,9 @@ void SContentBrowser::BindCommands()
 		));
 }
 
-void SContentBrowser::ExtendAssetViewMenu( FMenuBuilder& MenuBuilder )
-{
-	MenuBuilder.AddMenuEntry(
-		LOCTEXT("ShowCollectionOption", "Show Collections"),
-		LOCTEXT("ShowCollectionOptionToolTip", "Show the collections list in the view."),
-		FSlateIcon(),
-		FUIAction(
-		FExecuteAction::CreateSP(this, &SContentBrowser::ToggleShowCollections),
-		FCanExecuteAction(),
-		FIsActionChecked::CreateSP(this, &SContentBrowser::IsShowingCollections)
-		),
-		NAME_None,
-		EUserInterfaceActionType::ToggleButton
-		);
-}
-
-void SContentBrowser::ToggleShowCollections()
-{
-	GetMutableDefault<UContentBrowserSettings>()->SetDisplayCollections(!GetDefault<UContentBrowserSettings>()->GetDisplayCollections());
-	GetMutableDefault<UContentBrowserSettings>()->PostEditChange();
-}
-
-bool SContentBrowser::IsShowingCollections() const
-{
-	return GetDefault<UContentBrowserSettings>()->GetDisplayCollections();
-}
-
 EVisibility SContentBrowser::GetCollectionViewVisibility() const
 {
-	return IsShowingCollections() ? EVisibility::Visible : EVisibility::Collapsed;
+	return GetDefault<UContentBrowserSettings>()->GetDisplayCollections() ? EVisibility::Visible : EVisibility::Collapsed;
 }
 
 FText SContentBrowser::GetHighlightedText() const
diff --git a/Engine/Source/Editor/ContentBrowser/Private/SContentBrowser.h b/Engine/Source/Editor/ContentBrowser/Private/SContentBrowser.h
--- a/Engine/Source/Editor/ContentBrowser/Private/SContentBrowser.h
+++ b/Engine/Source/Editor/ContentBrowser/Private/SContentBrowser.h
@@ -320,15 +320,6 @@ private:
 	/** Bind our UI commands */
 	void BindCommands();
 
-	/** Adds menu extensions to the asset view */
-	void ExtendAssetViewMenu( FMenuBuilder& MenuBuilder );
-
-	/** Enables or disables collections from being shown */
-	void ToggleShowCollections();
-
-	/** @return Whether or not we are showing collections */
-	bool IsShowingCollections() const;
-
 	/** Gets the visibility of the collection view */
 	EVisibility GetCollectionViewVisibility() const;
 private:
