Added support for navigating collections using the path breadcumbs in the Content Browser

UE-17646 - Collections don't display bread crumb path when nested

diff --git a/Engine/Source/Developer/CollectionManager/Private/CollectionManager.cpp b/Engine/Source/Developer/CollectionManager/Private/CollectionManager.cpp
--- a/Engine/Source/Developer/CollectionManager/Private/CollectionManager.cpp
+++ b/Engine/Source/Developer/CollectionManager/Private/CollectionManager.cpp
@@ -154,6 +154,21 @@ void FCollectionManager::GetChildCollectionNames(FName CollectionName, ECollecti
 	}
 }
 
+TOptional<FCollectionNameType> FCollectionManager::GetParentCollection(FName CollectionName, ECollectionShareType::Type ShareType) const
+{
+	const TSharedRef<FCollection>* const CollectionRefPtr = CachedCollections.Find(FCollectionNameType(CollectionName, ShareType));
+	if (CollectionRefPtr)
+	{
+		const FCollectionNameType* const ParentCollectionKeyPtr = CachedCollectionNamesFromGuids.Find((*CollectionRefPtr)->GetParentCollectionGuid());
+		if (ParentCollectionKeyPtr)
+		{
+			return *ParentCollectionKeyPtr;
+		}
+	}
+
+	return TOptional<FCollectionNameType>();
+}
+
 bool FCollectionManager::CollectionExists(FName CollectionName, ECollectionShareType::Type ShareType) const
 {
 	if (ShareType == ECollectionShareType::CST_All)
diff --git a/Engine/Source/Developer/CollectionManager/Private/CollectionManager.h b/Engine/Source/Developer/CollectionManager/Private/CollectionManager.h
--- a/Engine/Source/Developer/CollectionManager/Private/CollectionManager.h
+++ b/Engine/Source/Developer/CollectionManager/Private/CollectionManager.h
@@ -16,6 +16,7 @@ public:
 	virtual void GetRootCollectionNames(ECollectionShareType::Type ShareType, TArray<FName>& CollectionNames) const override;
 	virtual void GetChildCollections(FName CollectionName, ECollectionShareType::Type ShareType, TArray<FCollectionNameType>& OutCollections) const override;
 	virtual void GetChildCollectionNames(FName CollectionName, ECollectionShareType::Type ShareType, ECollectionShareType::Type ChildShareType, TArray<FName>& CollectionNames) const override;
+	virtual TOptional<FCollectionNameType> GetParentCollection(FName CollectionName, ECollectionShareType::Type ShareType) const override;
 	virtual bool CollectionExists(FName CollectionName, ECollectionShareType::Type ShareType) const override;
 	virtual bool GetAssetsInCollection(FName CollectionName, ECollectionShareType::Type ShareType, TArray<FName>& AssetPaths, ECollectionRecursionFlags::Flags RecursionMode = ECollectionRecursionFlags::Self) const override;
 	virtual bool GetObjectsInCollection(FName CollectionName, ECollectionShareType::Type ShareType, TArray<FName>& ObjectPaths, ECollectionRecursionFlags::Flags RecursionMode = ECollectionRecursionFlags::Self) const override;
diff --git a/Engine/Source/Developer/CollectionManager/Public/CollectionManagerTypes.h b/Engine/Source/Developer/CollectionManager/Public/CollectionManagerTypes.h
--- a/Engine/Source/Developer/CollectionManager/Public/CollectionManagerTypes.h
+++ b/Engine/Source/Developer/CollectionManager/Public/CollectionManagerTypes.h
@@ -132,6 +132,11 @@ struct FCollectionNameType
 		return Name == Other.Name && Type == Other.Type;
 	}
 
+	bool operator!=(const FCollectionNameType& Other) const
+	{
+		return !(*this == Other);
+	}
+
 	friend inline uint32 GetTypeHash( const FCollectionNameType& Key )
 	{
 		uint32 Hash = 0;
diff --git a/Engine/Source/Developer/CollectionManager/Public/ICollectionManager.h b/Engine/Source/Developer/CollectionManager/Public/ICollectionManager.h
--- a/Engine/Source/Developer/CollectionManager/Public/ICollectionManager.h
+++ b/Engine/Source/Developer/CollectionManager/Public/ICollectionManager.h
@@ -31,6 +31,9 @@ public:
 	/** Returns the list of child collections of the given collection that are also of the specified share type */
 	virtual void GetChildCollectionNames(FName CollectionName, ECollectionShareType::Type ShareType, ECollectionShareType::Type ChildShareType, TArray<FName>& CollectionNames) const = 0;
 
+	/** Returns the parent collection of the given collection, or an unset value if there is no parent set */
+	virtual TOptional<FCollectionNameType> GetParentCollection(FName CollectionName, ECollectionShareType::Type ShareType) const = 0;
+
 	/** Returns true if the collection exists */
 	virtual bool CollectionExists(FName CollectionName, ECollectionShareType::Type ShareType) const = 0;
 
diff --git a/Engine/Source/Editor/ContentBrowser/Private/SContentBrowser.cpp b/Engine/Source/Editor/ContentBrowser/Private/SContentBrowser.cpp
--- a/Engine/Source/Editor/ContentBrowser/Private/SContentBrowser.cpp
+++ b/Engine/Source/Editor/ContentBrowser/Private/SContentBrowser.cpp
@@ -1254,7 +1254,29 @@ void SContentBrowser::OnPathClicked( const FString& CrumbData )
 
 	if ( SourcesData.Collections.Num() > 0 )
 	{
-		// Collection crumb was clicked. Since we don't have a hierarchy of collections, this does nothing.
+		// Collection crumb was clicked. See if we've clicked on a different collection in the hierarchy, and change the path if required.
+		TOptional<FCollectionNameType> CollectionClicked;
+		{
+			FString CollectionName;
+			FString CollectionTypeString;
+			if (CrumbData.Split(TEXT("?"), &CollectionName, &CollectionTypeString))
+			{
+				const int32 CollectionType = FCString::Atoi(*CollectionTypeString);
+				if (CollectionType >= 0 && CollectionType < ECollectionShareType::CST_All)
+				{
+					CollectionClicked = FCollectionNameType(FName(*CollectionName), ECollectionShareType::Type(CollectionType));
+				}
+			}
+		}
+
+		if ( CollectionClicked.IsSet() && SourcesData.Collections[0] != CollectionClicked.GetValue() )
+		{
+			TArray<FCollectionNameType> Collections;
+			Collections.Add(CollectionClicked.GetValue());
+			CollectionViewPtr->SetSelectedCollections(Collections);
+
+			CollectionSelected(CollectionClicked.GetValue());
+		}
 	}
 	else if ( SourcesData.PackagePaths.Num() == 0 )
 	{
@@ -1283,8 +1305,52 @@ TSharedPtr<SWidget> SContentBrowser::OnGetCrumbDelimiterContent(const FString& C
 	FSourcesData SourcesData = AssetViewPtr->GetSourcesData();
 
 	TSharedPtr<SWidget> Widget = SNullWidget::NullWidget;
+	TSharedPtr<SWidget> MenuWidget;
 
-	if( SourcesData.PackagePaths.Num() > 0 )
+	if( SourcesData.Collections.Num() > 0 )
+	{
+		TOptional<FCollectionNameType> CollectionClicked;
+		{
+			FString CollectionName;
+			FString CollectionTypeString;
+			if (CrumbData.Split(TEXT("?"), &CollectionName, &CollectionTypeString))
+			{
+				const int32 CollectionType = FCString::Atoi(*CollectionTypeString);
+				if (CollectionType >= 0 && CollectionType < ECollectionShareType::CST_All)
+				{
+					CollectionClicked = FCollectionNameType(FName(*CollectionName), ECollectionShareType::Type(CollectionType));
+				}
+			}
+		}
+
+		if( CollectionClicked.IsSet() )
+		{
+			FCollectionManagerModule& CollectionManagerModule = FCollectionManagerModule::GetModule();
+
+			TArray<FCollectionNameType> ChildCollections;
+			CollectionManagerModule.Get().GetChildCollections(CollectionClicked->Name, CollectionClicked->Type, ChildCollections);
+
+			if( ChildCollections.Num() > 0 )
+			{
+				FMenuBuilder MenuBuilder( true, nullptr );
+
+				for( const FCollectionNameType& ChildCollection : ChildCollections )
+				{
+					const FString ChildCollectionCrumbData = FString::Printf(TEXT("%s?%s"), *ChildCollection.Name.ToString(), *FString::FromInt(ChildCollection.Type));
+
+					MenuBuilder.AddMenuEntry(
+						FText::FromName(ChildCollection.Name),
+						FText::GetEmpty(),
+						FSlateIcon(FEditorStyle::GetStyleSetName(), ECollectionShareType::GetIconStyleName(ChildCollection.Type)),
+						FUIAction(FExecuteAction::CreateSP(this, &SContentBrowser::OnPathMenuItemClicked, ChildCollectionCrumbData))
+						);
+				}
+
+				MenuWidget = MenuBuilder.MakeWidget();
+			}
+		}
+	}
+	else if( SourcesData.PackagePaths.Num() > 0 )
 	{
 		TArray<FString> SubPaths;
 		const bool bRecurse = false;
@@ -1308,7 +1374,7 @@ TSharedPtr<SWidget> SContentBrowser::OnGetCrumbDelimiterContent(const FString& C
 
 		if( SubPaths.Num() > 0 )
 		{
-			FMenuBuilder MenuBuilder( true, NULL );
+			FMenuBuilder MenuBuilder( true, nullptr );
 
 			for( int32 PathIndex = 0; PathIndex < SubPaths.Num(); ++PathIndex )
 			{
@@ -1320,21 +1386,26 @@ TSharedPtr<SWidget> SContentBrowser::OnGetCrumbDelimiterContent(const FString& C
 					FText::FromString(PathWithoutParent),
 					FText::GetEmpty(),
 					FSlateIcon(FEditorStyle::GetStyleSetName(), "ContentBrowser.BreadcrumbPathPickerFolder"),
-					FUIAction(FExecuteAction::CreateSP(this, &SContentBrowser::OnPathMenuItemClicked, SubPath)));
+					FUIAction(FExecuteAction::CreateSP(this, &SContentBrowser::OnPathMenuItemClicked, SubPath))
+					);
 			}
 
-
-			// Do not allow the menu to become too large if there are many directories
-			Widget =
-				SNew( SVerticalBox )
-				+SVerticalBox::Slot()
-				.MaxHeight( 400.0f )
-				[
-					MenuBuilder.MakeWidget()
-				];
+			MenuWidget = MenuBuilder.MakeWidget();
 		}
 	}
 
+	if( MenuWidget.IsValid() )
+	{
+		// Do not allow the menu to become too large if there are many directories
+		Widget =
+			SNew( SVerticalBox )
+			+SVerticalBox::Slot()
+			.MaxHeight( 400.0f )
+			[
+				MenuWidget.ToSharedRef()
+			];
+	}
+
 	return Widget;
 }
 
@@ -1913,15 +1984,29 @@ void SContentBrowser::UpdatePath()
 	}
 	else if ( SourcesData.Collections.Num() > 0 )
 	{
-		const FString CollectionName = SourcesData.Collections[0].Name.ToString();
-		const FString CollectionType = FString::FromInt(SourcesData.Collections[0].Type);
-		const FString CrumbData = CollectionName + TEXT("?") + CollectionType;
+		FCollectionManagerModule& CollectionManagerModule = FCollectionManagerModule::GetModule();
+		TArray<FCollectionNameType> CollectionPathItems;
+
+		// Walk up the parents of this collection so that we can generate a complete path (this loop also adds the child collection to the array)
+		for (TOptional<FCollectionNameType> CurrentCollection = SourcesData.Collections[0]; 
+			CurrentCollection.IsSet(); 
+			CurrentCollection = CollectionManagerModule.Get().GetParentCollection(CurrentCollection->Name, CurrentCollection->Type)
+			)
+		{
+			CollectionPathItems.Insert(CurrentCollection.GetValue(), 0);
+		}
 
-		FFormatNamedArguments Args;
-		Args.Add(TEXT("CollectionName"), FText::FromString(CollectionName));
-		const FText DisplayName = FText::Format(LOCTEXT("CollectionPathIndicator", "{CollectionName} (Collection)"), Args);
+		// Now add each part of the path to the breadcrumb trail
+		for (const FCollectionNameType& CollectionPathItem : CollectionPathItems)
+		{
+			const FString CrumbData = FString::Printf(TEXT("%s?%s"), *CollectionPathItem.Name.ToString(), *FString::FromInt(CollectionPathItem.Type));
 
-		PathBreadcrumbTrail->PushCrumb(DisplayName, CrumbData);
+			FFormatNamedArguments Args;
+			Args.Add(TEXT("CollectionName"), FText::FromName(CollectionPathItem.Name));
+			const FText DisplayName = FText::Format(LOCTEXT("CollectionPathIndicator", "{CollectionName} (Collection)"), Args);
+
+			PathBreadcrumbTrail->PushCrumb(DisplayName, CrumbData);
+		}
 	}
 	else
 	{
