Stopped you being able to drag the collection "folders" in the asset view

diff --git a/Engine/Source/Editor/ContentBrowser/Private/SAssetView.cpp b/Engine/Source/Editor/ContentBrowser/Private/SAssetView.cpp
--- a/Engine/Source/Editor/ContentBrowser/Private/SAssetView.cpp
+++ b/Engine/Source/Editor/ContentBrowser/Private/SAssetView.cpp
@@ -3306,7 +3306,7 @@ FReply SAssetView::OnDraggingAssetItem( const FGeometry& MyGeometry, const FPoin
 		{
 			// are we dragging some folders?
 			TArray<FString> SelectedFolders = GetSelectedFolders();
-			if(SelectedFolders.Num() > 0)
+			if(SelectedFolders.Num() > 0 && SourcesData.Collections.Num() == 0)
 			{
 				return FReply::Handled().BeginDragDrop(FAssetPathDragDropOp::New(SelectedFolders));
 			}
