The "New..." option in a collection context menu will now create a child collection

UE-17636 - Creating a new Collection always placed at root level

The + button in the collections view will still create root level collections.

diff --git a/Engine/Source/Editor/ContentBrowser/Private/CollectionContextMenu.cpp b/Engine/Source/Editor/ContentBrowser/Private/CollectionContextMenu.cpp
--- a/Engine/Source/Editor/ContentBrowser/Private/CollectionContextMenu.cpp
+++ b/Engine/Source/Editor/ContentBrowser/Private/CollectionContextMenu.cpp
@@ -50,12 +50,15 @@ TSharedPtr<SWidget> FCollectionContextMenu::MakeCollectionTreeContextMenu(TShare
 
 	MenuBuilder.BeginSection("CollectionOptions", LOCTEXT("CollectionListOptionsMenuHeading", "Collection Options"));
 	{
-		// New... (submenu)
-		MenuBuilder.AddSubMenu(
-			LOCTEXT("NewCollection", "New..."),
-			LOCTEXT("NewCollectionTooltip", "Create a collection."),
-			FNewMenuDelegate::CreateRaw( this, &FCollectionContextMenu::MakeNewCollectionSubMenu )
-			);
+		if ( SelectedCollections.Num() == 1 )
+		{
+			// New... (submenu)
+			MenuBuilder.AddSubMenu(
+				LOCTEXT("NewChildCollection", "New..."),
+				LOCTEXT("NewChildCollectionTooltip", "Create child a collection."),
+				FNewMenuDelegate::CreateRaw( this, &FCollectionContextMenu::MakeNewCollectionSubMenu, TOptional<FCollectionNameType>( FCollectionNameType( SelectedCollections[0]->CollectionName, SelectedCollections[0]->CollectionType ) ) )
+				);
+		}
 
 		// Only add rename/delete if at least one collection is selected
 		if ( SelectedCollections.Num() )
@@ -125,37 +128,42 @@ TSharedPtr<SWidget> FCollectionContextMenu::MakeCollectionTreeContextMenu(TShare
 	return MenuBuilder.MakeWidget();
 }
 
-void FCollectionContextMenu::MakeNewCollectionSubMenu(FMenuBuilder& MenuBuilder)
+void FCollectionContextMenu::MakeNewCollectionSubMenu(FMenuBuilder& MenuBuilder, TOptional<FCollectionNameType> ParentCollection)
 {
-	MenuBuilder.BeginSection("CollectionNewCollection", LOCTEXT("NewCollectionMenuHeading", "New Collection"));
+	const FText MenuHeading = (ParentCollection.IsSet()) ? LOCTEXT("NewChildCollectionMenuHeading", "New Child Collection"): LOCTEXT("NewCollectionMenuHeading", "New Collection");
+
+	MenuBuilder.BeginSection("CollectionNewCollection", MenuHeading);
 	{
+		const bool bCanCreateSharedChildren = !ParentCollection.IsSet() || ECollectionShareType::IsValidChildType( ParentCollection->Type, ECollectionShareType::CST_Shared );
 		MenuBuilder.AddMenuEntry(
 			LOCTEXT("NewCollection_Shared", "Shared Collection"),
 			LOCTEXT("NewCollection_SharedTooltip", "Create a collection that can be seen by anyone."),
 			FSlateIcon( FEditorStyle::GetStyleSetName(), ECollectionShareType::GetIconStyleName( ECollectionShareType::CST_Shared ) ),
 			FUIAction(
-				FExecuteAction::CreateSP( this, &FCollectionContextMenu::ExecuteNewCollection, ECollectionShareType::CST_Shared ),
-				FCanExecuteAction::CreateSP( this, &FCollectionContextMenu::CanExecuteNewCollection, ECollectionShareType::CST_Shared )
+				FExecuteAction::CreateSP( this, &FCollectionContextMenu::ExecuteNewCollection, ECollectionShareType::CST_Shared, ParentCollection ),
+				FCanExecuteAction::CreateSP( this, &FCollectionContextMenu::CanExecuteNewCollection, ECollectionShareType::CST_Shared, bCanCreateSharedChildren )
 				)
 			);
 
+		const bool bCanCreatePrivateChildren = !ParentCollection.IsSet() || ECollectionShareType::IsValidChildType( ParentCollection->Type, ECollectionShareType::CST_Private );
 		MenuBuilder.AddMenuEntry(
 			LOCTEXT("NewCollection_Private", "Private Collection"),
 			LOCTEXT("NewCollection_PrivateTooltip", "Create a collection that can only be seen by you."),
 			FSlateIcon( FEditorStyle::GetStyleSetName(), ECollectionShareType::GetIconStyleName( ECollectionShareType::CST_Private ) ),
 			FUIAction(
-				FExecuteAction::CreateSP( this, &FCollectionContextMenu::ExecuteNewCollection, ECollectionShareType::CST_Private ),
-				FCanExecuteAction::CreateSP( this, &FCollectionContextMenu::CanExecuteNewCollection, ECollectionShareType::CST_Private )
+				FExecuteAction::CreateSP( this, &FCollectionContextMenu::ExecuteNewCollection, ECollectionShareType::CST_Private, ParentCollection ),
+				FCanExecuteAction::CreateSP( this, &FCollectionContextMenu::CanExecuteNewCollection, ECollectionShareType::CST_Private, bCanCreatePrivateChildren )
 				)
 			);
 
+		const bool bCanCreateLocalChildren = !ParentCollection.IsSet() || ECollectionShareType::IsValidChildType( ParentCollection->Type, ECollectionShareType::CST_Local );
 		MenuBuilder.AddMenuEntry(
 			LOCTEXT("NewCollection_Local", "Local Collection"),
 			LOCTEXT("NewCollection_LocalTooltip", "Create a collection that is not in source control and can only be seen by you."),
 			FSlateIcon( FEditorStyle::GetStyleSetName(), ECollectionShareType::GetIconStyleName( ECollectionShareType::CST_Local ) ),
 			FUIAction(
-				FExecuteAction::CreateSP( this, &FCollectionContextMenu::ExecuteNewCollection, ECollectionShareType::CST_Local ),
-				FCanExecuteAction::CreateSP( this, &FCollectionContextMenu::CanExecuteNewCollection, ECollectionShareType::CST_Local )
+				FExecuteAction::CreateSP( this, &FCollectionContextMenu::ExecuteNewCollection, ECollectionShareType::CST_Local, ParentCollection ),
+				FCanExecuteAction::CreateSP( this, &FCollectionContextMenu::CanExecuteNewCollection, ECollectionShareType::CST_Local, bCanCreateLocalChildren )
 				)
 			);
 	}
@@ -285,14 +293,14 @@ bool FCollectionContextMenu::CanRenameSelectedCollections() const
 	return false;
 }
 
-void FCollectionContextMenu::ExecuteNewCollection(ECollectionShareType::Type CollectionType)
+void FCollectionContextMenu::ExecuteNewCollection(ECollectionShareType::Type CollectionType, TOptional<FCollectionNameType> ParentCollection)
 {
 	if ( !ensure(CollectionView.IsValid()) )
 	{
 		return;
 	}
 
-	CollectionView.Pin()->CreateCollectionItem(CollectionType);
+	CollectionView.Pin()->CreateCollectionItem(CollectionType, ParentCollection);
 }
 
 void FCollectionContextMenu::ExecuteSetCollectionShareType(ECollectionShareType::Type CollectionType)
@@ -407,9 +415,9 @@ void FCollectionContextMenu::ExecutePickColor()
 	OpenColorPicker(PickerArgs);
 }
 
-bool FCollectionContextMenu::CanExecuteNewCollection(ECollectionShareType::Type CollectionType) const
+bool FCollectionContextMenu::CanExecuteNewCollection(ECollectionShareType::Type CollectionType, bool bIsValidChildType) const
 {
-	return (CollectionType == ECollectionShareType::CST_Local) || (bProjectUnderSourceControl && ISourceControlModule::Get().IsEnabled() && ISourceControlModule::Get().GetProvider().IsAvailable());
+	return bIsValidChildType && (CollectionType == ECollectionShareType::CST_Local || (bProjectUnderSourceControl && ISourceControlModule::Get().IsEnabled() && ISourceControlModule::Get().GetProvider().IsAvailable()));
 }
 
 bool FCollectionContextMenu::CanExecuteSetCollectionShareType(ECollectionShareType::Type CollectionType) const
diff --git a/Engine/Source/Editor/ContentBrowser/Private/CollectionContextMenu.h b/Engine/Source/Editor/ContentBrowser/Private/CollectionContextMenu.h
--- a/Engine/Source/Editor/ContentBrowser/Private/CollectionContextMenu.h
+++ b/Engine/Source/Editor/ContentBrowser/Private/CollectionContextMenu.h
@@ -15,7 +15,7 @@ public:
 	TSharedPtr<SWidget> MakeCollectionTreeContextMenu(TSharedPtr< FUICommandList > InCommandList);
 
 	/** Makes the new collection submenu */
-	void MakeNewCollectionSubMenu(FMenuBuilder& MenuBuilder);
+	void MakeNewCollectionSubMenu(FMenuBuilder& MenuBuilder, TOptional<FCollectionNameType> ParentCollection);
 
 	/** Makes the collection share type submenu */
 	void MakeCollectionShareTypeSubMenu(FMenuBuilder& MenuBuilder);
@@ -32,7 +32,7 @@ protected:
 
 private:
 	/** Handler for when a collection is selected in the "New" menu */
-	void ExecuteNewCollection(ECollectionShareType::Type CollectionType);
+	void ExecuteNewCollection(ECollectionShareType::Type CollectionType, TOptional<FCollectionNameType> ParentCollection);
 
 	/** Handler for when a collection share type is changed in the "Share Type" menu */
 	void ExecuteSetCollectionShareType(ECollectionShareType::Type CollectionType);
@@ -53,7 +53,7 @@ private:
 	void ExecutePickColor();
 
 	/** Handler to check to see if "New Collection" can be executed */
-	bool CanExecuteNewCollection(ECollectionShareType::Type CollectionType) const;
+	bool CanExecuteNewCollection(ECollectionShareType::Type CollectionType, bool bIsValidChildType) const;
 
 	/** Handler to check to see if a entry in the "Share Type" menu can be executed */
 	bool CanExecuteSetCollectionShareType(ECollectionShareType::Type CollectionType) const;
diff --git a/Engine/Source/Editor/ContentBrowser/Private/SCollectionView.cpp b/Engine/Source/Editor/ContentBrowser/Private/SCollectionView.cpp
--- a/Engine/Source/Editor/ContentBrowser/Private/SCollectionView.cpp
+++ b/Engine/Source/Editor/ContentBrowser/Private/SCollectionView.cpp
@@ -677,7 +677,7 @@ FReply SCollectionView::MakeAddCollectionMenu()
 
 	CollectionContextMenu->UpdateProjectSourceControl();
 
-	CollectionContextMenu->MakeNewCollectionSubMenu(MenuBuilder);
+	CollectionContextMenu->MakeNewCollectionSubMenu(MenuBuilder, TOptional<FCollectionNameType>());
 
 	FSlateApplication::Get().PushMenu(
 		AsShared(),
@@ -707,7 +707,7 @@ EVisibility SCollectionView::GetAddCollectionButtonVisibility() const
 	return (bAllowCollectionButtons && ( !CollectionsExpandableAreaPtr.IsValid() || CollectionsExpandableAreaPtr->IsExpanded() ) ) ? EVisibility::Visible : EVisibility::Collapsed;
 }
 
-void SCollectionView::CreateCollectionItem( ECollectionShareType::Type CollectionType )
+void SCollectionView::CreateCollectionItem( ECollectionShareType::Type CollectionType, TOptional<FCollectionNameType> ParentCollection )
 {
 	if ( ensure(CollectionType != ECollectionShareType::CST_All) )
 	{
@@ -721,6 +721,19 @@ void SCollectionView::CreateCollectionItem( ECollectionShareType::Type Collectio
 		// Adding a new collection now, so clear any filter we may have applied
 		SearchBoxPtr->SetText(FText::GetEmpty());
 
+		if ( ParentCollection.IsSet() )
+		{
+			TSharedPtr<FCollectionItem> ParentCollectionItem = AvailableCollections.FindRef(ParentCollection.GetValue());
+			if ( ParentCollectionItem.IsValid() )
+			{
+				ParentCollectionItem->ChildCollections.Add(NewItem);
+				NewItem->ParentCollection = ParentCollectionItem;
+
+				// Make sure the parent is expanded so we can see its newly added child item
+				CollectionTreePtr->SetItemExpansion(ParentCollectionItem, true);
+			}
+		}
+
 		// Mark the new collection for rename and that it is new so it will be created upon successful rename
 		NewItem->bRenaming = true;
 		NewItem->bNewCollection = true;
@@ -1239,8 +1252,16 @@ bool SCollectionView::CollectionNameChangeCommit( const TSharedPtr< FCollectionI
 	if ( CollectionItem->bNewCollection )
 	{
 		CollectionItem->bNewCollection = false;
+		
+		// Cache this here as CreateCollection will invalidate the current parent pointer
+		TOptional<FCollectionNameType> NewCollectionParentKey;
+		TSharedPtr<FCollectionItem> ParentCollectionItem = CollectionItem->ParentCollection.Pin();
+		if ( ParentCollectionItem.IsValid() )
+		{
+			NewCollectionParentKey = FCollectionNameType(ParentCollectionItem->CollectionName, ParentCollectionItem->CollectionType);
+		}
 
-		// If we can canceled the name change when creating a new asset, we want to silently remove it
+		// If we canceled the name change when creating a new asset, we want to silently remove it
 		if ( !bChangeConfirmed )
 		{
 			AvailableCollections.Remove(FCollectionNameType(CollectionItem->CollectionName, CollectionItem->CollectionType));
@@ -1257,6 +1278,12 @@ bool SCollectionView::CollectionNameChangeCommit( const TSharedPtr< FCollectionI
 			OutWarningMessage = FText::Format( LOCTEXT("CreateCollectionFailed", "Failed to create the collection. {0}"), CollectionManagerModule.Get().GetLastError());
 			return false;
 		}
+
+		if ( NewCollectionParentKey.IsSet() )
+		{
+			// Try and set the parent correctly (if this fails for any reason, the collection will still be added, but will just appear at the root)
+			CollectionManagerModule.Get().ReparentCollection(NewNameFinal, CollectionType, NewCollectionParentKey->Name, NewCollectionParentKey->Type);
+		}
 	}
 	else
 	{
diff --git a/Engine/Source/Editor/ContentBrowser/Private/SCollectionView.h b/Engine/Source/Editor/ContentBrowser/Private/SCollectionView.h
--- a/Engine/Source/Editor/ContentBrowser/Private/SCollectionView.h
+++ b/Engine/Source/Editor/ContentBrowser/Private/SCollectionView.h
@@ -94,7 +94,7 @@ private:
 	EVisibility GetAddCollectionButtonVisibility() const;
 	
 	/** Sets up an inline creation process for a new collection of the specified type */
-	void CreateCollectionItem( ECollectionShareType::Type CollectionType );
+	void CreateCollectionItem( ECollectionShareType::Type CollectionType, TOptional<FCollectionNameType> ParentCollection );
 
 	/** Sets up an inline rename for the specified collection */
 	void RenameCollectionItem( const TSharedPtr<FCollectionItem>& ItemToRename );
