Fixed SPathView displaying some root paths incorrectly when dealing with templates

Templates can mount "root" items which are actually sub-items under a root (see FUnrealEdMisc::MountTemplateSharedPaths), and SPathView was previously displaying those paths as proper root items, rather than nesting them under their actual root.

diff --git a/Engine/Source/Editor/ContentBrowser/Private/SPathView.cpp b/Engine/Source/Editor/ContentBrowser/Private/SPathView.cpp
--- a/Engine/Source/Editor/ContentBrowser/Private/SPathView.cpp
+++ b/Engine/Source/Editor/ContentBrowser/Private/SPathView.cpp
@@ -949,7 +949,10 @@ void SPathView::Populate()
 			{
 				CleanRootPathName = CleanRootPathName.Mid( 0, CleanRootPathName.Len() - 1 );
 			}
-			AddRootItem(CleanRootPathName);
+
+			// Templates can mount "root" items which are actually sub-items under a root (see FUnrealEdMisc::MountTemplateSharedPaths)
+			// We can use AddPath here (rather than AddRootItem), as this will ensure that both the root and any sub-path items are added correctly
+			AddPath(CleanRootPathName);
 		}
 	}
 	
