Fixed being able to rename or deleted some read-only types in the Content Browser

UE-13653 - Able to rename redirectors by clicking on their name in the content browser twice
UE-18082 - Hide the delete button when selecting a code class in the content browser

You can no longer rename classes, redirectors, or collections by clicking on their name in the Content Browser. The delegate handling whether the asset name was read-only hadn't been hooked into the "can rename" logic used by the context menu.

Additionally, you can no longer delete a collection via the asset view, and you no longer see some asset specific context menu options when you only have classes selected.

diff --git a/Engine/Source/Editor/ContentBrowser/Private/AssetContextMenu.cpp b/Engine/Source/Editor/ContentBrowser/Private/AssetContextMenu.cpp
--- a/Engine/Source/Editor/ContentBrowser/Private/AssetContextMenu.cpp
+++ b/Engine/Source/Editor/ContentBrowser/Private/AssetContextMenu.cpp
@@ -203,6 +203,11 @@ bool FAssetContextMenu::AddImportedAssetMenuOptions(FMenuBuilder& MenuBuilder)
 
 bool FAssetContextMenu::AddCommonMenuOptions(FMenuBuilder& MenuBuilder)
 {
+	TArray< FAssetData > AssetViewSelectedAssets = AssetView.Pin()->GetSelectedAssets();
+
+	int32 NumAssetItems, NumClassItems;
+	ContentBrowserUtils::CountItemTypes(AssetViewSelectedAssets, NumAssetItems, NumClassItems);
+
 	MenuBuilder.BeginSection("CommonAssetActions", LOCTEXT("CommonAssetActionsMenuHeading", "Common"));
 	{
 		// Edit
@@ -213,52 +218,56 @@ bool FAssetContextMenu::AddCommonMenuOptions(FMenuBuilder& MenuBuilder)
 			FUIAction( FExecuteAction::CreateSP(this, &FAssetContextMenu::ExecuteEditAsset) )
 			);
 	
-		// Rename
-		MenuBuilder.AddMenuEntry(FGenericCommands::Get().Rename, NAME_None,
-			LOCTEXT("Rename", "Rename"),
-			LOCTEXT("RenameTooltip", "Rename the selected asset."),
-			FSlateIcon(FEditorStyle::GetStyleSetName(), "ContentBrowser.AssetActions.Rename")
-			);
+		// Only add these options if assets are selected
+		if (NumAssetItems > 0)
+		{
+			// Rename
+			MenuBuilder.AddMenuEntry(FGenericCommands::Get().Rename, NAME_None,
+				LOCTEXT("Rename", "Rename"),
+				LOCTEXT("RenameTooltip", "Rename the selected asset."),
+				FSlateIcon(FEditorStyle::GetStyleSetName(), "ContentBrowser.AssetActions.Rename")
+				);
 
-		// Duplicate
-		MenuBuilder.AddMenuEntry(FGenericCommands::Get().Duplicate, NAME_None,
-			LOCTEXT("Duplicate", "Duplicate"),
-			LOCTEXT("DuplicateTooltip", "Create a copy of the selected asset(s)."),
-			FSlateIcon(FEditorStyle::GetStyleSetName(), "ContentBrowser.AssetActions.Duplicate")
-			);
+			// Duplicate
+			MenuBuilder.AddMenuEntry(FGenericCommands::Get().Duplicate, NAME_None,
+				LOCTEXT("Duplicate", "Duplicate"),
+				LOCTEXT("DuplicateTooltip", "Create a copy of the selected asset(s)."),
+				FSlateIcon(FEditorStyle::GetStyleSetName(), "ContentBrowser.AssetActions.Duplicate")
+				);
 
-		// Save
-		MenuBuilder.AddMenuEntry(
-			LOCTEXT("SaveAsset", "Save"),
-			LOCTEXT("SaveAssetTooltip", "Saves the asset to file."),
-			FSlateIcon(FEditorStyle::GetStyleSetName(), "Level.SaveIcon16x"),
-			FUIAction(
-				FExecuteAction::CreateSP( this, &FAssetContextMenu::ExecuteSaveAsset ),
-				FCanExecuteAction::CreateSP( this, &FAssetContextMenu::CanExecuteSaveAsset )
-				)
-			);
+			// Save
+			MenuBuilder.AddMenuEntry(
+				LOCTEXT("SaveAsset", "Save"),
+				LOCTEXT("SaveAssetTooltip", "Saves the asset to file."),
+				FSlateIcon(FEditorStyle::GetStyleSetName(), "Level.SaveIcon16x"),
+				FUIAction(
+					FExecuteAction::CreateSP( this, &FAssetContextMenu::ExecuteSaveAsset ),
+					FCanExecuteAction::CreateSP( this, &FAssetContextMenu::CanExecuteSaveAsset )
+					)
+				);
 
-		// Delete
-		MenuBuilder.AddMenuEntry(FGenericCommands::Get().Delete, NAME_None,
-			LOCTEXT("Delete", "Delete"),
-			LOCTEXT("DeleteTooltip", "Delete the selected assets."),
-			FSlateIcon(FEditorStyle::GetStyleSetName(), "ContentBrowser.AssetActions.Delete")
-			);
+			// Delete
+			MenuBuilder.AddMenuEntry(FGenericCommands::Get().Delete, NAME_None,
+				LOCTEXT("Delete", "Delete"),
+				LOCTEXT("DeleteTooltip", "Delete the selected assets."),
+				FSlateIcon(FEditorStyle::GetStyleSetName(), "ContentBrowser.AssetActions.Delete")
+				);
 
-		// Asset Actions sub-menu
-		MenuBuilder.AddSubMenu(
-			LOCTEXT("AssetActionsSubMenuLabel", "Asset Actions"),
-			LOCTEXT("AssetActionsSubMenuToolTip", "Other asset actions"),
-			FNewMenuDelegate::CreateSP(this, &FAssetContextMenu::MakeAssetActionsSubMenu),
-			FUIAction(
-				FExecuteAction(),
-				FCanExecuteAction::CreateSP( this, &FAssetContextMenu::CanExecuteAssetActions )
-				),
-			NAME_None,
-			EUserInterfaceActionType::Button,
-			false, 
-			FSlateIcon(FEditorStyle::GetStyleSetName(), "ContentBrowser.AssetActions")
-			);
+			// Asset Actions sub-menu
+			MenuBuilder.AddSubMenu(
+				LOCTEXT("AssetActionsSubMenuLabel", "Asset Actions"),
+				LOCTEXT("AssetActionsSubMenuToolTip", "Other asset actions"),
+				FNewMenuDelegate::CreateSP(this, &FAssetContextMenu::MakeAssetActionsSubMenu),
+				FUIAction(
+					FExecuteAction(),
+					FCanExecuteAction::CreateSP( this, &FAssetContextMenu::CanExecuteAssetActions )
+					),
+				NAME_None,
+				EUserInterfaceActionType::Button,
+				false, 
+				FSlateIcon(FEditorStyle::GetStyleSetName(), "ContentBrowser.AssetActions")
+				);
+		}
 	}
 	MenuBuilder.EndSection();
 
@@ -1835,11 +1844,10 @@ bool FAssetContextMenu::CanExecuteRename() const
 	TArray< FString > SelectedFolders = AssetView.Pin()->GetSelectedFolders();
 
 	const bool bOneAssetSelected = AssetViewSelectedAssets.Num() == 1 && SelectedFolders.Num() == 0		// A single asset
-		&& AssetViewSelectedAssets[0].AssetClass != UObjectRedirector::StaticClass()->GetFName()		// Which isn't a redirector
-		&& AssetViewSelectedAssets[0].AssetClass != NAME_Class;											// And isn't a class
+		&& ContentBrowserUtils::CanRenameAsset(AssetViewSelectedAssets[0]);								// Which can be renamed
 
 	const bool bOneFolderSelected = AssetViewSelectedAssets.Num() == 0 && SelectedFolders.Num() == 1	// A single folder
-		&& !ContentBrowserUtils::IsClassPath(SelectedFolders[0]);										// Which doesn't belong to a class path
+		&& ContentBrowserUtils::CanRenameFolder(SelectedFolders[0]);									// Which can be renamed
 	
 	return (bOneAssetSelected || bOneFolderSelected) && !AssetView.Pin()->IsThumbnailEditMode();
 }
@@ -1855,8 +1863,18 @@ bool FAssetContextMenu::CanExecuteDelete() const
 	int32 NumAssetPaths, NumClassPaths;
 	ContentBrowserUtils::CountPathTypes(SelectedFolders, NumAssetPaths, NumClassPaths);
 
-	// We can't delete classes, or folders containing classes
-	return (NumAssetItems > 0 && NumClassItems == 0) || (NumAssetPaths > 0 && NumClassPaths == 0);
+	bool bHasSelectedCollections = false;
+	for (const FString& SelectedFolder : SelectedFolders)
+	{
+		if (ContentBrowserUtils::IsCollectionPath(SelectedFolder))
+		{
+			bHasSelectedCollections = true;
+			break;
+		}
+	}
+
+	// We can't delete classes, or folders containing classes, or any collection folders
+	return ((NumAssetItems > 0 && NumClassItems == 0) || (NumAssetPaths > 0 && NumClassPaths == 0)) && !bHasSelectedCollections;
 }
 
 bool FAssetContextMenu::CanExecuteRemoveFromCollection() const 
diff --git a/Engine/Source/Editor/ContentBrowser/Private/AssetViewWidgets.cpp b/Engine/Source/Editor/ContentBrowser/Private/AssetViewWidgets.cpp
--- a/Engine/Source/Editor/ContentBrowser/Private/AssetViewWidgets.cpp
+++ b/Engine/Source/Editor/ContentBrowser/Private/AssetViewWidgets.cpp
@@ -469,6 +469,34 @@ FReply SAssetViewItem::OnDrop( const FGeometry& MyGeometry, const FDragDropEvent
 	return FReply::Unhandled();
 }
 
+bool SAssetViewItem::IsNameReadOnly() const
+{
+	if (ThumbnailEditMode.Get())
+	{
+		// Read-only while editing thumbnails
+		return true;
+	}
+
+	if (!AssetItem.IsValid())
+	{
+		// Read-only if no valid asset item
+		return true;
+	}
+
+	if (AssetItem->GetType() != EAssetItemType::Folder)
+	{
+		// Read-only if we can't be renamed
+		return !ContentBrowserUtils::CanRenameAsset(StaticCastSharedPtr<FAssetViewAsset>(AssetItem)->Data);
+	}
+	else
+	{
+		// Read-only if we can't be renamed
+		return !ContentBrowserUtils::CanRenameFolder(StaticCastSharedPtr<FAssetViewFolder>(AssetItem)->FolderPath);
+	}
+
+	return false;
+}
+
 void SAssetViewItem::HandleBeginNameChange( const FText& OriginalText )
 {
 	OnRenameBegin.ExecuteIfBound(AssetItem, OriginalText.ToString(), LastGeometry.GetClippingRect());
@@ -1196,7 +1224,7 @@ void SAssetListItem::Construct( const FArguments& InArgs )
 					.OnVerifyTextChanged(this, &SAssetListItem::HandleVerifyNameChanged)
 					.HighlightText(InArgs._HighlightText)
 					.IsSelected(InArgs._IsSelected)
-					.IsReadOnly(ThumbnailEditMode)
+					.IsReadOnly(this, &SAssetListItem::IsNameReadOnly)
 				]
 
 				+SVerticalBox::Slot()
@@ -1344,7 +1372,7 @@ void SAssetTileItem::Construct( const FArguments& InArgs )
 					.OnVerifyTextChanged(this, &SAssetTileItem::HandleVerifyNameChanged)
 					.HighlightText(InArgs._HighlightText)
 					.IsSelected(InArgs._IsSelected)
-					.IsReadOnly(ThumbnailEditMode)
+					.IsReadOnly(this, &SAssetTileItem::IsNameReadOnly)
 					.Justification(ETextJustify::Center)
 					.LineBreakPolicy(FBreakIterator::CreateCamelCaseBreakIterator())
 			]
@@ -1567,7 +1595,7 @@ TSharedRef<SWidget> SAssetColumnItem::GenerateWidgetForColumn( const FName& Colu
 				.OnVerifyTextChanged(this, &SAssetColumnItem::HandleVerifyNameChanged)
 				.HighlightText(HighlightText)
 				.IsSelected(InIsSelected)
-				.IsReadOnly(ThumbnailEditMode)
+				.IsReadOnly(this, &SAssetColumnItem::IsNameReadOnly)
 			];
 
 		if(AssetItem.IsValid())
diff --git a/Engine/Source/Editor/ContentBrowser/Private/AssetViewWidgets.h b/Engine/Source/Editor/ContentBrowser/Private/AssetViewWidgets.h
--- a/Engine/Source/Editor/ContentBrowser/Private/AssetViewWidgets.h
+++ b/Engine/Source/Editor/ContentBrowser/Private/AssetViewWidgets.h
@@ -144,6 +144,9 @@ protected:
 	/** Used by OnDragEnter, OnDragOver, and OnDrop to check and update the validity of the drag operation */
 	bool ValidateDragDrop( const FGeometry& MyGeometry, const FDragDropEvent& DragDropEvent, bool& OutIsKnownDragOperation ) const;
 
+	/** Check to see if the name should be read-only */
+	bool IsNameReadOnly() const;
+
 	/** Handles starting a name change */
 	virtual void HandleBeginNameChange( const FText& OriginalText );
 
diff --git a/Engine/Source/Editor/ContentBrowser/Private/ContentBrowserUtils.cpp b/Engine/Source/Editor/ContentBrowser/Private/ContentBrowserUtils.cpp
--- a/Engine/Source/Editor/ContentBrowser/Private/ContentBrowserUtils.cpp
+++ b/Engine/Source/Editor/ContentBrowser/Private/ContentBrowserUtils.cpp
@@ -405,6 +405,18 @@ bool ContentBrowserUtils::PromptToLoadAssets(const TArray<FString>& UnloadedObje
 	return bShouldLoadAssets;
 }
 
+bool ContentBrowserUtils::CanRenameFolder(const FString& InFolderPath)
+{
+	// Cannot rename folders that are part of a classes or collections root
+	return !ContentBrowserUtils::IsClassPath(InFolderPath) && !ContentBrowserUtils::IsCollectionPath(InFolderPath);
+}
+
+bool ContentBrowserUtils::CanRenameAsset(const FAssetData& InAssetData)
+{
+	// Cannot rename redirectors or classes
+	return !InAssetData.IsRedirector() && InAssetData.AssetClass != NAME_Class;
+}
+
 void ContentBrowserUtils::RenameAsset(UObject* Asset, const FString& NewName, FText& ErrorMessage)
 {
 	FAssetToolsModule& AssetToolsModule = FModuleManager::LoadModuleChecked<FAssetToolsModule>("AssetTools");
diff --git a/Engine/Source/Editor/ContentBrowser/Private/ContentBrowserUtils.h b/Engine/Source/Editor/ContentBrowser/Private/ContentBrowserUtils.h
--- a/Engine/Source/Editor/ContentBrowser/Private/ContentBrowserUtils.h
+++ b/Engine/Source/Editor/ContentBrowser/Private/ContentBrowserUtils.h
@@ -51,12 +51,15 @@ namespace ContentBrowserUtils
 	 */
 	bool PromptToLoadAssets(const TArray<FString>& UnloadedObjects);
 
+	/** Checks to see if the given folder can be renamed */
+	bool CanRenameFolder(const FString& InFolderPath);
+
+	/** Checks to see if the given asset can be renamed */
+	bool CanRenameAsset(const FAssetData& InAssetData);
+
 	/** Renames an asset */
 	void RenameAsset(UObject* Asset, const FString& NewName, FText& ErrorMessage);
 
-	/** Renames a folder */
-	void RenameFolder(const FString& FolderPath, const FString& NewName, FText& ErrorMessage);
-
 	/** 
 	 * Copies assets to a new path 
 	 * @param Assets The assets to copy
