Fixed the Content Browser assets not updating their SCC state correctly when the SCC provider is changed while the editor is running

diff --git a/Engine/Source/Editor/ContentBrowser/Private/AssetViewWidgets.cpp b/Engine/Source/Editor/ContentBrowser/Private/AssetViewWidgets.cpp
--- a/Engine/Source/Editor/ContentBrowser/Private/AssetViewWidgets.cpp
+++ b/Engine/Source/Editor/ContentBrowser/Private/AssetViewWidgets.cpp
@@ -347,13 +347,11 @@ void SAssetViewItem::Construct( const FArguments& InArgs )
 	// Set our tooltip - this will refresh each time it's opened to make sure it's up-to-date
 	SetToolTip(SNew(SAssetViewItemToolTip).AssetViewItem(SharedThis(this)));
 
-	// refresh SCC state icon
-	HandleSourceControlStateChanged();
-
 	SourceControlStateDelay = 0.0f;
 	bSourceControlStateRequested = false;
 
-	ISourceControlModule::Get().GetProvider().RegisterSourceControlStateChanged_Handle(FSourceControlStateChanged::FDelegate::CreateSP(this, &SAssetViewItem::HandleSourceControlStateChanged));
+	ISourceControlModule::Get().RegisterProviderChanged(FSourceControlProviderChanged::FDelegate::CreateSP(this, &SAssetViewItem::HandleSourceControlProviderChanged));
+	SourceControlStateChangedDelegateHandle = ISourceControlModule::Get().GetProvider().RegisterSourceControlStateChanged_Handle(FSourceControlStateChanged::FDelegate::CreateSP(this, &SAssetViewItem::HandleSourceControlStateChanged));
 
 	// Source control state may have already been cached, make sure the control is in sync with 
 	// cached state as the delegate is not going to be invoked again until source control state 
@@ -570,6 +568,19 @@ const FSlateBrush* SAssetViewItem::GetSCCStateImage() const
 	return SCCStateBrush;
 }
 
+void SAssetViewItem::HandleSourceControlProviderChanged(ISourceControlProvider& OldProvider, ISourceControlProvider& NewProvider)
+{
+	OldProvider.UnregisterSourceControlStateChanged_Handle(SourceControlStateChangedDelegateHandle);
+	SourceControlStateChangedDelegateHandle = NewProvider.RegisterSourceControlStateChanged_Handle(FSourceControlStateChanged::FDelegate::CreateSP(this, &SAssetViewItem::HandleSourceControlStateChanged));
+	
+	// Reset this so the state will be queried from the new provider on the next Tick
+	SourceControlStateDelay = 0.0f;
+	bSourceControlStateRequested = false;
+	SCCStateBrush = nullptr;
+	
+	HandleSourceControlStateChanged();
+}
+
 void SAssetViewItem::HandleSourceControlStateChanged()
 {
 	if ( ISourceControlModule::Get().IsEnabled() && AssetItem.IsValid() && (AssetItem->GetType() == EAssetItemType::Normal) && !AssetItem->IsTemporaryItem() && !FPackageName::IsScriptPackage(CachedPackageName) )
@@ -893,10 +904,10 @@ void SAssetViewItem::UpdateSourceControlState(float InDeltaTime)
 {
 	SourceControlStateDelay += InDeltaTime;
 
-	if ( !bSourceControlStateRequested && SourceControlStateDelay > 1.0f && AssetItem.IsValid() )
+	if ( !bSourceControlStateRequested && SourceControlStateDelay > 1.0f && ISourceControlModule::Get().IsEnabled() && AssetItem.IsValid() )
 	{
 		// Only update the SCC state for non-temporary asset items that aren't a built in script
-		if ( AssetItem.IsValid() && !AssetItem->IsTemporaryItem() && AssetItem->GetType() != EAssetItemType::Folder && !FPackageName::IsScriptPackage(CachedPackageName) )
+		if ( !AssetItem->IsTemporaryItem() && AssetItem->GetType() != EAssetItemType::Folder && !FPackageName::IsScriptPackage(CachedPackageName) )
 		{
 			// Request the most recent SCC state for this asset
 			ISourceControlModule::Get().QueueStatusUpdate(CachedPackageFileName);
diff --git a/Engine/Source/Editor/ContentBrowser/Private/AssetViewWidgets.h b/Engine/Source/Editor/ContentBrowser/Private/AssetViewWidgets.h
--- a/Engine/Source/Editor/ContentBrowser/Private/AssetViewWidgets.h
+++ b/Engine/Source/Editor/ContentBrowser/Private/AssetViewWidgets.h
@@ -210,6 +210,9 @@ protected:
 	/** Set mips to be resident while this (loaded) asset is visible */
 	void SetForceMipLevelsToBeResident(bool bForce) const;
 
+	/** Delegate handler for when the source control provider changes */
+	void HandleSourceControlProviderChanged(class ISourceControlProvider& OldProvider, class ISourceControlProvider& NewProvider);
+
 	/** Delegate handler for when source control state changes */
 	void HandleSourceControlStateChanged();
 
@@ -311,6 +314,9 @@ protected:
 
 	/** Cached brush for the source control state */
 	const FSlateBrush* SCCStateBrush;
+
+	/** Delegate handle for the HandleSourceControlStateChanged function callback */
+	FDelegateHandle SourceControlStateChangedDelegateHandle;
 };
 
 
