Made the collections view update lazy

This now queues the list view update until the next Tick to avoid potentially updating the list multiple times per-frame as collection manager updates come in.

diff --git a/Engine/Source/Editor/ContentBrowser/Private/SCollectionView.cpp b/Engine/Source/Editor/ContentBrowser/Private/SCollectionView.cpp
--- a/Engine/Source/Editor/ContentBrowser/Private/SCollectionView.cpp
+++ b/Engine/Source/Editor/ContentBrowser/Private/SCollectionView.cpp
@@ -67,6 +67,8 @@ void SCollectionView::Construct( const FArguments& InArgs )
 	bAllowRightClickMenu = InArgs._AllowRightClickMenu;
 	bAllowCollectionDrag = InArgs._AllowCollectionDrag;
 	bDraggedOver = false;
+
+	bQueueCollectionItemsUpdate = false;
 	bQueueSCCRefresh = true;
 
 	FCollectionManagerModule& CollectionManagerModule = FCollectionManagerModule::GetModule();
@@ -233,41 +235,30 @@ void SCollectionView::Construct( const FArguments& InArgs )
 
 void SCollectionView::HandleCollectionCreated( const FCollectionNameType& Collection )
 {
-	UpdateCollectionItems();
+	bQueueCollectionItemsUpdate = true;
 }
 
 void SCollectionView::HandleCollectionRenamed( const FCollectionNameType& OriginalCollection, const FCollectionNameType& NewCollection )
 {
-	// If the original collection was expanded, we want to pass that expansion state onto its new entry
-	bool bWasExpanded = false;
-	{
-		TSharedPtr<FCollectionItem> OriginalCollectionItem = AvailableCollections.FindRef(OriginalCollection);
-		if (OriginalCollectionItem.IsValid())
-		{
-			bWasExpanded = CollectionTreePtr->IsItemExpanded(OriginalCollectionItem);
-		}
-	}
-
-	UpdateCollectionItems();
+	bQueueCollectionItemsUpdate = true;
 
-	if (bWasExpanded)
+	// Rename the item in-place so we can maintain its expansion and selection states correctly once the view is refreshed on the next Tick
+	TSharedPtr<FCollectionItem> CollectionItem = AvailableCollections.FindRef(OriginalCollection);
+	if (CollectionItem.IsValid())
 	{
-		TSharedPtr<FCollectionItem> NewCollectionItem = AvailableCollections.FindRef(NewCollection);
-		if (NewCollectionItem.IsValid())
-		{
-			CollectionTreePtr->SetItemExpansion(NewCollectionItem, true);
-		}
+		CollectionItem->CollectionName = NewCollection.Name;
+		CollectionItem->CollectionType = NewCollection.Type;
 	}
 }
 
 void SCollectionView::HandleCollectionReparented( const FCollectionNameType& Collection, const TOptional<FCollectionNameType>& OldParent, const TOptional<FCollectionNameType>& NewParent )
 {
-	UpdateCollectionItems();
+	bQueueCollectionItemsUpdate = true;
 }
 
 void SCollectionView::HandleCollectionDestroyed( const FCollectionNameType& Collection )
 {
-	UpdateCollectionItems();
+	bQueueCollectionItemsUpdate = true;
 }
 
 void SCollectionView::HandleCollectionUpdated( const FCollectionNameType& Collection )
@@ -721,6 +712,12 @@ void SCollectionView::Tick( const FGeometry& AllottedGeometry, const double InCu
 {
 	SCompoundWidget::Tick(AllottedGeometry, InCurrentTime, InDeltaTime);
 
+	if (bQueueCollectionItemsUpdate)
+	{
+		bQueueCollectionItemsUpdate = false;
+		UpdateCollectionItems();
+	}
+
 	if (bQueueSCCRefresh && ISourceControlModule::Get().IsEnabled())
 	{
 		bQueueSCCRefresh = false;
diff --git a/Engine/Source/Editor/ContentBrowser/Private/SCollectionView.h b/Engine/Source/Editor/ContentBrowser/Private/SCollectionView.h
--- a/Engine/Source/Editor/ContentBrowser/Private/SCollectionView.h
+++ b/Engine/Source/Editor/ContentBrowser/Private/SCollectionView.h
@@ -283,6 +283,9 @@ private:
 	/** Delegate handle for the HandleSourceControlStateChanged function callback */
 	FDelegateHandle SourceControlStateChangedDelegateHandle;
 
+	/** True if we should queue a collection items update for the next Tick */
+	bool bQueueCollectionItemsUpdate;
+
 	/** True if we should queue an SCC refresh for the collections on the next Tick */
 	bool bQueueSCCRefresh;
 };
\ No newline at end of file
