Collection "folders" in the asset view now take their color from the collection they represent

UE-18242 - Collections folder icon does not update color in asset view

diff --git a/Engine/Source/Editor/ContentBrowser/Private/AssetViewTypes.h b/Engine/Source/Editor/ContentBrowser/Private/AssetViewTypes.h
--- a/Engine/Source/Editor/ContentBrowser/Private/AssetViewTypes.h
+++ b/Engine/Source/Editor/ContentBrowser/Private/AssetViewTypes.h
@@ -82,6 +82,9 @@ struct FAssetViewFolder : public FAssetViewItem
 	/** Whether this is a developer folder */
 	bool bDeveloperFolder;
 
+	/** Whether this is a collection folder */
+	bool bCollectionFolder;
+
 	/** Whether this folder is a new folder */
 	bool bNewFolder;
 
@@ -91,6 +94,7 @@ struct FAssetViewFolder : public FAssetViewItem
 	{
 		FolderName = FText::FromString(FPaths::GetBaseFilename(FolderPath));
 		bDeveloperFolder = ContentBrowserUtils::IsDevelopersFolder(FolderPath);
+		bCollectionFolder = ContentBrowserUtils::IsCollectionPath(FolderPath);
 	}
 
 	/** Set the name of this folder (without path) */
diff --git a/Engine/Source/Editor/ContentBrowser/Private/AssetViewWidgets.cpp b/Engine/Source/Editor/ContentBrowser/Private/AssetViewWidgets.cpp
--- a/Engine/Source/Editor/ContentBrowser/Private/AssetViewWidgets.cpp
+++ b/Engine/Source/Editor/ContentBrowser/Private/AssetViewWidgets.cpp
@@ -8,6 +8,7 @@
 #include "AssetViewWidgets.h"
 #include "SThumbnailEditModeTools.h"
 #include "CollectionManagerModule.h"
+#include "CollectionViewUtils.h"
 #include "DragAndDrop/AssetDragDropOp.h"
 #include "DragAndDrop/AssetPathDragDropOp.h"
 #include "DragDropHandler.h"
@@ -84,16 +85,17 @@ TSharedRef<SWidget> FAssetViewItemHelper::CreateListTileItemContents(T* const In
 
 		TSharedPtr<FAssetViewFolder> AssetFolderItem = StaticCastSharedPtr<FAssetViewFolder>(InTileOrListItem->AssetItem);
 
-		const bool bIsDeveloperFolder = AssetFolderItem->bDeveloperFolder;
-
 		ECollectionShareType::Type CollectionFolderShareType = ECollectionShareType::CST_All;
-		const bool bIsCollectionFolder = ContentBrowserUtils::IsCollectionPath(AssetFolderItem->FolderPath, nullptr, &CollectionFolderShareType);
+		if (AssetFolderItem->bCollectionFolder)
+		{
+			ContentBrowserUtils::IsCollectionPath(AssetFolderItem->FolderPath, nullptr, &CollectionFolderShareType);
+		}
 
-		const FSlateBrush* FolderBaseImage = bIsDeveloperFolder 
+		const FSlateBrush* FolderBaseImage = AssetFolderItem->bDeveloperFolder 
 			? FEditorStyle::GetBrush("ContentBrowser.ListViewDeveloperFolderIcon.Base") 
 			: FEditorStyle::GetBrush("ContentBrowser.ListViewFolderIcon.Base");
 
-		const FSlateBrush* FolderTintImage = bIsDeveloperFolder 
+		const FSlateBrush* FolderTintImage = AssetFolderItem->bDeveloperFolder 
 			? FEditorStyle::GetBrush("ContentBrowser.ListViewDeveloperFolderIcon.Mask") 
 			: FEditorStyle::GetBrush("ContentBrowser.ListViewFolderIcon.Mask");
 
@@ -105,7 +107,7 @@ TSharedRef<SWidget> FAssetViewItemHelper::CreateListTileItemContents(T* const In
 			.ColorAndOpacity(InTileOrListItem, &T::GetAssetColor)
 		];
 
-		if (bIsCollectionFolder)
+		if (AssetFolderItem->bCollectionFolder)
 		{
 			FLinearColor IconColor = FLinearColor::White;
 			switch(CollectionFolderShareType)
@@ -1060,7 +1062,22 @@ FSlateColor SAssetViewItem::GetAssetColor() const
 	{
 		if(AssetItem->GetType() == EAssetItemType::Folder)
 		{
-			const TSharedPtr<FLinearColor> Color = ContentBrowserUtils::LoadColor( StaticCastSharedPtr<FAssetViewFolder>(AssetItem)->FolderPath );
+			TSharedPtr<FAssetViewFolder> AssetFolderItem = StaticCastSharedPtr<FAssetViewFolder>(AssetItem);
+
+			TSharedPtr<FLinearColor> Color;
+			if (AssetFolderItem->bCollectionFolder)
+			{
+				FName CollectionName;
+				ECollectionShareType::Type CollectionFolderShareType = ECollectionShareType::CST_All;
+				ContentBrowserUtils::IsCollectionPath(AssetFolderItem->FolderPath, &CollectionName, &CollectionFolderShareType);
+
+				Color = CollectionViewUtils::LoadColor( CollectionName.ToString(), CollectionFolderShareType );
+			}
+			else
+			{
+				Color = ContentBrowserUtils::LoadColor( AssetFolderItem->FolderPath );
+			}
+
 			if ( Color.IsValid() )
 			{
 				return *Color.Get();
