Now cache out the split path when performing a text filter against an asset

This avoids doing repeated work when a query will perform multiple basic string tests.

diff --git a/Engine/Source/Editor/ContentBrowser/Private/FrontendFilters.cpp b/Engine/Source/Editor/ContentBrowser/Private/FrontendFilters.cpp
--- a/Engine/Source/Editor/ContentBrowser/Private/FrontendFilters.cpp
+++ b/Engine/Source/Editor/ContentBrowser/Private/FrontendFilters.cpp
@@ -36,6 +36,39 @@ public:
 		AssetPtr->PackageName.AppendString(AssetFullPath);
 		AssetPtr->GetExportTextName(AssetExportTextName);
 
+		// Test each piece of the path name, apart from the first
+		{
+			const TCHAR* Ptr = AssetFullPath.GetCharArray().GetData();
+			if (Ptr)
+			{
+				// Test each piece of the path name, apart from the first
+				bool bIsFirst = true;
+				while (const TCHAR* Delimiter = FCString::Strchr(Ptr, '/'))
+				{
+					const int32 Length = Delimiter - Ptr;
+
+					if (Length > 0)
+					{
+						if (bIsFirst)
+						{
+							bIsFirst = false;
+						}
+						else
+						{
+							AssetSplitPath.Emplace(Length, Ptr);
+						}
+					}
+
+					Ptr += (Length + 1);
+				}
+
+				if (*Ptr != 0)
+				{
+					AssetSplitPath.Emplace(Ptr);
+				}
+			}
+		}
+
 		if (FCollectionManagerModule::IsModuleAvailable())
 		{
 			FCollectionManagerModule& CollectionManagerModule = FCollectionManagerModule::GetModule();
@@ -48,6 +81,7 @@ public:
 		AssetPtr = nullptr;
 		AssetFullPath.Reset();
 		AssetExportTextName.Reset();
+		AssetSplitPath.Reset();
 		AssetCollectionNames.Reset();
 	}
 
@@ -63,39 +97,11 @@ public:
 
 	virtual bool TestBasicStringExpression(const FTextFilterString& InValue, const ETextFilterTextComparisonMode InTextComparisonMode) const override
 	{
-		const TCHAR* Ptr = AssetFullPath.GetCharArray().GetData();
-		if (Ptr)
+		for (const FString& AssetPathPart : AssetSplitPath)
 		{
-			// Test each piece of the path name, apart from the first
-			bool bIsFirst = true;
-			while (const TCHAR* Delimiter = FCString::Strchr(Ptr, '/'))
+			if (TextFilterUtils::TestBasicStringExpression(AssetPathPart, InValue, InTextComparisonMode))
 			{
-				const int32 Length = Delimiter - Ptr;
-
-				if (Length > 0)
-				{
-					if (bIsFirst)
-					{
-						bIsFirst = false;
-					}
-					else
-					{
-						if (TextFilterUtils::TestBasicStringExpression(FString(Length, Ptr), InValue, InTextComparisonMode))
-						{
-							return true;
-						}
-					}
-				}
-
-				Ptr += (Length + 1);
-			}
-
-			if (*Ptr != 0)
-			{
-				if (TextFilterUtils::TestBasicStringExpression(Ptr, InValue, InTextComparisonMode))
-				{
-					return true;
-				}
+				return true;
 			}
 		}
 
@@ -218,15 +224,15 @@ private:
 	/** The export text name of the current asset */
 	FString AssetExportTextName;
 
+	/** Split path of the current asset */
+	TArray<FString> AssetSplitPath;
+
 	/** Names of the collections that the current asset is in */
 	TArray<FName> AssetCollectionNames;
 
 	/** Are we supposed to include the class name in our basic string tests? */
 	bool bIncludeClassName;
 
-	/** Scratch string that can be used for temporary string conversions while avoiding re-allocations */
-	mutable FString TempString;
-
 	/** Keys used by TestComplexExpression */
 	const FName NameKeyName;
 	const FName PathKeyName;
