Fixed the Content Browser search text not updating when the selected dynamic collection is updated

diff --git a/Engine/Source/Editor/ContentBrowser/Private/SContentBrowser.cpp b/Engine/Source/Editor/ContentBrowser/Private/SContentBrowser.cpp
--- a/Engine/Source/Editor/ContentBrowser/Private/SContentBrowser.cpp
+++ b/Engine/Source/Editor/ContentBrowser/Private/SContentBrowser.cpp
@@ -645,6 +645,7 @@ void SContentBrowser::Construct( const FArguments& InArgs, const FName& InInstan
 	FCollectionManagerModule& CollectionManagerModule = FCollectionManagerModule::GetModule();
 	CollectionManagerModule.Get().OnCollectionRenamed().AddSP(this, &SContentBrowser::HandleCollectionRenamed);
 	CollectionManagerModule.Get().OnCollectionDestroyed().AddSP(this, &SContentBrowser::HandleCollectionRemoved);
+	CollectionManagerModule.Get().OnCollectionUpdated().AddSP(this, &SContentBrowser::HandleCollectionUpdated);
 
 	FAssetRegistryModule& AssetRegistryModule = FModuleManager::LoadModuleChecked<FAssetRegistryModule>(TEXT("AssetRegistry"));
 	AssetRegistryModule.Get().OnPathRemoved().AddSP(this, &SContentBrowser::HandlePathRemoved);
@@ -1052,7 +1053,7 @@ void SContentBrowser::SourcesChanged(const TArray<FString>& SelectedPaths, const
 	}
 
 	// A dynamic collection should apply its search query to the CB search, so we need to stash the current search so that we can restore it again later
-	if (SourcesData.Collections.Num() == 1 &&  SourcesData.IsDynamicCollection())
+	if (SourcesData.IsDynamicCollection())
 	{
 		// Only stash the user search term once in case we're switching between dynamic collections
 		if (!StashedSearchBoxText.IsSet())
@@ -2238,6 +2239,26 @@ void SContentBrowser::HandleCollectionRenamed(const FCollectionNameType& Origina
 	return HandleCollectionRemoved(OriginalCollection);
 }
 
+void SContentBrowser::HandleCollectionUpdated(const FCollectionNameType& Collection)
+{
+	const FSourcesData& SourcesData = AssetViewPtr->GetSourcesData();
+
+	// If we're currently viewing the dynamic collection that was updated, make sure our active filter text is up-to-date
+	if (SourcesData.IsDynamicCollection() && SourcesData.Collections[0] == Collection)
+	{
+		FCollectionManagerModule& CollectionManagerModule = FCollectionManagerModule::GetModule();
+
+		const FCollectionNameType& DynamicCollection = SourcesData.Collections[0];
+
+		FString DynamicQueryString;
+		CollectionManagerModule.Get().GetDynamicQueryText(DynamicCollection.Name, DynamicCollection.Type, DynamicQueryString);
+
+		const FText DynamicQueryText = FText::FromString(DynamicQueryString);
+		SetSearchBoxText(DynamicQueryText);
+		SearchBoxPtr->SetText(DynamicQueryText);
+	}
+}
+
 void SContentBrowser::HandlePathRemoved(const FString& Path)
 {
 	const FName PathName(*Path);
diff --git a/Engine/Source/Editor/ContentBrowser/Private/SContentBrowser.h b/Engine/Source/Editor/ContentBrowser/Private/SContentBrowser.h
--- a/Engine/Source/Editor/ContentBrowser/Private/SContentBrowser.h
+++ b/Engine/Source/Editor/ContentBrowser/Private/SContentBrowser.h
@@ -308,6 +308,9 @@ private:
 	/** Handles an on collection renamed event */
 	void HandleCollectionRenamed(const FCollectionNameType& OriginalCollection, const FCollectionNameType& NewCollection);
 
+	/** Handles an on collection updated event */
+	void HandleCollectionUpdated(const FCollectionNameType& Collection);
+
 	/** Handles a path removed event */
 	void HandlePathRemoved(const FString& Path);
 
