Dynamic collections can now be referenced in the Content Browser text filters

UETOOL-332 - Collections 2.0
UETOOL-401 - Add support for dynamic collections

You can now reference both static and dynamic collections in the Content Browser text filter, including recursive references (the evaluator guards against cyclic collection references).

The text filter now recursively gathers up any dynamic collection references when the filter text is changed. These are then tested depth first as sub-expressions against each object, to see whether the object is "contained" within the given dynamic collection.

diff --git a/Engine/Source/Developer/CollectionManager/Private/Collection.cpp b/Engine/Source/Developer/CollectionManager/Private/Collection.cpp
--- a/Engine/Source/Developer/CollectionManager/Private/Collection.cpp
+++ b/Engine/Source/Developer/CollectionManager/Private/Collection.cpp
@@ -2,6 +2,7 @@
 
 #include "CollectionManagerPrivatePCH.h"
 #include "ISourceControlModule.h"
+#include "TextFilterExpressionEvaluator.h"
 
 #define LOCTEXT_NAMESPACE "CollectionManager"
 
@@ -55,8 +56,7 @@ TSharedRef<FCollection> FCollection::Clone(const FString& InFilename, bool InUse
 
 bool FCollection::Load(FText& OutError)
 {
-	ObjectSet.Empty();
-	DiskSnapshot = FCollectionSnapshot();
+	Empty();
 
 	FString FullFileContentsString;
 	if (!FFileHelper::LoadFileToString(FullFileContentsString, *SourceFilename))
@@ -377,6 +377,15 @@ bool FCollection::DeleteSourceFile(FText& OutError)
 	return bSuccessfullyDeleted;
 }
 
+void FCollection::Empty()
+{
+	ObjectSet.Reset();
+	DynamicQueryText.Reset();
+	DynamicQueryExpressionEvaluatorPtr.Reset();
+
+	DiskSnapshot.TakeSnapshot(*this);
+}
+
 bool FCollection::AddObjectToCollection(FName ObjectPath)
 {
 	if (StorageMode == ECollectionStorageMode::Static && !ObjectSet.Contains(ObjectPath))
@@ -472,6 +481,26 @@ FString FCollection::GetDynamicQueryText() const
 	return (StorageMode == ECollectionStorageMode::Dynamic) ? DynamicQueryText : FString();
 }
 
+bool FCollection::TestDynamicQuery(const ITextFilterExpressionContext& InContext) const
+{
+	if (StorageMode == ECollectionStorageMode::Dynamic)
+	{
+		if (!DynamicQueryExpressionEvaluatorPtr.IsValid())
+		{
+			DynamicQueryExpressionEvaluatorPtr = MakeShareable(new FTextFilterExpressionEvaluator(ETextFilterExpressionEvaluatorMode::Complex));
+		}
+
+		if (!DynamicQueryExpressionEvaluatorPtr->GetFilterText().ToString().Equals(DynamicQueryText, ESearchCase::CaseSensitive))
+		{
+			DynamicQueryExpressionEvaluatorPtr->SetFilterText(FText::FromString(DynamicQueryText));
+		}
+
+		return DynamicQueryExpressionEvaluatorPtr->TestTextFilter(InContext);
+	}
+
+	return false;
+}
+
 FCollectionStatusInfo FCollection::GetStatusInfo() const
 {
 	FCollectionStatusInfo StatusInfo;
@@ -617,11 +646,8 @@ bool FCollection::MergeWithCollection(const FCollection& Other)
 		bHasChanges = true;
 		StorageMode = Other.StorageMode;
 
-		// Storage mode has changed!
-		// Reset the static and dynamic parts of the collection so we just copy over the new data verbatim
-		ObjectSet.Reset();
-		DynamicQueryText.Reset();
-		DiskSnapshot.TakeSnapshot(*this);
+		// Storage mode has changed! Empty the collection so we just copy over the new data verbatim
+		Empty();
 	}
 
 	if (StorageMode == ECollectionStorageMode::Static)
diff --git a/Engine/Source/Developer/CollectionManager/Private/Collection.h b/Engine/Source/Developer/CollectionManager/Private/Collection.h
--- a/Engine/Source/Developer/CollectionManager/Private/Collection.h
+++ b/Engine/Source/Developer/CollectionManager/Private/Collection.h
@@ -2,6 +2,8 @@
 
 #pragma once
 
+class FTextFilterExpressionEvaluator;
+
 struct ECollectionVersion
 {
 	enum Type
@@ -50,6 +52,8 @@ public:
 	bool Merge(const FCollection& NewCollection);
 	/** Deletes the source file for this collection. If false, OutError is a human readable warning depicting the error. */
 	bool DeleteSourceFile(FText& OutError);
+	/** Empty this collection */
+	void Empty();
 
 	/** Adds a single object to the collection. Static collections only. */
 	bool AddObjectToCollection(FName ObjectPath);
@@ -70,6 +74,8 @@ public:
 	bool SetDynamicQueryText(const FString& InQueryText);
 	/** Get the dynamic query text for this collection. Dynamic collections only. */
 	FString GetDynamicQueryText() const;
+	/** Tests the dynamic query for against the context provided. Dynamic collections only. */
+	bool TestDynamicQuery(const ITextFilterExpressionContext& InContext) const;
 
 	/** Get the status info for this collection */
 	FCollectionStatusInfo GetStatusInfo() const;
@@ -173,6 +179,9 @@ private:
 	/** The dynamic query string for this collection. Dynamic collections only. */
 	FString DynamicQueryText;
 
+	/** Expression evaluator that can be used test against the compiled DynamicQueryText */
+	mutable TSharedPtr<FTextFilterExpressionEvaluator> DynamicQueryExpressionEvaluatorPtr;
+
 	/** The file version for this collection */
 	ECollectionVersion::Type FileVersion;
 
diff --git a/Engine/Source/Developer/CollectionManager/Private/CollectionManager.cpp b/Engine/Source/Developer/CollectionManager/Private/CollectionManager.cpp
--- a/Engine/Source/Developer/CollectionManager/Private/CollectionManager.cpp
+++ b/Engine/Source/Developer/CollectionManager/Private/CollectionManager.cpp
@@ -1139,6 +1139,34 @@ bool FCollectionManager::GetDynamicQueryText(FName CollectionName, ECollectionSh
 	return true;
 }
 
+bool FCollectionManager::TestDynamicQuery(FName CollectionName, ECollectionShareType::Type ShareType, const ITextFilterExpressionContext& InContext, bool& OutResult) const
+{
+	if (!ensure(ShareType < ECollectionShareType::CST_All))
+	{
+		// Bad share type
+		LastError = LOCTEXT("Error_Internal", "There was an internal error.");
+		return false;
+	}
+
+	const FCollectionNameType CollectionKey(CollectionName, ShareType);
+	const TSharedRef<FCollection>* const CollectionRefPtr = AvailableCollections.Find(CollectionKey);
+	if (!CollectionRefPtr)
+	{
+		// Collection doesn't exist
+		LastError = LOCTEXT("Error_DoesntExist", "The collection doesn't exist.");
+		return false;
+	}
+
+	if ((*CollectionRefPtr)->GetStorageMode() != ECollectionStorageMode::Dynamic)
+	{
+		LastError = LOCTEXT("Error_TestNeedsDynamicCollection", "Search queries can only be tested on dynamic collections.");
+		return false;
+	}
+
+	OutResult = (*CollectionRefPtr)->TestDynamicQuery(InContext);
+	return true;
+}
+
 bool FCollectionManager::EmptyCollection(FName CollectionName, ECollectionShareType::Type ShareType)
 {
 	if (!ensure(ShareType < ECollectionShareType::CST_All))
@@ -1148,20 +1176,33 @@ bool FCollectionManager::EmptyCollection(FName CollectionName, ECollectionShareT
 		return false;
 	}
 
-	TArray<FName> ObjectPaths;
-	if (!GetObjectsInCollection(CollectionName, ShareType, ObjectPaths))
+	const FCollectionNameType CollectionKey(CollectionName, ShareType);
+	const TSharedRef<FCollection>* const CollectionRefPtr = AvailableCollections.Find(CollectionKey);
+	if (!CollectionRefPtr)
 	{
-		// Failed to load collection
+		// Collection doesn't exist
+		LastError = LOCTEXT("Error_DoesntExist", "The collection doesn't exist.");
 		return false;
 	}
-	
-	if (ObjectPaths.Num() == 0)
+
+	if ((*CollectionRefPtr)->IsEmpty())
 	{
-		// Collection already empty
+		// Already empty - nothing to do
 		return true;
 	}
+
+	(*CollectionRefPtr)->Empty();
 	
-	return RemoveFromCollection(CollectionName, ShareType, ObjectPaths);
+	if ((*CollectionRefPtr)->Save(LastError))
+	{
+		CollectionFileCaches[ShareType]->IgnoreFileModification((*CollectionRefPtr)->GetSourceFilename());
+
+		CollectionCache.HandleCollectionChanged();
+		CollectionUpdatedEvent.Broadcast(CollectionKey);
+		return true;
+	}
+
+	return false;
 }
 
 bool FCollectionManager::SaveCollection(FName CollectionName, ECollectionShareType::Type ShareType)
diff --git a/Engine/Source/Developer/CollectionManager/Private/CollectionManager.h b/Engine/Source/Developer/CollectionManager/Private/CollectionManager.h
--- a/Engine/Source/Developer/CollectionManager/Private/CollectionManager.h
+++ b/Engine/Source/Developer/CollectionManager/Private/CollectionManager.h
@@ -124,6 +124,7 @@ public:
 	virtual bool RemoveFromCollection(FName CollectionName, ECollectionShareType::Type ShareType, const TArray<FName>& ObjectPaths, int32* OutNumRemoved = nullptr) override;
 	virtual bool SetDynamicQueryText(FName CollectionName, ECollectionShareType::Type ShareType, const FString& InQueryText) override;
 	virtual bool GetDynamicQueryText(FName CollectionName, ECollectionShareType::Type ShareType, FString& OutQueryText) const override;
+	virtual bool TestDynamicQuery(FName CollectionName, ECollectionShareType::Type ShareType, const ITextFilterExpressionContext& InContext, bool& OutResult) const override;
 	virtual bool EmptyCollection(FName CollectionName, ECollectionShareType::Type ShareType) override;
 	virtual bool SaveCollection(FName CollectionName, ECollectionShareType::Type ShareType) override;
 	virtual bool UpdateCollection(FName CollectionName, ECollectionShareType::Type ShareType) override;
diff --git a/Engine/Source/Developer/CollectionManager/Public/ICollectionManager.h b/Engine/Source/Developer/CollectionManager/Public/ICollectionManager.h
--- a/Engine/Source/Developer/CollectionManager/Public/ICollectionManager.h
+++ b/Engine/Source/Developer/CollectionManager/Public/ICollectionManager.h
@@ -4,6 +4,8 @@
 
 #include "CollectionManagerTypes.h"
 
+class ITextFilterExpressionContext;
+
 class ICollectionManager
 {
 public:
@@ -149,6 +151,17 @@ public:
 	 */
 	virtual bool GetDynamicQueryText(FName CollectionName, ECollectionShareType::Type ShareType, FString& OutQueryText) const = 0;
 
+	/**
+	 * Tests the dynamic query for the specified collection against the context provided.
+	 *
+	 * @param CollectionName The collection to get the query from.
+	 * @param ShareType The way the collection is shared.
+	 * @param InContext The context to test against.
+	 * @param OutResult Filled with the result of the query.
+	 * @return true if the get was successful. If false, GetLastError will return a human readable string description of the error.
+	 */
+	virtual bool TestDynamicQuery(FName CollectionName, ECollectionShareType::Type ShareType, const ITextFilterExpressionContext& InContext, bool& OutResult) const = 0;
+
 	/**
 	  * Removes all assets from the specified collection.
 	  *
diff --git a/Engine/Source/Editor/ContentBrowser/Private/FrontendFilters.cpp b/Engine/Source/Editor/ContentBrowser/Private/FrontendFilters.cpp
--- a/Engine/Source/Editor/ContentBrowser/Private/FrontendFilters.cpp
+++ b/Engine/Source/Editor/ContentBrowser/Private/FrontendFilters.cpp
@@ -12,13 +12,153 @@
 // FFrontendFilter_Text
 /////////////////////////////////////////
 
+/** Expression context which gathers up the names of any dynamic collections being referenced by the current query */
+class FFrontendFilter_GatherDynamicCollectionsExpressionContext : public ITextFilterExpressionContext
+{
+public:
+	FFrontendFilter_GatherDynamicCollectionsExpressionContext(TArray<FCollectionNameType>& OutReferencedDynamicCollections)
+		: AvailableDynamicCollections()
+		, ReferencedDynamicCollections(OutReferencedDynamicCollections)
+		, CurrentRecursionDepth(0)
+		, CollectionKeyName("Collection")
+		, TagKeyName("Tag")
+	{
+		if (FCollectionManagerModule::IsModuleAvailable())
+		{
+			FCollectionManagerModule& CollectionManagerModule = FCollectionManagerModule::GetModule();
+
+			TArray<FCollectionNameType> AvailableCollections;
+			CollectionManagerModule.Get().GetCollections(AvailableCollections);
+
+			for (const FCollectionNameType& AvailableCollection : AvailableCollections)
+			{
+				// Only care about dynamic collections
+				ECollectionStorageMode::Type StorageMode = ECollectionStorageMode::Static;
+				CollectionManagerModule.Get().GetCollectionStorageMode(AvailableCollection.Name, AvailableCollection.Type, StorageMode);
+				if (StorageMode != ECollectionStorageMode::Dynamic)
+				{
+					continue;
+				}
+
+				AvailableDynamicCollections.Add(AvailableCollection);
+			}
+		}
+	}
+
+	~FFrontendFilter_GatherDynamicCollectionsExpressionContext()
+	{
+		// Sort and populate the final list of referenced dynamic collections
+		FoundDynamicCollections.Sort([](const FDynamicCollectionNameAndDepth& A, const FDynamicCollectionNameAndDepth& B)
+		{
+			return A.RecursionDepth > B.RecursionDepth;
+		});
+
+		ReferencedDynamicCollections.Reset();
+		ReferencedDynamicCollections.Reserve(FoundDynamicCollections.Num());
+		for (const auto& FoundDynamicCollection : FoundDynamicCollections)
+		{
+			ReferencedDynamicCollections.Add(FoundDynamicCollection.Collection);
+		}
+	}
+
+	virtual bool TestBasicStringExpression(const FTextFilterString& InValue, const ETextFilterTextComparisonMode InTextComparisonMode) const override
+	{
+		TestAgainstAvailableCollections(InValue, InTextComparisonMode);
+		return false;
+	}
+
+	virtual bool TestComplexExpression(const FName& InKey, const FTextFilterString& InValue, const ETextFilterComparisonOperation InComparisonOperation, const ETextFilterTextComparisonMode InTextComparisonMode) const override
+	{
+		// Special case for collections, as these aren't contained within the asset registry meta-data
+		if (InKey == CollectionKeyName || InKey == TagKeyName)
+		{
+			// Collections can only work with Equal or NotEqual type tests
+			if (InComparisonOperation != ETextFilterComparisonOperation::Equal && InComparisonOperation != ETextFilterComparisonOperation::NotEqual)
+			{
+				return false;
+			}
+
+			TestAgainstAvailableCollections(InValue, InTextComparisonMode);
+		}
+
+		return false;
+	}
+
+private:
+	bool TestAgainstAvailableCollections(const FTextFilterString& InValue, const ETextFilterTextComparisonMode InTextComparisonMode) const
+	{
+		for (const FCollectionNameType& DynamicCollection : AvailableDynamicCollections)
+		{
+			const FString DynamicCollectionNameStr = DynamicCollection.Name.ToString();
+			if (TextFilterUtils::TestBasicStringExpression(DynamicCollectionNameStr, InValue, InTextComparisonMode))
+			{
+				const bool bCollectionAlreadyProcessed = FoundDynamicCollections.ContainsByPredicate([&](const FDynamicCollectionNameAndDepth& Other)
+				{
+					return DynamicCollection == Other.Collection;
+				});
+
+				if (!bCollectionAlreadyProcessed)
+				{
+					FoundDynamicCollections.Add(FDynamicCollectionNameAndDepth(DynamicCollection, CurrentRecursionDepth));
+
+					if (FCollectionManagerModule::IsModuleAvailable())
+					{
+						FCollectionManagerModule& CollectionManagerModule = FCollectionManagerModule::GetModule();
+						
+						// Also need to gather any collections referenced by this dynamic collection
+						++CurrentRecursionDepth;
+						bool bUnused = false;
+						CollectionManagerModule.Get().TestDynamicQuery(DynamicCollection.Name, DynamicCollection.Type, *this, bUnused);
+						--CurrentRecursionDepth;
+					}
+				}
+
+				return true;
+			}
+		}
+
+		return false;
+	}
+
+	/** Contains a collection name along with its recursion depth in the dynamic query - used so we can test them depth first */
+	struct FDynamicCollectionNameAndDepth
+	{
+		FDynamicCollectionNameAndDepth(FCollectionNameType InCollection, const int32 InRecursionDepth)
+			: Collection(InCollection)
+			, RecursionDepth(InRecursionDepth)
+		{
+		}
+
+		FCollectionNameType Collection;
+		int32 RecursionDepth;
+	};
+
+	/** The currently available dynamic collections */
+	TArray<FCollectionNameType> AvailableDynamicCollections;
+
+	/** This will be populated with any dynamic collections that are being referenced by the current query - these collections may not all match when tested against the actual asset data */
+	TArray<FCollectionNameType>& ReferencedDynamicCollections;
+
+	/** Dynamic collections that have currently be found as part of the query (or recursive sub-query) */
+	mutable TArray<FDynamicCollectionNameAndDepth> FoundDynamicCollections;
+
+	/** Incremented when we test a sub-query, decremented once we're done */
+	mutable int32 CurrentRecursionDepth;
+
+	/** Keys used by TestComplexExpression */
+	const FName CollectionKeyName;
+	const FName TagKeyName;
+};
+
+/** Expression context to test the given asset data against the current text filter */
 class FFrontendFilter_TextFilterExpressionContext : public ITextFilterExpressionContext
 {
 public:
 	typedef TRemoveReference<FAssetFilterType>::Type* FAssetFilterTypePtr;
 
-	FFrontendFilter_TextFilterExpressionContext()
-		: AssetPtr(nullptr)
+	FFrontendFilter_TextFilterExpressionContext(const TArray<FCollectionNameType>& InReferencedDynamicCollections)
+		: ReferencedDynamicCollections(InReferencedDynamicCollections)
+		, AssetPtr(nullptr)
 		, bIncludeClassName(true)
 		, NameKeyName("Name")
 		, PathKeyName("Path")
@@ -73,6 +213,18 @@ public:
 		{
 			FCollectionManagerModule& CollectionManagerModule = FCollectionManagerModule::GetModule();
 			CollectionManagerModule.Get().GetCollectionsContainingObject(AssetPtr->ObjectPath, ECollectionShareType::CST_All, AssetCollectionNames, ECollectionRecursionFlags::SelfAndChildren);
+
+			// Test the dynamic collections from the active query against the current asset
+			// We can do this as a flat list since FFrontendFilter_GatherDynamicCollectionsExpressionContext has already taken care of processing the recursion
+			for (const FCollectionNameType& DynamicCollection : ReferencedDynamicCollections)
+			{
+				bool bPassesCollectionFilter = false;
+				CollectionManagerModule.Get().TestDynamicQuery(DynamicCollection.Name, DynamicCollection.Type, *this, bPassesCollectionFilter);
+				if (bPassesCollectionFilter)
+				{
+					AssetCollectionNames.AddUnique(DynamicCollection.Name);
+				}
+			}
 		}
 	}
 
@@ -215,6 +367,9 @@ public:
 	}
 
 private:
+	/** An array of dynamic collections that are being referenced by the current query. These should be tested against each asset when it's looking for collections that contain it */
+	const TArray<FCollectionNameType>& ReferencedDynamicCollections;
+
 	/** Pointer to the asset we're currently filtering */
 	FAssetFilterTypePtr AssetPtr;
 
@@ -244,13 +399,31 @@ private:
 
 FFrontendFilter_Text::FFrontendFilter_Text()
 	: FFrontendFilter(nullptr)
-	, TextFilterExpressionContext(MakeShareable(new FFrontendFilter_TextFilterExpressionContext()))
+	, ReferencedDynamicCollections()
+	, TextFilterExpressionContext(MakeShareable(new FFrontendFilter_TextFilterExpressionContext(ReferencedDynamicCollections)))
 	, TextFilterExpressionEvaluator(ETextFilterExpressionEvaluatorMode::Complex)
 {
+	FCollectionManagerModule& CollectionManagerModule = FCollectionManagerModule::GetModule();
+
+	// We need to watch for collection changes so that we can keep ReferencedDynamicCollections up-to-date
+	OnCollectionCreatedHandle = CollectionManagerModule.Get().OnCollectionCreated().AddRaw(this, &FFrontendFilter_Text::HandleCollectionCreated);
+	OnCollectionDestroyedHandle = CollectionManagerModule.Get().OnCollectionDestroyed().AddRaw(this, &FFrontendFilter_Text::HandleCollectionDestroyed);
+	OnCollectionRenamedHandle = CollectionManagerModule.Get().OnCollectionRenamed().AddRaw(this, &FFrontendFilter_Text::HandleCollectionRenamed);
+	OnCollectionUpdatedHandle = CollectionManagerModule.Get().OnCollectionUpdated().AddRaw(this, &FFrontendFilter_Text::HandleCollectionUpdated);
 }
 
 FFrontendFilter_Text::~FFrontendFilter_Text()
 {
+	// Check IsModuleAvailable as we might be in the process of shutting down...
+	if (FCollectionManagerModule::IsModuleAvailable())
+	{
+		FCollectionManagerModule& CollectionManagerModule = FCollectionManagerModule::GetModule();
+
+		CollectionManagerModule.Get().OnCollectionCreated().Remove(OnCollectionCreatedHandle);
+		CollectionManagerModule.Get().OnCollectionDestroyed().Remove(OnCollectionDestroyedHandle);
+		CollectionManagerModule.Get().OnCollectionRenamed().Remove(OnCollectionRenamedHandle);
+		CollectionManagerModule.Get().OnCollectionUpdated().Remove(OnCollectionUpdatedHandle);
+	}
 }
 
 bool FFrontendFilter_Text::PassesFilter(FAssetFilterType InItem) const
@@ -270,6 +443,8 @@ void FFrontendFilter_Text::SetRawFilterText(const FText& InFilterText)
 {
 	if (TextFilterExpressionEvaluator.SetFilterText(InFilterText))
 	{
+		RebuildReferencedDynamicCollections();
+
 		// Will trigger a re-filter with the new text
 		BroadcastChangedEvent();
 	}
@@ -291,6 +466,47 @@ void FFrontendFilter_Text::SetIncludeClassName(const bool InIncludeClassName)
 	}
 }
 
+void FFrontendFilter_Text::HandleCollectionCreated(const FCollectionNameType& Collection)
+{
+	RebuildReferencedDynamicCollections();
+
+	// Will trigger a re-filter with the new collections
+	BroadcastChangedEvent();
+}
+
+void FFrontendFilter_Text::HandleCollectionDestroyed(const FCollectionNameType& Collection)
+{
+	if (ReferencedDynamicCollections.Contains(Collection))
+	{
+		RebuildReferencedDynamicCollections();
+
+		// Will trigger a re-filter with the new collections
+		BroadcastChangedEvent();
+	}
+}
+
+void FFrontendFilter_Text::HandleCollectionRenamed(const FCollectionNameType& OriginalCollection, const FCollectionNameType& NewCollection)
+{
+	int32 FoundIndex = INDEX_NONE;
+	if (ReferencedDynamicCollections.Find(OriginalCollection, FoundIndex))
+	{
+		ReferencedDynamicCollections[FoundIndex] = NewCollection;
+	}
+}
+
+void FFrontendFilter_Text::HandleCollectionUpdated(const FCollectionNameType& Collection)
+{
+	RebuildReferencedDynamicCollections();
+
+	// Will trigger a re-filter with the new collections
+	BroadcastChangedEvent();
+}
+
+void FFrontendFilter_Text::RebuildReferencedDynamicCollections()
+{
+	TextFilterExpressionEvaluator.TestTextFilter(FFrontendFilter_GatherDynamicCollectionsExpressionContext(ReferencedDynamicCollections));
+}
+
 
 /////////////////////////////////////////
 // FFrontendFilter_CheckedOut
diff --git a/Engine/Source/Editor/ContentBrowser/Private/FrontendFilters.h b/Engine/Source/Editor/ContentBrowser/Private/FrontendFilters.h
--- a/Engine/Source/Editor/ContentBrowser/Private/FrontendFilters.h
+++ b/Engine/Source/Editor/ContentBrowser/Private/FrontendFilters.h
@@ -36,11 +36,35 @@ public:
 	void SetIncludeClassName(const bool InIncludeClassName);
 
 private:
+	/** Handles an on collection created event */
+	void HandleCollectionCreated(const FCollectionNameType& Collection);
+
+	/** Handles an on collection destroyed event */
+	void HandleCollectionDestroyed(const FCollectionNameType& Collection);
+
+	/** Handles an on collection renamed event */
+	void HandleCollectionRenamed(const FCollectionNameType& OriginalCollection, const FCollectionNameType& NewCollection);
+
+	/** Handles an on collection updated event */
+	void HandleCollectionUpdated(const FCollectionNameType& Collection);
+
+	/** Rebuild the array of dynamic collections that are being referenced by the current query */
+	void RebuildReferencedDynamicCollections();
+
+	/** An array of dynamic collections that are being referenced by the current query. These should be tested against each asset when it's looking for collections that contain it */
+	TArray<FCollectionNameType> ReferencedDynamicCollections;
+
 	/** Transient context data, used when calling PassesFilter. Kept around to minimize re-allocations between multiple calls to PassesFilter */
 	TSharedRef<class FFrontendFilter_TextFilterExpressionContext> TextFilterExpressionContext;
 
 	/** Expression evaluator that can be used to perform complex text filter queries */
 	FTextFilterExpressionEvaluator TextFilterExpressionEvaluator;
+
+	/** Delegate handles */
+	FDelegateHandle OnCollectionCreatedHandle;
+	FDelegateHandle OnCollectionDestroyedHandle;
+	FDelegateHandle OnCollectionRenamedHandle;
+	FDelegateHandle OnCollectionUpdatedHandle;
 };
 
 /** A filter that displays only checked out assets */
