Fixed crash when shutting down the editor

UE-18403 - Crash occurs exiting the editor with Collections visible

The SCC provider would be nulled out when the SCC module is unloaded, which was triggering a notification to update the collection item states, however the collection manager module had also been unloaded by the time that was being processed, leading to an assert.

This change makes sure that the collection manager module is still available before attempting to use it to update the collection item states.

diff --git a/Engine/Source/Editor/ContentBrowser/Private/SCollectionView.cpp b/Engine/Source/Editor/ContentBrowser/Private/SCollectionView.cpp
--- a/Engine/Source/Editor/ContentBrowser/Private/SCollectionView.cpp
+++ b/Engine/Source/Editor/ContentBrowser/Private/SCollectionView.cpp
@@ -301,56 +301,60 @@ void SCollectionView::HandleSourceControlStateChanged()
 
 void SCollectionView::UpdateCollectionItemStatus( const TSharedRef<FCollectionItem>& CollectionItem )
 {
-	FCollectionManagerModule& CollectionManagerModule = FCollectionManagerModule::GetModule();
-
 	TOptional<ECollectionItemStatus> NewStatus;
 
-	FCollectionStatusInfo StatusInfo;
-	if (CollectionManagerModule.Get().GetCollectionStatusInfo(CollectionItem->CollectionName, CollectionItem->CollectionType, StatusInfo))
+	// Check IsModuleAvailable as we might be in the process of shutting down, and were notified due to the SCC provider being nulled out...
+	if (FCollectionManagerModule::IsModuleAvailable())
 	{
-		// Test the SCC state first as this should take priority when reporting the status back to the user
-		if (StatusInfo.bUseSCC)
+		FCollectionManagerModule& CollectionManagerModule = FCollectionManagerModule::GetModule();
+
+		FCollectionStatusInfo StatusInfo;
+		if (CollectionManagerModule.Get().GetCollectionStatusInfo(CollectionItem->CollectionName, CollectionItem->CollectionType, StatusInfo))
 		{
-			if (StatusInfo.SCCState.IsValid() && StatusInfo.SCCState->IsSourceControlled())
+			// Test the SCC state first as this should take priority when reporting the status back to the user
+			if (StatusInfo.bUseSCC)
 			{
-				if (StatusInfo.SCCState->IsCheckedOutOther())
+				if (StatusInfo.SCCState.IsValid() && StatusInfo.SCCState->IsSourceControlled())
+				{
+					if (StatusInfo.SCCState->IsCheckedOutOther())
+					{
+						NewStatus = ECollectionItemStatus::IsCheckedOutByAnotherUser;
+					}
+					else if (StatusInfo.SCCState->IsConflicted())
+					{
+						NewStatus = ECollectionItemStatus::IsConflicted;
+					}
+					else if (!StatusInfo.SCCState->IsCurrent())
+					{
+						NewStatus = ECollectionItemStatus::IsOutOfDate;
+					}
+					else if (StatusInfo.SCCState->IsModified())
+					{
+						NewStatus = ECollectionItemStatus::HasLocalChanges;
+					}
+				}
+				else
 				{
-					NewStatus = ECollectionItemStatus::IsCheckedOutByAnotherUser;
+					NewStatus = ECollectionItemStatus::IsMissingSCCProvider;
 				}
-				else if (StatusInfo.SCCState->IsConflicted())
+			}
+
+			// Not set by the SCC status, so check just use the local state
+			if (!NewStatus.IsSet())
+			{
+				if (StatusInfo.bIsDirty)
 				{
-					NewStatus = ECollectionItemStatus::IsConflicted;
+					NewStatus = ECollectionItemStatus::HasLocalChanges;
 				}
-				else if (!StatusInfo.SCCState->IsCurrent())
+				else if (StatusInfo.bIsEmpty)
 				{
-					NewStatus = ECollectionItemStatus::IsOutOfDate;
+					NewStatus = ECollectionItemStatus::IsUpToDateAndEmpty;
 				}
-				else if (StatusInfo.SCCState->IsModified())
+				else
 				{
-					NewStatus = ECollectionItemStatus::HasLocalChanges;
+					NewStatus = ECollectionItemStatus::IsUpToDateAndPopulated;
 				}
 			}
-			else
-			{
-				NewStatus = ECollectionItemStatus::IsMissingSCCProvider;
-			}
-		}
-
-		// Not set by the SCC status, so check just use the local state
-		if (!NewStatus.IsSet())
-		{
-			if (StatusInfo.bIsDirty)
-			{
-				NewStatus = ECollectionItemStatus::HasLocalChanges;
-			}
-			else if (StatusInfo.bIsEmpty)
-			{
-				NewStatus = ECollectionItemStatus::IsUpToDateAndEmpty;
-			}
-			else
-			{
-				NewStatus = ECollectionItemStatus::IsUpToDateAndPopulated;
-			}
 		}
 	}
 
@@ -720,7 +724,7 @@ void SCollectionView::Tick( const FGeometry& AllottedGeometry, const double InCu
 		UpdateCollectionItems();
 	}
 
-	if (bQueueSCCRefresh && ISourceControlModule::Get().IsEnabled())
+	if (bQueueSCCRefresh && FCollectionManagerModule::IsModuleAvailable() && ISourceControlModule::Get().IsEnabled())
 	{
 		bQueueSCCRefresh = false;
 
