Content Browser: Fixed a bug where filters using the default category could still end up with more than one category instance in the menu

diff --git a/Engine/Source/Editor/ContentBrowser/Private/SFilterList.cpp b/Engine/Source/Editor/ContentBrowser/Private/SFilterList.cpp
--- a/Engine/Source/Editor/ContentBrowser/Private/SFilterList.cpp
+++ b/Engine/Source/Editor/ContentBrowser/Private/SFilterList.cpp
@@ -429,7 +429,7 @@ void SFilterList::Construct( const FArguments& InArgs )
 		AllFrontendFilters.Add(Filter);
 	}
 
-	AllFrontendFilterCategories.Add(DefaultCategory);
+	AllFrontendFilterCategories.AddUnique(DefaultCategory);
 
 	// Auto add all inverse filters
 	for (auto FilterIt = AllFrontendFilters.CreateConstIterator(); FilterIt; ++FilterIt)
