Fixed newly added collections losing their selection state once you name them

UE-18995 - New Collections that are renamed dont focus the renamed collection

diff --git a/Engine/Source/Editor/ContentBrowser/Private/SCollectionView.cpp b/Engine/Source/Editor/ContentBrowser/Private/SCollectionView.cpp
--- a/Engine/Source/Editor/ContentBrowser/Private/SCollectionView.cpp
+++ b/Engine/Source/Editor/ContentBrowser/Private/SCollectionView.cpp
@@ -248,6 +248,9 @@ void SCollectionView::HandleCollectionRenamed( const FCollectionNameType& Origin
 	{
 		CollectionItem->CollectionName = NewCollection.Name;
 		CollectionItem->CollectionType = NewCollection.Type;
+
+		AvailableCollections.Remove(OriginalCollection);
+		AvailableCollections.Add(NewCollection, CollectionItem);
 	}
 }
 
@@ -1461,6 +1464,13 @@ bool SCollectionView::CollectionNameChangeCommit( const TSharedPtr< FCollectionI
 			return false;
 		}
 
+		// Since we're really adding a new collection (as our placeholder item is currently transient), we don't get a rename event from the collections manager
+		// We'll spoof one here that so that our placeholder tree item is updated with the final name - this will preserve its expansion and selection state
+		HandleCollectionRenamed(
+			FCollectionNameType(CollectionItem->CollectionName, CollectionItem->CollectionType), 
+			FCollectionNameType(NewNameFinal, CollectionItem->CollectionType)
+			);
+
 		if ( NewCollectionParentKey.IsSet() )
 		{
 			// Try and set the parent correctly (if this fails for any reason, the collection will still be added, but will just appear at the root)
