Fixed a crash when summoning the collections context menu with no collections were selected

UE-19478 - Crash when right-clicking in the collections area of the source panel with private collections present

diff --git a/Engine/Source/Editor/ContentBrowser/Private/CollectionContextMenu.cpp b/Engine/Source/Editor/ContentBrowser/Private/CollectionContextMenu.cpp
--- a/Engine/Source/Editor/ContentBrowser/Private/CollectionContextMenu.cpp
+++ b/Engine/Source/Editor/ContentBrowser/Private/CollectionContextMenu.cpp
@@ -70,18 +70,26 @@ TSharedPtr<SWidget> FCollectionContextMenu::MakeCollectionTreeContextMenu(TShare
 		const bool bHasSingleSelectedCollection = SelectedCollections.Num() == 1;
 		const bool bIsFirstSelectedCollectionStatic = SelectedCollections.Num() > 0 && SelectedCollections[0]->StorageMode == ECollectionStorageMode::Static;
 
-		// New... (submenu)
-		MenuBuilder.AddSubMenu(
-			LOCTEXT("NewChildCollection", "New..."),
-			LOCTEXT("NewChildCollectionTooltip", "Create child a collection."),
-			FNewMenuDelegate::CreateRaw( this, &FCollectionContextMenu::MakeNewCollectionSubMenu, ECollectionStorageMode::Static, SCollectionView::FCreateCollectionPayload( FCollectionNameType( SelectedCollections[0]->CollectionName, SelectedCollections[0]->CollectionType ) ) ),
-			FUIAction(
-				FExecuteAction(),
-				FCanExecuteAction::CreateLambda( [=]{ return bHasSingleSelectedCollection && bIsFirstSelectedCollectionStatic; } )
-				),
-			NAME_None,
-			EUserInterfaceActionType::Button
-			);
+		{
+			TOptional<FCollectionNameType> ParentCollection;
+			if (SelectedCollections.Num() > 0) 
+			{
+				ParentCollection = FCollectionNameType(SelectedCollections[0]->CollectionName, SelectedCollections[0]->CollectionType);
+			}
+
+			// New... (submenu)
+			MenuBuilder.AddSubMenu(
+				LOCTEXT("NewChildCollection", "New..."),
+				LOCTEXT("NewChildCollectionTooltip", "Create child a collection."),
+				FNewMenuDelegate::CreateRaw( this, &FCollectionContextMenu::MakeNewCollectionSubMenu, ECollectionStorageMode::Static, SCollectionView::FCreateCollectionPayload( ParentCollection ) ),
+				FUIAction(
+					FExecuteAction(),
+					FCanExecuteAction::CreateLambda( [=]{ return bHasSingleSelectedCollection && bIsFirstSelectedCollectionStatic; } )
+					),
+				NAME_None,
+				EUserInterfaceActionType::Button
+				);
+		}
 
 		// Rename
 		MenuBuilder.AddMenuEntry( FGenericCommands::Get().Rename, NAME_None, LOCTEXT("RenameCollection", "Rename"), LOCTEXT("RenameCollectionTooltip", "Rename this collection."));
