Fixed drag and drop only working in the Content Browser tile view

UE-19914 - Different Drag and Drop Behavior between Tile, List, and Column Views in the Content Browser

The list view was missing its callbacks for the drag and drop operations, and the column view wasn't forwarding on the events from the column row on to the internal asset view item.

diff --git a/Engine/Source/Editor/ContentBrowser/Private/AssetViewWidgets.cpp b/Engine/Source/Editor/ContentBrowser/Private/AssetViewWidgets.cpp
--- a/Engine/Source/Editor/ContentBrowser/Private/AssetViewWidgets.cpp
+++ b/Engine/Source/Editor/ContentBrowser/Private/AssetViewWidgets.cpp
@@ -1631,10 +1631,15 @@ TSharedRef<SWidget> SAssetColumnItem::GenerateWidgetForColumn( const FName& Colu
 			AssetItem->RenamedRequestEvent.BindSP( InlineRenameWidget.Get(), &SInlineEditableTextBlock::EnterEditingMode );
 		}
 
-		return SNew( SAssetColumnItemNameBox, SharedThis(this) )
-			.Padding( ColumnItemPadding )
+		return SNew(SBorder)
+			.BorderImage(this, &SAssetViewItem::GetBorderImage)
+			.Padding(0)
 			[
-				Content.ToSharedRef()
+				SNew( SAssetColumnItemNameBox, SharedThis(this) )
+				.Padding( ColumnItemPadding )
+				[
+					Content.ToSharedRef()
+				]
 			];
 	}
 	else if ( ColumnName == "Class" )
diff --git a/Engine/Source/Editor/ContentBrowser/Private/AssetViewWidgets.h b/Engine/Source/Editor/ContentBrowser/Private/AssetViewWidgets.h
--- a/Engine/Source/Editor/ContentBrowser/Private/AssetViewWidgets.h
+++ b/Engine/Source/Editor/ContentBrowser/Private/AssetViewWidgets.h
@@ -115,12 +115,15 @@ public:
 	/** Performs common initialization logic for all asset view items */
 	void Construct( const FArguments& InArgs );
 
+	/** NOTE: Any functions overridden from the base widget classes *must* also be overridden by SAssetColumnViewRow and forwarded on to its internal item */
 	virtual void Tick( const FGeometry& AllottedGeometry, const double InCurrentTime, const float InDeltaTime ) override;
 	virtual TSharedPtr<IToolTip> GetToolTip() override;
 	virtual void OnDragEnter( const FGeometry& MyGeometry, const FDragDropEvent& DragDropEvent ) override;
 	virtual void OnDragLeave( const FDragDropEvent& DragDropEvent ) override;
 	virtual FReply OnDragOver( const FGeometry& MyGeometry, const FDragDropEvent& DragDropEvent ) override;
 	virtual FReply OnDrop( const FGeometry& MyGeometry, const FDragDropEvent& DragDropEvent ) override;
+	virtual bool OnVisualizeTooltip( const TSharedPtr<SWidget>& TooltipContent ) override;
+	virtual void OnToolTipClosing() override;
 
 	/** Returns the color this item should be tinted with */
 	virtual FSlateColor GetAssetColor() const;
@@ -134,12 +137,6 @@ public:
 	/** Delegate handling when an asset is loaded */
 	void HandleAssetLoaded(UObject* InAsset) const;
 
-	/** About to show a tool tip */
-	virtual bool OnVisualizeTooltip( const TSharedPtr<SWidget>& TooltipContent ) override;
-
-	/** Tooltip is closing */
-	virtual void OnToolTipClosing() override;
-
 protected:
 	/** Used by OnDragEnter, OnDragOver, and OnDrop to check and update the validity of the drag operation */
 	bool ValidateDragDrop( const FGeometry& MyGeometry, const FDragDropEvent& DragDropEvent, bool& OutIsKnownDragOperation ) const;
@@ -655,6 +652,31 @@ public:
 		this->AssetColumnItem->Tick(AllottedGeometry, InCurrentTime, InDeltaTime);
 	}
 
+	virtual TSharedPtr<IToolTip> GetToolTip() override
+	{
+		return AssetColumnItem->GetToolTip();
+	}
+
+	virtual void OnDragEnter(const FGeometry& MyGeometry, const FDragDropEvent& DragDropEvent) override
+	{
+		AssetColumnItem->OnDragEnter(MyGeometry, DragDropEvent);
+	}
+
+	virtual void OnDragLeave(const FDragDropEvent& DragDropEvent) override
+	{
+		AssetColumnItem->OnDragLeave(DragDropEvent);
+	}
+
+	virtual FReply OnDragOver(const FGeometry& MyGeometry, const FDragDropEvent& DragDropEvent) override
+	{
+		return AssetColumnItem->OnDragOver(MyGeometry, DragDropEvent);
+	}
+
+	virtual FReply OnDrop(const FGeometry& MyGeometry, const FDragDropEvent& DragDropEvent) override
+	{
+		return AssetColumnItem->OnDrop(MyGeometry, DragDropEvent);
+	}
+
 	virtual bool OnVisualizeTooltip(const TSharedPtr<SWidget>& TooltipContent) override
 	{
 		// We take the content from the asset column item during construction,
@@ -662,5 +684,10 @@ public:
 		return AssetColumnItem->OnVisualizeTooltip(TooltipContent);
 	}
 
+	virtual void OnToolTipClosing() override
+	{
+		AssetColumnItem->OnToolTipClosing();
+	}
+
 	TSharedPtr<SAssetColumnItem> AssetColumnItem;
 };
diff --git a/Engine/Source/Editor/ContentBrowser/Private/SAssetView.cpp b/Engine/Source/Editor/ContentBrowser/Private/SAssetView.cpp
--- a/Engine/Source/Editor/ContentBrowser/Private/SAssetView.cpp
+++ b/Engine/Source/Editor/ContentBrowser/Private/SAssetView.cpp
@@ -2775,7 +2775,10 @@ TSharedRef<ITableRow> SAssetView::MakeListViewWidget(TSharedPtr<FAssetViewItem>
 			.OnItemDestroyed(this, &SAssetView::AssetItemWidgetDestroyed)
 			.ShouldAllowToolTip(this, &SAssetView::ShouldAllowToolTips)
 			.HighlightText(HighlightedText)
-			.IsSelected( FIsSelected::CreateSP(TableRowWidget.Get(), &STableRow<TSharedPtr<FAssetViewItem>>::IsSelectedExclusively) );
+			.IsSelected( FIsSelected::CreateSP(TableRowWidget.Get(), &STableRow<TSharedPtr<FAssetViewItem>>::IsSelectedExclusively) )
+			.OnAssetsDragDropped(this, &SAssetView::OnAssetsDragDropped)
+			.OnPathsDragDropped(this, &SAssetView::OnPathsDragDropped)
+			.OnFilesDragDropped(this, &SAssetView::OnFilesDragDropped);
 
 		TableRowWidget->SetContent(Item);
 
