Added copy/paste keyboard shortcuts to the content browser.

UE-16354

diff --git a/Engine/Source/Editor/ContentBrowser/Private/SAssetView.cpp b/Engine/Source/Editor/ContentBrowser/Private/SAssetView.cpp
--- a/Engine/Source/Editor/ContentBrowser/Private/SAssetView.cpp
+++ b/Engine/Source/Editor/ContentBrowser/Private/SAssetView.cpp
@@ -1207,10 +1207,30 @@ FReply SAssetView::OnKeyChar( const FGeometry& MyGeometry,const FCharacterEvent&
 FReply SAssetView::OnKeyDown( const FGeometry& MyGeometry, const FKeyEvent& InKeyEvent )
 {
 	{
+		const bool bTestOnly = true;
+		if ( InKeyEvent.IsControlDown() && InKeyEvent.GetCharacter() == 'V' && IsAssetPathSelected() )
+		{
+			FString DestPaths;
+			TArray<FString> DestPathsSplit;
+
+			// Get the copied asset paths
+			FPlatformMisc::ClipboardPaste( DestPaths );
+			DestPaths.ParseIntoArrayLines( DestPathsSplit );
+
+			// Get assets and copy them
+			TArray<UObject*> ObjectsToCopy;
+			for (FString DestPath : DestPathsSplit)
+			{
+				if ( !(DestPath == TEXT("None")) )
+				{
+					ObjectsToCopy.Add( LoadObject<UObject>( NULL, *DestPath ));
+				}
+			}
+			ContentBrowserUtils::CopyAssets( ObjectsToCopy, SourcesData.PackagePaths[0].ToString() );
+		}
 		// Swallow the key-presses used by the quick-jump in OnKeyChar to avoid other things (such as the viewport commands) getting them instead
 		// eg) Pressing "W" without this would set the viewport to "translate" mode
-		const bool bTestOnly = true;
-		if(HandleQuickJumpKeyDown(InKeyEvent.GetCharacter(), InKeyEvent.IsControlDown(), InKeyEvent.IsAltDown(), bTestOnly).IsEventHandled())
+		else if( HandleQuickJumpKeyDown(InKeyEvent.GetCharacter(), InKeyEvent.IsControlDown(), InKeyEvent.IsAltDown(), bTestOnly).IsEventHandled() )
 		{
 			return FReply::Handled();
 		}
