- Enable path column and give option for asset picker
- Added type to be optional as lots of asset picker only shows one type
- Added new option to sort by path first
- Made persona asset browser to show path, and sort by path
- Made skeleton widget to hide type but path.

Merging

//depot/UE4-Orion/Engine/Source/Editor/...

to //depot/UE4/Engine/Source/Editor/...

diff --git a/Engine/Source/Editor/ContentBrowser/Private/SAssetPicker.cpp b/Engine/Source/Editor/ContentBrowser/Private/SAssetPicker.cpp
--- a/Engine/Source/Editor/ContentBrowser/Private/SAssetPicker.cpp
+++ b/Engine/Source/Editor/ContentBrowser/Private/SAssetPicker.cpp
@@ -219,6 +219,9 @@ void SAssetPicker::Construct( const FArguments& InArgs )
 		.AllowDragging( InArgs._AssetPickerConfig.bAllowDragging )
 		.CanShowClasses( InArgs._AssetPickerConfig.bCanShowClasses )
 		.CanShowFolders( InArgs._AssetPickerConfig.bCanShowFolders )
+		.ShowPathInColumnView( InArgs._AssetPickerConfig.bShowPathInColumnView)
+		.ShowTypeInColumnView( InArgs._AssetPickerConfig.bShowTypeInColumnView)
+		.SortByPathInColumnView( InArgs._AssetPickerConfig.bSortByPathInColumnView)
 		.FilterRecursivelyWithBackendFilter( false )
 		.CanShowRealTimeThumbnails( InArgs._AssetPickerConfig.bCanShowRealTimeThumbnails )
 		.CanShowDevelopersFolder( InArgs._AssetPickerConfig.bCanShowDevelopersFolder )
diff --git a/Engine/Source/Editor/ContentBrowser/Private/SAssetView.cpp b/Engine/Source/Editor/ContentBrowser/Private/SAssetView.cpp
--- a/Engine/Source/Editor/ContentBrowser/Private/SAssetView.cpp
+++ b/Engine/Source/Editor/ContentBrowser/Private/SAssetView.cpp
@@ -154,6 +154,10 @@ void SAssetView::Construct( const FArguments& InArgs )
 
 	SelectionMode = InArgs._SelectionMode;
 
+	bShowPathInColumnView = InArgs._ShowPathInColumnView;
+	bShowTypeInColumnView = InArgs._ShowTypeInColumnView;
+	bSortByPathInColumnView = bShowPathInColumnView & InArgs._SortByPathInColumnView;
+
 	bPendingUpdateThumbnails = false;
 	CurrentThumbnailSize = TileViewThumbnailSize;
 
@@ -394,6 +398,19 @@ void SAssetView::Construct( const FArguments& InArgs )
 		AssetsToSync.Add( InArgs._InitialAssetSelection );
 		SyncToAssets( AssetsToSync );
 	}
+
+	// If currently looking at column, and you could choose to sort by path in column first and then name
+	// Generalizing this is a bit difficult because the column ID is not accessible or is not known
+	// Currently I assume this won't work, if this view mode is not column. Otherwise, I don't think sorting by path
+	// is a good idea. 
+	if (CurrentViewType == EAssetViewType::Column && bSortByPathInColumnView)
+	{
+		SortManager.SetSortColumnId(EColumnSortPriority::Primary, SortManager.PathColumnId);
+		SortManager.SetSortColumnId(EColumnSortPriority::Secondary, SortManager.NameColumnId);
+		SortManager.SetSortMode(EColumnSortPriority::Primary, EColumnSortMode::Ascending);
+		SortManager.SetSortMode(EColumnSortPriority::Secondary, EColumnSortMode::Ascending);
+		SortList();
+	}
 }
 
 END_SLATE_FUNCTION_BUILD_OPTIMIZATION
@@ -1287,7 +1304,7 @@ TSharedRef<SAssetListView> SAssetView::CreateListView()
 
 TSharedRef<SAssetColumnView> SAssetView::CreateColumnView()
 {
-	return SNew(SAssetColumnView)
+	TSharedPtr<SAssetColumnView> ColumnView = SNew(SAssetColumnView)
 		.SelectionMode( SelectionMode )
 		.ListItemsSource(&FilteredAssetItems)
 		.OnGenerateRow(this, &SAssetView::MakeColumnViewWidget)
@@ -1305,14 +1322,34 @@ TSharedRef<SAssetColumnView> SAssetView::CreateColumnView()
 			.SortPriority(TAttribute< EColumnSortPriority::Type >::Create(TAttribute< EColumnSortPriority::Type >::FGetter::CreateSP(this, &SAssetView::GetColumnSortPriority, SortManager.NameColumnId)))
 			.OnSort( FOnSortModeChanged::CreateSP( this, &SAssetView::OnSortColumnHeader ) )
 			.DefaultLabel( LOCTEXT("Column_Name", "Name") )
-			//@TODO: Query the OnAssetTagWantsToBeDisplayed column filter here too, in case the user wants to bury the type column
-			+ SHeaderRow::Column(SortManager.ClassColumnId)
-			.FillWidth(160)
-			.SortMode( TAttribute< EColumnSortMode::Type >::Create( TAttribute< EColumnSortMode::Type >::FGetter::CreateSP( this, &SAssetView::GetColumnSortMode, SortManager.ClassColumnId ) ) )
-			.SortPriority(TAttribute< EColumnSortPriority::Type >::Create(TAttribute< EColumnSortPriority::Type >::FGetter::CreateSP(this, &SAssetView::GetColumnSortPriority, SortManager.ClassColumnId)))
-			.OnSort( FOnSortModeChanged::CreateSP( this, &SAssetView::OnSortColumnHeader ) )
-			.DefaultLabel( LOCTEXT("Column_Class", "Type") )
 		);
+
+	if(bShowTypeInColumnView)
+	{
+		ColumnView->GetHeaderRow()->AddColumn(
+				SHeaderRow::Column(SortManager.ClassColumnId)
+				.FillWidth(160)
+				.SortMode(TAttribute< EColumnSortMode::Type >::Create(TAttribute< EColumnSortMode::Type >::FGetter::CreateSP(this, &SAssetView::GetColumnSortMode, SortManager.ClassColumnId)))
+				.SortPriority(TAttribute< EColumnSortPriority::Type >::Create(TAttribute< EColumnSortPriority::Type >::FGetter::CreateSP(this, &SAssetView::GetColumnSortPriority, SortManager.ClassColumnId)))
+				.OnSort(FOnSortModeChanged::CreateSP(this, &SAssetView::OnSortColumnHeader))
+				.DefaultLabel(LOCTEXT("Column_Class", "Type"))
+			);
+	}
+
+
+	if (bShowPathInColumnView)
+	{
+		ColumnView->GetHeaderRow()->AddColumn(
+				SHeaderRow::Column(SortManager.PathColumnId)
+				.FillWidth(160)
+				.SortMode(TAttribute< EColumnSortMode::Type >::Create(TAttribute< EColumnSortMode::Type >::FGetter::CreateSP(this, &SAssetView::GetColumnSortMode, SortManager.PathColumnId)))
+				.SortPriority(TAttribute< EColumnSortPriority::Type >::Create(TAttribute< EColumnSortPriority::Type >::FGetter::CreateSP(this, &SAssetView::GetColumnSortPriority, SortManager.PathColumnId)))
+				.OnSort(FOnSortModeChanged::CreateSP(this, &SAssetView::OnSortColumnHeader))
+				.DefaultLabel(LOCTEXT("Column_Path", "Path"))
+			);
+	}
+
+	return ColumnView.ToSharedRef();
 }
 
 bool SAssetView::IsValidSearchToken(const FString& Token) const
@@ -1786,7 +1823,8 @@ void SAssetView::SetMajorityAssetType(FName NewMajorityAssetType)
 		for ( int32 ColumnIdx = Columns.Num() - 1; ColumnIdx >= 0; --ColumnIdx )
 		{
 			const FName ColumnId = Columns[ColumnIdx].ColumnId;
-			if ( ColumnId != SortManager.NameColumnId && ColumnId != SortManager.ClassColumnId && ColumnId != NAME_None )
+			if ( ColumnId != SortManager.NameColumnId && (bShowTypeInColumnView && ColumnId != SortManager.ClassColumnId) && 
+					(bShowPathInColumnView && ColumnId != SortManager.PathColumnId) && ColumnId != NAME_None )
 			{
 				ColumnView->GetHeaderRow()->RemoveColumn(ColumnId);
 			}
diff --git a/Engine/Source/Editor/ContentBrowser/Private/SAssetView.h b/Engine/Source/Editor/ContentBrowser/Private/SAssetView.h
--- a/Engine/Source/Editor/ContentBrowser/Private/SAssetView.h
+++ b/Engine/Source/Editor/ContentBrowser/Private/SAssetView.h
@@ -32,6 +32,9 @@ public:
 		, _AllowDragging(true)
 		, _AllowFocusOnSync(true)
 		, _FillEmptySpaceInTileView(true)
+		, _ShowPathInColumnView(false)
+		, _ShowTypeInColumnView(true)
+		, _SortByPathInColumnView(false)
 		{}
 
 		/** Called to check if an asset should be filtered out by external code */
@@ -139,6 +142,15 @@ public:
 		/** Whether this asset view should allow the thumbnails to consume empty space after the user scale is applied */
 		SLATE_ARGUMENT( bool, FillEmptySpaceInTileView )
 
+		/** Should show Path in column view if true */
+		SLATE_ARGUMENT(bool, ShowPathInColumnView)
+
+		/** Should show Type in column view if true */
+		SLATE_ARGUMENT(bool, ShowTypeInColumnView)
+
+		/** Sort by path in the column view. Only works if the initial view type is Column */
+		SLATE_ARGUMENT(bool, SortByPathInColumnView)
+
 		/** Called to check if an asset tag should be display in details view. */
 		SLATE_EVENT( FOnShouldDisplayAssetTag, OnAssetTagWantsToBeDisplayed )
 
@@ -789,6 +801,15 @@ private:
 	/** Indicates if the context menu is going to load the assets, and if so to preload before the context menu is shown, and warn about the pending load. */
 	bool bPreloadAssetsForContextMenu;
 
+	/** If true, it will show path column in the asset view */
+	bool bShowPathInColumnView;
+
+	/** If true, it will show type in the asset view */
+	bool bShowTypeInColumnView;
+
+	/** If true, it sorts by path and then name */
+	bool bSortByPathInColumnView;
+
 	/** The current selection mode used by the asset view */
 	ESelectionMode::Type SelectionMode;
 
diff --git a/Engine/Source/Editor/ContentBrowser/Public/IContentBrowserSingleton.h b/Engine/Source/Editor/ContentBrowser/Public/IContentBrowserSingleton.h
--- a/Engine/Source/Editor/ContentBrowser/Public/IContentBrowserSingleton.h
+++ b/Engine/Source/Editor/ContentBrowser/Public/IContentBrowserSingleton.h
@@ -134,6 +134,13 @@ struct FAssetPickerConfig
 	/** Indicates that we would like to build the filter UI with the Asset Picker */
 	bool bAddFilterUI;
 
+	/** If true, show path in column view */
+	bool bShowPathInColumnView; 
+	/** If true, show class in column view */
+	bool bShowTypeInColumnView;
+	/** If true, sort by path in column view. Only works if initial view type is Column */
+	bool bSortByPathInColumnView;
+
 	/** Override the default filter context menu layout */
 	EAssetTypeCategories::Type DefaultFilterMenuExpansion;
 
@@ -153,6 +160,9 @@ struct FAssetPickerConfig
 		, bCanShowDevelopersFolder(false)
 		, bPreloadAssetsForContextMenu(true)
 		, bAddFilterUI(false)
+		, bShowPathInColumnView(false)
+		, bShowTypeInColumnView(true)
+		, bSortByPathInColumnView(false)
 		, DefaultFilterMenuExpansion(EAssetTypeCategories::Basic)
 	{}
 };
diff --git a/Engine/Source/Editor/Persona/Private/SAnimationSequenceBrowser.cpp b/Engine/Source/Editor/Persona/Private/SAnimationSequenceBrowser.cpp
--- a/Engine/Source/Editor/Persona/Private/SAnimationSequenceBrowser.cpp
+++ b/Engine/Source/Editor/Persona/Private/SAnimationSequenceBrowser.cpp
@@ -422,6 +422,8 @@ void SAnimationSequenceBrowser::Construct(const FArguments& InArgs)
 	Config.Filter.ClassNames.Add(UVertexAnimation::StaticClass()->GetFName()); //@TODO: Is currently ignored due to the skeleton check
 	Config.InitialAssetViewType = EAssetViewType::Column;
 	Config.bAddFilterUI = true;
+	Config.bShowPathInColumnView = true;
+	Config.bSortByPathInColumnView = true;
 
 	TSharedPtr<FPersona> Persona = PersonaPtr.Pin();
 	if (Persona.IsValid())
diff --git a/Engine/Source/Editor/UnrealEd/Private/SSkeletonWidget.cpp b/Engine/Source/Editor/UnrealEd/Private/SSkeletonWidget.cpp
--- a/Engine/Source/Editor/UnrealEd/Private/SSkeletonWidget.cpp
+++ b/Engine/Source/Editor/UnrealEd/Private/SSkeletonWidget.cpp
@@ -25,6 +25,8 @@ void SSkeletonListWidget::Construct(const FArguments& InArgs)
 	AssetPickerConfig.Filter.ClassNames.Add(USkeleton::StaticClass()->GetFName());
 	AssetPickerConfig.OnAssetSelected = FOnAssetSelected::CreateSP(this, &SSkeletonListWidget::SkeletonSelectionChanged);
 	AssetPickerConfig.InitialAssetViewType = EAssetViewType::Column;
+	AssetPickerConfig.bShowPathInColumnView = true;
+	AssetPickerConfig.bShowTypeInColumnView = false;
 
 	this->ChildSlot
 		[
@@ -113,6 +115,8 @@ void SSkeletonCompareWidget::Construct(const FArguments& InArgs)
 	AssetPickerConfig.Filter.ClassNames.Add(USkeleton::StaticClass()->GetFName());
 	AssetPickerConfig.OnAssetSelected = FOnAssetSelected::CreateSP(this, &SSkeletonCompareWidget::SkeletonSelectionChanged);
 	AssetPickerConfig.InitialAssetViewType = EAssetViewType::Column;
+	AssetPickerConfig.bShowPathInColumnView = true;
+	AssetPickerConfig.bShowTypeInColumnView = false;
 
 	TSharedPtr<SToolTip> SkeletonTooltip = IDocumentation::Get()->CreateToolTip(FText::FromString("Pick a skeleton for this mesh"), NULL, FString("Shared/Editors/Persona"), FString("Skeleton"));
 
@@ -384,6 +388,8 @@ void SAnimationRemapSkeleton::UpdateAssetPicker()
 	AssetPickerConfig.bAllowNullSelection = false;
 	AssetPickerConfig.InitialAssetViewType = EAssetViewType::Column;
 	AssetPickerConfig.OnShouldFilterAsset = FOnShouldFilterAsset::CreateSP(this, &SAnimationRemapSkeleton::OnShouldFilterAsset);
+	AssetPickerConfig.bShowPathInColumnView = true;
+	AssetPickerConfig.bShowTypeInColumnView = false;
 
 	FContentBrowserModule& ContentBrowserModule = FModuleManager::Get().LoadModuleChecked<FContentBrowserModule>(TEXT("ContentBrowser"));
 
