Fix CIS (non-unity build, shadowing issue).

diff --git a/Engine/Source/Editor/ContentBrowser/Private/SAssetView.cpp b/Engine/Source/Editor/ContentBrowser/Private/SAssetView.cpp
--- a/Engine/Source/Editor/ContentBrowser/Private/SAssetView.cpp
+++ b/Engine/Source/Editor/ContentBrowser/Private/SAssetView.cpp
@@ -1304,7 +1304,7 @@ TSharedRef<SAssetListView> SAssetView::CreateListView()
 
 TSharedRef<SAssetColumnView> SAssetView::CreateColumnView()
 {
-	TSharedPtr<SAssetColumnView> ColumnView = SNew(SAssetColumnView)
+	TSharedPtr<SAssetColumnView> NewColumnView = SNew(SAssetColumnView)
 		.SelectionMode( SelectionMode )
 		.ListItemsSource(&FilteredAssetItems)
 		.OnGenerateRow(this, &SAssetView::MakeColumnViewWidget)
@@ -1326,7 +1326,7 @@ TSharedRef<SAssetColumnView> SAssetView::CreateColumnView()
 
 	if(bShowTypeInColumnView)
 	{
-		ColumnView->GetHeaderRow()->AddColumn(
+		NewColumnView->GetHeaderRow()->AddColumn(
 				SHeaderRow::Column(SortManager.ClassColumnId)
 				.FillWidth(160)
 				.SortMode(TAttribute< EColumnSortMode::Type >::Create(TAttribute< EColumnSortMode::Type >::FGetter::CreateSP(this, &SAssetView::GetColumnSortMode, SortManager.ClassColumnId)))
@@ -1339,7 +1339,7 @@ TSharedRef<SAssetColumnView> SAssetView::CreateColumnView()
 
 	if (bShowPathInColumnView)
 	{
-		ColumnView->GetHeaderRow()->AddColumn(
+		NewColumnView->GetHeaderRow()->AddColumn(
 				SHeaderRow::Column(SortManager.PathColumnId)
 				.FillWidth(160)
 				.SortMode(TAttribute< EColumnSortMode::Type >::Create(TAttribute< EColumnSortMode::Type >::FGetter::CreateSP(this, &SAssetView::GetColumnSortMode, SortManager.PathColumnId)))
@@ -1349,7 +1349,7 @@ TSharedRef<SAssetColumnView> SAssetView::CreateColumnView()
 			);
 	}
 
-	return ColumnView.ToSharedRef();
+	return NewColumnView.ToSharedRef();
 }
 
 bool SAssetView::IsValidSearchToken(const FString& Token) const
