Removed restrictive CanExecuteDiffSelected logic, it's safe to diff any matching asset type

diff --git a/Engine/Source/Editor/ContentBrowser/Private/AssetContextMenu.cpp b/Engine/Source/Editor/ContentBrowser/Private/AssetContextMenu.cpp
--- a/Engine/Source/Editor/ContentBrowser/Private/AssetContextMenu.cpp
+++ b/Engine/Source/Editor/ContentBrowser/Private/AssetContextMenu.cpp
@@ -2016,14 +2016,7 @@ bool FAssetContextMenu::CanExecuteDiffSelected() const
 		FAssetData const& FirstSelection = SelectedAssets[0];
 		FAssetData const& SecondSelection = SelectedAssets[1];
 
-		if (FirstSelection.AssetClass != SecondSelection.AssetClass)
-		{
-			bCanDiffSelected = false;
-		}
-		else 
-		{
-			bCanDiffSelected = (FirstSelection.AssetClass == UBlueprint::StaticClass()->GetFName());
-		}
+		bCanDiffSelected = FirstSelection.AssetClass == SecondSelection.AssetClass;
 	}
 
 	return bCanDiffSelected;
