Asset source content files are now visible on asset tooltips

diff --git a/Engine/Source/Editor/ContentBrowser/Private/AssetViewWidgets.cpp b/Engine/Source/Editor/ContentBrowser/Private/AssetViewWidgets.cpp
--- a/Engine/Source/Editor/ContentBrowser/Private/AssetViewWidgets.cpp
+++ b/Engine/Source/Editor/ContentBrowser/Private/AssetViewWidgets.cpp
@@ -7,6 +7,7 @@
 #include "AssetViewTypes.h"
 #include "AssetViewWidgets.h"
 #include "SThumbnailEditModeTools.h"
+#include "AssetSourceFilenameCache.h"
 #include "CollectionManagerModule.h"
 #include "CollectionViewUtils.h"
 #include "DragAndDrop/AssetDragDropOp.h"
@@ -646,6 +647,16 @@ TSharedRef<SWidget> SAssetViewItem::CreateToolTipWidget() const
 				AddToToolTipInfoBox(InfoBox, ToolTipTagItem.Key, ToolTipTagItem.Value, ToolTipTagItem.bImportant);
 			}
 
+			// Add asset source files
+			TOptional<FAssetImportInfo> ImportInfo = FAssetSourceFilenameCache::ExtractAssetImportInfo(AssetData.TagsAndValues);
+			if (ImportInfo.IsSet())
+			{
+				for (const auto& File : ImportInfo->SourceFiles)
+				{
+					AddToToolTipInfoBox( InfoBox, LOCTEXT("TileViewTooltipSourceFile", "Source File"), FText::FromString(File.RelativeFilename), false );
+				}
+			}
+
 			TSharedRef<SVerticalBox> OverallTooltipVBox = SNew(SVerticalBox);
 
 			// Top section (asset name, type, is checked out)
