jira: UE-13247
Add a guard in case we cannot load an asset and notify the user the operation fail to execute.

diff --git a/Engine/Source/Editor/ContentBrowser/Private/SAssetView.cpp b/Engine/Source/Editor/ContentBrowser/Private/SAssetView.cpp
--- a/Engine/Source/Editor/ContentBrowser/Private/SAssetView.cpp
+++ b/Engine/Source/Editor/ContentBrowser/Private/SAssetView.cpp
@@ -15,6 +15,9 @@
 #include "KismetEditorUtilities.h"
 #include "IPluginManager.h"
 #include "NativeClassHierarchy.h"
+#include "MessageLog.h"
+#include "SNotificationList.h"
+#include "NotificationManager.h"
 
 #define LOCTEXT_NAMESPACE "ContentBrowser"
 
@@ -3454,8 +3457,26 @@ void SAssetView::AssetRenameCommit(const TSharedPtr<FAssetViewItem>& Item, const
 
 		// Committed rename
 		Asset = ItemAsAsset->Data.GetAsset();
-		ContentBrowserUtils::RenameAsset(Asset, NewName, ErrorMessage);
-		bSuccess = true;
+		if(Asset == NULL)
+		{
+			//put back the original name
+			RenamingAsset.Reset();
+			
+			//Notify the user rename fail and link the output log
+			FNotificationInfo Info(LOCTEXT("RenameAssetsFailed", "Failed to rename assets"));
+			Info.ExpireDuration = 5.0f;
+			Info.Hyperlink = FSimpleDelegate::CreateStatic([](){ FGlobalTabmanager::Get()->InvokeTab(FName("OutputLog")); });
+			Info.HyperlinkText = LOCTEXT("ShowOutputLogHyperlink", "Show Output Log");
+			FSlateNotificationManager::Get().AddNotification(Info);
+			
+			//Set the content browser error message
+			ErrorMessage = LOCTEXT("RenameAssetsFailed", "Failed to rename assets");
+		}
+		else
+		{
+			ContentBrowserUtils::RenameAsset(Asset, NewName, ErrorMessage);
+			bSuccess = true;
+		}
 	}
 	else if ( ItemType == EAssetItemType::Creation || ItemType == EAssetItemType::Duplication )
 	{
