Fixed the Content Browser details view not removing type-specific columns

UE-20412 - Content Browser continues duplicating Blueprint columns as folders are navigated

diff --git a/Engine/Source/Editor/ContentBrowser/Private/SAssetView.cpp b/Engine/Source/Editor/ContentBrowser/Private/SAssetView.cpp
--- a/Engine/Source/Editor/ContentBrowser/Private/SAssetView.cpp
+++ b/Engine/Source/Editor/ContentBrowser/Private/SAssetView.cpp
@@ -1826,8 +1826,12 @@ void SAssetView::SetMajorityAssetType(FName NewMajorityAssetType)
 		for ( int32 ColumnIdx = Columns.Num() - 1; ColumnIdx >= 0; --ColumnIdx )
 		{
 			const FName ColumnId = Columns[ColumnIdx].ColumnId;
-			if ( ColumnId != SortManager.NameColumnId && (bShowTypeInColumnView && ColumnId != SortManager.ClassColumnId) && 
-					(bShowPathInColumnView && ColumnId != SortManager.PathColumnId) && ColumnId != NAME_None )
+
+			const bool bIsFixedNameColumn = ColumnId == SortManager.NameColumnId;
+			const bool bIsFixedClassColumn = bShowTypeInColumnView && ColumnId == SortManager.ClassColumnId;
+			const bool bIsFixedPathColumn = bShowPathInColumnView && ColumnId == SortManager.PathColumnId;
+
+			if ( ColumnId != NAME_None && !(bIsFixedNameColumn || bIsFixedClassColumn || bIsFixedPathColumn) )
 			{
 				ColumnView->GetHeaderRow()->RemoveColumn(ColumnId);
 			}
