Reinstanted Thumbnail Edit Mode Primitive Tools.

diff --git a/Engine/Source/Editor/ContentBrowser/Private/AssetViewWidgets.cpp b/Engine/Source/Editor/ContentBrowser/Private/AssetViewWidgets.cpp
--- a/Engine/Source/Editor/ContentBrowser/Private/AssetViewWidgets.cpp
+++ b/Engine/Source/Editor/ContentBrowser/Private/AssetViewWidgets.cpp
@@ -174,7 +174,7 @@ TSharedRef<SWidget> FAssetViewItemHelper::CreateListTileItemContents(T* const In
 		ItemContentsOverlay->AddSlot()
 		[
 			SNew(SThumbnailEditModeTools, InTileOrListItem->AssetThumbnail)
-			.SmallView(true)
+			.SmallView(!InTileOrListItem->CanDisplayPrimitiveTools())
 			.Visibility(InTileOrListItem, &T::GetThumbnailEditModeUIVisibility)
 		];
 
diff --git a/Engine/Source/Editor/ContentBrowser/Private/AssetViewWidgets.h b/Engine/Source/Editor/ContentBrowser/Private/AssetViewWidgets.h
--- a/Engine/Source/Editor/ContentBrowser/Private/AssetViewWidgets.h
+++ b/Engine/Source/Editor/ContentBrowser/Private/AssetViewWidgets.h
@@ -407,6 +407,9 @@ public:
 	/** Handles committing a name change */
 	virtual void OnAssetDataChanged() override;
 
+	/** Whether the widget should allow primitive tools to be displayed */
+	bool CanDisplayPrimitiveTools() const { return false; }
+
 private:
 	/** Returns the size of the thumbnail widget */
 	FOptionalSize GetThumbnailBoxSize() const;
@@ -515,6 +518,9 @@ public:
 	/** Handles committing a name change */
 	virtual void OnAssetDataChanged() override;
 
+	/** Whether the widget should allow primitive tools to be displayed */
+	bool CanDisplayPrimitiveTools() const { return true; }
+
 protected:
 	/** SAssetViewItem interface */
 	virtual float GetNameTextWrapWidth() const override { return LastGeometry.Size.X - 2.f; }
