Content Browser: Hacks for VR editor "demo mode" to shrink CB thumbnail border when using DPI scaling

diff --git a/Engine/Source/Editor/ContentBrowser/Private/SAssetView.cpp b/Engine/Source/Editor/ContentBrowser/Private/SAssetView.cpp
--- a/Engine/Source/Editor/ContentBrowser/Private/SAssetView.cpp
+++ b/Engine/Source/Editor/ContentBrowser/Private/SAssetView.cpp
@@ -130,6 +130,13 @@ void SAssetView::Construct( const FArguments& InArgs )
 	TileViewThumbnailResolution = 256;
 	TileViewThumbnailSize = 128;
 	TileViewThumbnailPadding = 5;
+
+	const bool bIsVREditorDemo = FParse::Param( FCommandLine::Get(), TEXT( "VREditorDemo" ) );	// @todo vreditor: Remove this when no longer needed
+	if( bIsVREditorDemo )
+	{
+		TileViewThumbnailPadding = 0;
+	}
+
 	TileViewNameHeight = 36;
 	ThumbnailScaleSliderValue = InArgs._ThumbnailScale; 
 
@@ -992,6 +999,12 @@ void SAssetView::CalculateFillScale( const FGeometry& AllottedGeometry )
 	{
 		float ItemWidth = GetTileViewItemBaseWidth();
 
+ 		const bool bIsVREditorDemo = FParse::Param( FCommandLine::Get(), TEXT( "VREditorDemo" ) );	// @todo vreditor: Remove this when no longer needed
+		if( bIsVREditorDemo )
+		{
+			ItemWidth /= AllottedGeometry.Scale;
+		}
+
 		// Scrollbars are 16, but we add 1 to deal with half pixels.
 		const float ScrollbarWidth = 16 + 1;
 		float TotalWidth = AllottedGeometry.Size.X - ( ScrollbarWidth / AllottedGeometry.Scale );
