Copying //UE4/Dev-Editor to //UE4/Dev-Main (Source: //UE4/Dev-Editor @ 2973866)

diff --git a/Engine/Source/Editor/BehaviorTreeEditor/Private/SBehaviorTreeBlackboardView.cpp b/Engine/Source/Editor/BehaviorTreeEditor/Private/SBehaviorTreeBlackboardView.cpp
--- a/Engine/Source/Editor/BehaviorTreeEditor/Private/SBehaviorTreeBlackboardView.cpp
+++ b/Engine/Source/Editor/BehaviorTreeEditor/Private/SBehaviorTreeBlackboardView.cpp
@@ -130,7 +130,7 @@ private:
 
 		if(BlackboardEntryAction->Key.KeyType)
 		{
-			OutIconBrush = FClassIconFinder::FindIconForClass(BlackboardEntryAction->Key.KeyType->GetClass());
+			OutIconBrush = FSlateIconFinder::FindIconBrushForClass(BlackboardEntryAction->Key.KeyType->GetClass());
 		}
 	}
 
diff --git a/Engine/Source/Editor/BlueprintGraph/Classes/K2Node_BreakStruct.h b/Engine/Source/Editor/BlueprintGraph/Classes/K2Node_BreakStruct.h
--- a/Engine/Source/Editor/BlueprintGraph/Classes/K2Node_BreakStruct.h
+++ b/Engine/Source/Editor/BlueprintGraph/Classes/K2Node_BreakStruct.h
@@ -41,7 +41,7 @@ class UK2Node_BreakStruct : public UK2Node_StructMemberGet
 	virtual FLinearColor GetNodeTitleColor() const override;
 	virtual FText GetTooltipText() const override;
 	virtual void ValidateNodeDuringCompilation(class FCompilerResultsLog& MessageLog) const override;
-	virtual FName GetPaletteIcon(FLinearColor& OutColor) const override{ return TEXT("GraphEditor.BreakStruct_16x"); }
+	virtual FSlateIcon GetIconAndTint(FLinearColor& OutColor) const override;
 	//~ End  UEdGraphNode Interface
 
 	//~ Begin K2Node Interface
diff --git a/Engine/Source/Editor/BlueprintGraph/Classes/K2Node_CallFunction.h b/Engine/Source/Editor/BlueprintGraph/Classes/K2Node_CallFunction.h
--- a/Engine/Source/Editor/BlueprintGraph/Classes/K2Node_CallFunction.h
+++ b/Engine/Source/Editor/BlueprintGraph/Classes/K2Node_CallFunction.h
@@ -76,7 +76,7 @@ public:
 	virtual void PostPlacedNewNode() override;
 	virtual FString GetDocumentationLink() const override;
 	virtual FString GetDocumentationExcerptName() const override;
-	virtual FName GetPaletteIcon(FLinearColor& OutColor) const override;
+	virtual FSlateIcon GetIconAndTint(FLinearColor& OutColor) const override;
 	virtual bool CanPasteHere(const UEdGraph* TargetGraph) const override;
 	virtual void PinDefaultValueChanged(UEdGraphPin* Pin) override;
 	virtual void AddSearchMetaDataInfo(TArray<struct FSearchTagDataPair>& OutTaggedMetaData) const override;
@@ -182,7 +182,7 @@ public:
 	static bool CanFunctionSupportMultipleTargets(UFunction const* InFunction);
 
 	/** */
-	static FName GetPaletteIconForFunction(UFunction const* Function, FLinearColor& OutColor);
+	static FSlateIcon GetPaletteIconForFunction(UFunction const* Function, FLinearColor& OutColor);
 
 private: 
 	/* Looks at function metadata and properties to determine if this node should be using enum to exec expansion */
diff --git a/Engine/Source/Editor/BlueprintGraph/Classes/K2Node_CastByteToEnum.h b/Engine/Source/Editor/BlueprintGraph/Classes/K2Node_CastByteToEnum.h
--- a/Engine/Source/Editor/BlueprintGraph/Classes/K2Node_CastByteToEnum.h
+++ b/Engine/Source/Editor/BlueprintGraph/Classes/K2Node_CastByteToEnum.h
@@ -23,7 +23,7 @@ class UK2Node_CastByteToEnum : public UK2Node
 	virtual void AllocateDefaultPins() override;
 	virtual FText GetTooltipText() const override;
 	virtual FText GetNodeTitle(ENodeTitleType::Type TitleType) const override;
-	virtual FName GetPaletteIcon(FLinearColor& OutColor) const override{ return TEXT("GraphEditor.Enum_16x"); }
+	virtual FSlateIcon GetIconAndTint(FLinearColor& OutColor) const override;
 	virtual void ValidateNodeDuringCompilation(class FCompilerResultsLog& MessageLog) const override;
 	//~ End UEdGraphNode Interface
 
diff --git a/Engine/Source/Editor/BlueprintGraph/Classes/K2Node_CustomEvent.h b/Engine/Source/Editor/BlueprintGraph/Classes/K2Node_CustomEvent.h
--- a/Engine/Source/Editor/BlueprintGraph/Classes/K2Node_CustomEvent.h
+++ b/Engine/Source/Editor/BlueprintGraph/Classes/K2Node_CustomEvent.h
@@ -29,7 +29,7 @@ class UK2Node_CustomEvent : public UK2Node_Event
 	virtual FText GetTooltipText() const override;
 	virtual FString GetDocumentationLink() const override;
 	virtual FString GetDocumentationExcerptName() const override;
-	virtual FName GetPaletteIcon(FLinearColor& OutColor) const override;
+	virtual FSlateIcon GetIconAndTint(FLinearColor& OutColor) const override;
 	virtual void AutowireNewNode(UEdGraphPin* FromPin) override;
 	virtual void AddSearchMetaDataInfo(TArray<struct FSearchTagDataPair>& OutTaggedMetaData) const override;
 	//~ End UEdGraphNode Interface
diff --git a/Engine/Source/Editor/BlueprintGraph/Classes/K2Node_DynamicCast.h b/Engine/Source/Editor/BlueprintGraph/Classes/K2Node_DynamicCast.h
--- a/Engine/Source/Editor/BlueprintGraph/Classes/K2Node_DynamicCast.h
+++ b/Engine/Source/Editor/BlueprintGraph/Classes/K2Node_DynamicCast.h
@@ -19,7 +19,7 @@ class UK2Node_DynamicCast : public UK2Node
 	virtual void AllocateDefaultPins() override;
 	virtual FLinearColor GetNodeTitleColor() const override;
 	virtual FText GetNodeTitle(ENodeTitleType::Type TitleType) const override;
-	virtual FName GetPaletteIcon(FLinearColor& OutColor) const override{ return TEXT("GraphEditor.Cast_16x"); }
+	virtual FSlateIcon GetIconAndTint(FLinearColor& OutColor) const override;
 	virtual void GetContextMenuActions(const FGraphNodeContextMenuBuilder& Context) const override;
 	virtual void PostReconstructNode() override;
 	virtual void PostPlacedNewNode() override;
diff --git a/Engine/Source/Editor/BlueprintGraph/Classes/K2Node_EaseFunction.h b/Engine/Source/Editor/BlueprintGraph/Classes/K2Node_EaseFunction.h
--- a/Engine/Source/Editor/BlueprintGraph/Classes/K2Node_EaseFunction.h
+++ b/Engine/Source/Editor/BlueprintGraph/Classes/K2Node_EaseFunction.h
@@ -37,7 +37,7 @@ class UK2Node_EaseFunction : public UK2Node
 	virtual void PinDefaultValueChanged(UEdGraphPin* Pin) override;
 	virtual void PinTypeChanged(UEdGraphPin* Pin) override;
 	virtual FText GetTooltipText() const override;
-	virtual FName GetPaletteIcon(FLinearColor& OutColor) const override;
+	virtual FSlateIcon GetIconAndTint(FLinearColor& OutColor) const override;
 	//~ End UEdGraphNode Interface.
 
 	//~ Begin UK2Node Interface.
diff --git a/Engine/Source/Editor/BlueprintGraph/Classes/K2Node_Event.h b/Engine/Source/Editor/BlueprintGraph/Classes/K2Node_Event.h
--- a/Engine/Source/Editor/BlueprintGraph/Classes/K2Node_Event.h
+++ b/Engine/Source/Editor/BlueprintGraph/Classes/K2Node_Event.h
@@ -60,7 +60,7 @@ class UK2Node_Event : public UK2Node_EditablePinBase
 	BLUEPRINTGRAPH_API virtual bool IsDeprecated() const override;
 	BLUEPRINTGRAPH_API virtual FString GetDeprecationMessage() const override;
 	BLUEPRINTGRAPH_API virtual UObject* GetJumpTargetForDoubleClick() const override;
-	BLUEPRINTGRAPH_API virtual FName GetPaletteIcon(FLinearColor& OutColor) const override{ return TEXT("GraphEditor.Event_16x"); }
+	BLUEPRINTGRAPH_API virtual FSlateIcon GetIconAndTint(FLinearColor& OutColor) const override;
 	BLUEPRINTGRAPH_API virtual FString GetFindReferenceSearchString() const override;
 	//~ End UEdGraphNode Interface
 
diff --git a/Engine/Source/Editor/BlueprintGraph/Classes/K2Node_ExecutionSequence.h b/Engine/Source/Editor/BlueprintGraph/Classes/K2Node_ExecutionSequence.h
--- a/Engine/Source/Editor/BlueprintGraph/Classes/K2Node_ExecutionSequence.h
+++ b/Engine/Source/Editor/BlueprintGraph/Classes/K2Node_ExecutionSequence.h
@@ -15,7 +15,7 @@ class UK2Node_ExecutionSequence : public UK2Node
 	virtual FText GetTooltipText() const override;
 	virtual FLinearColor GetNodeTitleColor() const override;
 	virtual FText GetNodeTitle(ENodeTitleType::Type TitleType) const override;
-	virtual FName GetPaletteIcon(FLinearColor& OutColor) const override{ return TEXT("GraphEditor.Sequence_16x"); }
+	virtual FSlateIcon GetIconAndTint(FLinearColor& OutColor) const override;
 	//~ End UEdGraphNode Interface
 
 	//~ Begin UK2Node Interface
diff --git a/Engine/Source/Editor/BlueprintGraph/Classes/K2Node_ForEachElementInEnum.h b/Engine/Source/Editor/BlueprintGraph/Classes/K2Node_ForEachElementInEnum.h
--- a/Engine/Source/Editor/BlueprintGraph/Classes/K2Node_ForEachElementInEnum.h
+++ b/Engine/Source/Editor/BlueprintGraph/Classes/K2Node_ForEachElementInEnum.h
@@ -20,7 +20,7 @@ class UK2Node_ForEachElementInEnum : public UK2Node
 	virtual void AllocateDefaultPins() override;
 	virtual FText GetTooltipText() const override;
 	virtual FText GetNodeTitle(ENodeTitleType::Type TitleType) const override;
-	virtual FName GetPaletteIcon(FLinearColor& OutColor) const override{ return TEXT("GraphEditor.Macro.Loop_16x"); }
+	virtual FSlateIcon GetIconAndTint(FLinearColor& OutColor) const override;
 	//~ End UEdGraphNode Interface
 
 	//~ Begin UK2Node Interface
diff --git a/Engine/Source/Editor/BlueprintGraph/Classes/K2Node_GetDataTableRow.h b/Engine/Source/Editor/BlueprintGraph/Classes/K2Node_GetDataTableRow.h
--- a/Engine/Source/Editor/BlueprintGraph/Classes/K2Node_GetDataTableRow.h
+++ b/Engine/Source/Editor/BlueprintGraph/Classes/K2Node_GetDataTableRow.h
@@ -18,11 +18,7 @@ class BLUEPRINTGRAPH_API UK2Node_GetDataTableRow : public UK2Node
 	virtual void PinDefaultValueChanged(UEdGraphPin* Pin) override;
 	virtual FText GetTooltipText() const override;
 	virtual void ExpandNode(class FKismetCompilerContext& CompilerContext, UEdGraph* SourceGraph) override;
-	virtual FName GetPaletteIcon(FLinearColor& OutColor) const override
-    {
-        OutColor = GetNodeTitleColor();
-        return TEXT("Kismet.AllClasses.FunctionIcon");
-    }
+	virtual FSlateIcon GetIconAndTint(FLinearColor& OutColor) const override;
 	//~ End UEdGraphNode Interface.
 
 	//~ Begin UK2Node Interface
diff --git a/Engine/Source/Editor/BlueprintGraph/Classes/K2Node_GetEnumeratorName.h b/Engine/Source/Editor/BlueprintGraph/Classes/K2Node_GetEnumeratorName.h
--- a/Engine/Source/Editor/BlueprintGraph/Classes/K2Node_GetEnumeratorName.h
+++ b/Engine/Source/Editor/BlueprintGraph/Classes/K2Node_GetEnumeratorName.h
@@ -16,7 +16,7 @@ class UK2Node_GetEnumeratorName : public UK2Node
 	virtual FText GetTooltipText() const override;
 	virtual FText GetNodeTitle(ENodeTitleType::Type TitleType) const override;
 	virtual void ValidateNodeDuringCompilation(class FCompilerResultsLog& MessageLog) const override;
-	virtual FName GetPaletteIcon(FLinearColor& OutColor) const override{ return TEXT("GraphEditor.Enum_16x"); }
+	virtual FSlateIcon GetIconAndTint(FLinearColor& OutColor) const override;
 	//~ End UEdGraphNode Interface
 
 	//~ Begin UK2Node Interface
diff --git a/Engine/Source/Editor/BlueprintGraph/Classes/K2Node_IfThenElse.h b/Engine/Source/Editor/BlueprintGraph/Classes/K2Node_IfThenElse.h
--- a/Engine/Source/Editor/BlueprintGraph/Classes/K2Node_IfThenElse.h
+++ b/Engine/Source/Editor/BlueprintGraph/Classes/K2Node_IfThenElse.h
@@ -15,7 +15,7 @@ class UK2Node_IfThenElse : public UK2Node
 	virtual FLinearColor GetNodeTitleColor() const override;
 	virtual FText GetTooltipText() const override;
 	virtual FText GetNodeTitle(ENodeTitleType::Type TitleType) const override;
-	virtual FName GetPaletteIcon(FLinearColor& OutColor) const override{ return TEXT("GraphEditor.Branch_16x"); }
+	virtual FSlateIcon GetIconAndTint(FLinearColor& OutColor) const override;
 	//~ End UEdGraphNode Interface
 
 	//~ Begin K2Node Interface.
diff --git a/Engine/Source/Editor/BlueprintGraph/Classes/K2Node_InputAction.h b/Engine/Source/Editor/BlueprintGraph/Classes/K2Node_InputAction.h
--- a/Engine/Source/Editor/BlueprintGraph/Classes/K2Node_InputAction.h
+++ b/Engine/Source/Editor/BlueprintGraph/Classes/K2Node_InputAction.h
@@ -35,7 +35,7 @@ class UK2Node_InputAction : public UK2Node
 	virtual FLinearColor GetNodeTitleColor() const override;
 	virtual FText GetNodeTitle(ENodeTitleType::Type TitleType) const override;
 	virtual FText GetTooltipText() const override;
-	virtual FName GetPaletteIcon(FLinearColor& OutColor) const override{ return TEXT("GraphEditor.Event_16x"); }
+	virtual FSlateIcon GetIconAndTint(FLinearColor& OutColor) const override;
 	virtual bool IsCompatibleWithGraph(UEdGraph const* Graph) const override;
 	//~ End UEdGraphNode Interface.
 
diff --git a/Engine/Source/Editor/BlueprintGraph/Classes/K2Node_InputTouch.h b/Engine/Source/Editor/BlueprintGraph/Classes/K2Node_InputTouch.h
--- a/Engine/Source/Editor/BlueprintGraph/Classes/K2Node_InputTouch.h
+++ b/Engine/Source/Editor/BlueprintGraph/Classes/K2Node_InputTouch.h
@@ -36,7 +36,7 @@ class UK2Node_InputTouch : public UK2Node
 	virtual FLinearColor GetNodeTitleColor() const override;
 	virtual FText GetNodeTitle(ENodeTitleType::Type TitleType) const override;
 	virtual FText GetTooltipText() const override;
-	virtual FName GetPaletteIcon(FLinearColor& OutColor) const override{ return TEXT("GraphEditor.TouchEvent_16x"); }
+	virtual FSlateIcon GetIconAndTint(FLinearColor& OutColor) const override;
 	virtual void GetMenuActions(FBlueprintActionDatabaseRegistrar& ActionRegistrar) const override;
 	virtual FText GetMenuCategory() const override;
 	virtual bool IsCompatibleWithGraph(const UEdGraph* TargetGraph) const override;
diff --git a/Engine/Source/Editor/BlueprintGraph/Classes/K2Node_MakeArray.h b/Engine/Source/Editor/BlueprintGraph/Classes/K2Node_MakeArray.h
--- a/Engine/Source/Editor/BlueprintGraph/Classes/K2Node_MakeArray.h
+++ b/Engine/Source/Editor/BlueprintGraph/Classes/K2Node_MakeArray.h
@@ -27,7 +27,7 @@ public:
 	virtual void AllocateDefaultPins() override;
 	virtual void PostReconstructNode() override;
 	virtual FText GetTooltipText() const override;
-	virtual FName GetPaletteIcon(FLinearColor& OutColor) const override{ return TEXT("GraphEditor.MakeArray_16x"); }
+	virtual FSlateIcon GetIconAndTint(FLinearColor& OutColor) const override;
 	virtual void ValidateNodeDuringCompilation(class FCompilerResultsLog& MessageLog) const override;
 	// End of UEdGraphNode interface
 
diff --git a/Engine/Source/Editor/BlueprintGraph/Classes/K2Node_MakeStruct.h b/Engine/Source/Editor/BlueprintGraph/Classes/K2Node_MakeStruct.h
--- a/Engine/Source/Editor/BlueprintGraph/Classes/K2Node_MakeStruct.h
+++ b/Engine/Source/Editor/BlueprintGraph/Classes/K2Node_MakeStruct.h
@@ -42,7 +42,7 @@ class UK2Node_MakeStruct : public UK2Node_StructMemberSet
 	virtual FLinearColor GetNodeTitleColor() const override;
 	virtual FText GetTooltipText() const override;
 	virtual void ValidateNodeDuringCompilation(class FCompilerResultsLog& MessageLog) const override;
-	virtual FName GetPaletteIcon(FLinearColor& OutColor) const override{ return TEXT("GraphEditor.MakeStruct_16x"); }
+	virtual FSlateIcon GetIconAndTint(FLinearColor& OutColor) const override;
 	virtual void PostPlacedNewNode() override;
 	//~ End  UEdGraphNode Interface
 
diff --git a/Engine/Source/Editor/BlueprintGraph/Classes/K2Node_Select.h b/Engine/Source/Editor/BlueprintGraph/Classes/K2Node_Select.h
--- a/Engine/Source/Editor/BlueprintGraph/Classes/K2Node_Select.h
+++ b/Engine/Source/Editor/BlueprintGraph/Classes/K2Node_Select.h
@@ -43,7 +43,7 @@ class UK2Node_Select : public UK2Node, public INodeDependingOnEnumInterface
 	virtual void PinTypeChanged(UEdGraphPin* Pin) override;
 	virtual void AutowireNewNode(UEdGraphPin* FromPin) override;
 	virtual void PostPasteNode() override;
-	virtual FName GetPaletteIcon(FLinearColor& OutColor) const override{ return TEXT("GraphEditor.Select_16x"); }
+	virtual FSlateIcon GetIconAndTint(FLinearColor& OutColor) const override;
 	//~ End UEdGraphNode Interface
 
 	//~ Begin UK2Node Interface
diff --git a/Engine/Source/Editor/BlueprintGraph/Classes/K2Node_SpawnActor.h b/Engine/Source/Editor/BlueprintGraph/Classes/K2Node_SpawnActor.h
--- a/Engine/Source/Editor/BlueprintGraph/Classes/K2Node_SpawnActor.h
+++ b/Engine/Source/Editor/BlueprintGraph/Classes/K2Node_SpawnActor.h
@@ -20,7 +20,7 @@ class UK2Node_SpawnActor : public UK2Node
 	virtual bool IsDeprecated() const override;
 	virtual bool ShouldWarnOnDeprecation() const override;
 	virtual FString GetDeprecationMessage() const override;
-	virtual FName GetPaletteIcon(FLinearColor& OutColor) const override{ return TEXT("GraphEditor.SpawnActor_16x"); }
+	virtual FSlateIcon GetIconAndTint(FLinearColor& OutColor) const override;
 	virtual bool IsCompatibleWithGraph(const UEdGraph* TargetGraph) const override;
 	//~ End UEdGraphNode Interface.
 
diff --git a/Engine/Source/Editor/BlueprintGraph/Classes/K2Node_SpawnActorFromClass.h b/Engine/Source/Editor/BlueprintGraph/Classes/K2Node_SpawnActorFromClass.h
--- a/Engine/Source/Editor/BlueprintGraph/Classes/K2Node_SpawnActorFromClass.h
+++ b/Engine/Source/Editor/BlueprintGraph/Classes/K2Node_SpawnActorFromClass.h
@@ -21,7 +21,7 @@ class BLUEPRINTGRAPH_API UK2Node_SpawnActorFromClass : public UK2Node
 	virtual FText GetTooltipText() const override;
 	virtual void ExpandNode(class FKismetCompilerContext& CompilerContext, UEdGraph* SourceGraph) override;
 	virtual bool HasExternalDependencies(TArray<class UStruct*>* OptionalOutput) const override;
-	virtual FName GetPaletteIcon(FLinearColor& OutColor) const override{ return TEXT("GraphEditor.SpawnActor_16x"); }
+	virtual FSlateIcon GetIconAndTint(FLinearColor& OutColor) const override;
 	virtual bool IsCompatibleWithGraph(const UEdGraph* TargetGraph) const override;
 	virtual void PostPlacedNewNode() override;
 	//~ End UEdGraphNode Interface.
diff --git a/Engine/Source/Editor/BlueprintGraph/Classes/K2Node_Switch.h b/Engine/Source/Editor/BlueprintGraph/Classes/K2Node_Switch.h
--- a/Engine/Source/Editor/BlueprintGraph/Classes/K2Node_Switch.h
+++ b/Engine/Source/Editor/BlueprintGraph/Classes/K2Node_Switch.h
@@ -29,7 +29,7 @@ class UK2Node_Switch : public UK2Node
 	// UEdGraphNode interface
 	BLUEPRINTGRAPH_API virtual void AllocateDefaultPins() override;
 	BLUEPRINTGRAPH_API virtual FLinearColor GetNodeTitleColor() const override;
-	virtual FName GetPaletteIcon(FLinearColor& OutColor) const override{ return TEXT("GraphEditor.Switch_16x"); }
+	BLUEPRINTGRAPH_API virtual FSlateIcon GetIconAndTint(FLinearColor& OutColor) const override;
 	// End of UEdGraphNode interface
 
 	// UK2Node interface
diff --git a/Engine/Source/Editor/BlueprintGraph/Classes/K2Node_Timeline.h b/Engine/Source/Editor/BlueprintGraph/Classes/K2Node_Timeline.h
--- a/Engine/Source/Editor/BlueprintGraph/Classes/K2Node_Timeline.h
+++ b/Engine/Source/Editor/BlueprintGraph/Classes/K2Node_Timeline.h
@@ -43,7 +43,7 @@ class UK2Node_Timeline : public UK2Node
 	virtual TSharedPtr<class INameValidatorInterface> MakeNameValidator() const override;
 	virtual FText GetTooltipText() const override;
 	virtual FString GetDocumentationExcerptName() const override;
-	virtual FName GetPaletteIcon(FLinearColor& OutColor) const override{ return TEXT("GraphEditor.Timeline_16x"); }
+	virtual FSlateIcon GetIconAndTint(FLinearColor& OutColor) const override;
 	virtual bool ShouldShowNodeProperties() const override { return true; }
 	//~ End UEdGraphNode Interface.
 
diff --git a/Engine/Source/Editor/BlueprintGraph/Classes/K2Node_Variable.h b/Engine/Source/Editor/BlueprintGraph/Classes/K2Node_Variable.h
--- a/Engine/Source/Editor/BlueprintGraph/Classes/K2Node_Variable.h
+++ b/Engine/Source/Editor/BlueprintGraph/Classes/K2Node_Variable.h
@@ -62,7 +62,7 @@ public:
 	virtual FString GetDocumentationLink() const override;
 	virtual FString GetDocumentationExcerptName() const override;
 	virtual FName GetCornerIcon() const override;
-	virtual FName GetPaletteIcon(FLinearColor& OutColor) const override;
+	virtual FSlateIcon GetIconAndTint(FLinearColor& OutColor) const override;
 	virtual void AutowireNewNode(UEdGraphPin* FromPin) override;
 	virtual bool CanPasteHere(const UEdGraph* TargetGraph) const override;
 	virtual void PostPasteNode() override;
@@ -145,7 +145,7 @@ public:
 	 * @param  IconColorOut	A color out, further discerning the variable's type.
 	 * @return A icon representing the specified variable's type.
 	 */
-	static FName GetVariableIconAndColor(const UStruct* VarScope, FName VarName, FLinearColor& IconColorOut);
+	static FSlateIcon GetVariableIconAndColor(const UStruct* VarScope, FName VarName, FLinearColor& IconColorOut);
 
 	/**
 	 * Utility method intended to serve as a choke point for various slate 
@@ -155,7 +155,7 @@ public:
 	 * @param  IconColorOut	A color out, further discerning the variable's type.
 	 * @return A icon representing the specified variable's type.
 	 */
-	static FName GetVarIconFromPinType(const FEdGraphPinType& InPinType, FLinearColor& IconColorOut);
+	static FSlateIcon GetVarIconFromPinType(const FEdGraphPinType& InPinType, FLinearColor& IconColorOut);
 
 protected:
 	/**
diff --git a/Engine/Source/Editor/BlueprintGraph/Private/BlueprintBoundEventNodeSpawner.cpp b/Engine/Source/Editor/BlueprintGraph/Private/BlueprintBoundEventNodeSpawner.cpp
--- a/Engine/Source/Editor/BlueprintGraph/Private/BlueprintBoundEventNodeSpawner.cpp
+++ b/Engine/Source/Editor/BlueprintGraph/Private/BlueprintBoundEventNodeSpawner.cpp
@@ -63,7 +63,7 @@ UBlueprintBoundEventNodeSpawner* UBlueprintBoundEventNodeSpawner::Create(TSubcla
 	MenuSignature.Category = BlueprintBoundEventNodeSpawnerImpl::GetDefaultMenuCategory(EventDelegate);
 	//MenuSignature.Tooltip,  will be pulled from the node template
 	//MenuSignature.Keywords, will be pulled from the node template
-	MenuSignature.IconName = TEXT("GraphEditor.Event_16x");
+	MenuSignature.Icon = FSlateIcon("EditorStyle", "GraphEditor.Event_16x");
 
 	return NodeSpawner;
 }
diff --git a/Engine/Source/Editor/BlueprintGraph/Private/BlueprintComponentNodeSpawner.cpp b/Engine/Source/Editor/BlueprintGraph/Private/BlueprintComponentNodeSpawner.cpp
--- a/Engine/Source/Editor/BlueprintGraph/Private/BlueprintComponentNodeSpawner.cpp
+++ b/Engine/Source/Editor/BlueprintGraph/Private/BlueprintComponentNodeSpawner.cpp
@@ -3,7 +3,7 @@
 #include "BlueprintGraphPrivatePCH.h"
 #include "BlueprintComponentNodeSpawner.h"
 #include "K2Node_AddComponent.h"
-#include "ClassIconFinder.h" // for FindIconNameForClass()
+#include "SlateIconFinder.h" // for FindIconForClass()
 #include "BlueprintNodeTemplateCache.h" // for IsTemplateOuter()
 #include "ComponentAssetBroker.h" // for GetComponentsForAsset()/AssignAssetToComponent()
 #include "BlueprintActionFilter.h"	// for FBlueprintActionContext
@@ -69,7 +69,7 @@ UBlueprintComponentNodeSpawner* UBlueprintComponentNodeSpawner::Create(const FCo
 		{
 			MenuSignature.Keywords = FText::FromString(TEXT(" "));
 		}
-		MenuSignature.IconName = FClassIconFinder::FindIconNameForClass(nullptr);
+		MenuSignature.Icon = FSlateIconFinder::FindIconForClass(nullptr);
 		MenuSignature.DocLink  = TEXT("Shared/GraphNodes/Blueprint/UK2Node_AddComponent");
 		MenuSignature.DocExcerptTag = TEXT("AddComponent");
 
@@ -100,7 +100,7 @@ UBlueprintComponentNodeSpawner* UBlueprintComponentNodeSpawner::Create(const FCo
 	{
 		MenuSignature.Keywords = FText::FromString(TEXT(" "));
 	}
-	MenuSignature.IconName = FClassIconFinder::FindIconNameForClass(AuthoritativeClass);
+	MenuSignature.Icon = FSlateIconFinder::FindIconForClass(AuthoritativeClass);
 	MenuSignature.DocLink  = TEXT("Shared/GraphNodes/Blueprint/UK2Node_AddComponent");
 	MenuSignature.DocExcerptTag = AuthoritativeClass->GetName();
 
diff --git a/Engine/Source/Editor/BlueprintGraph/Private/BlueprintDelegateNodeSpawner.cpp b/Engine/Source/Editor/BlueprintGraph/Private/BlueprintDelegateNodeSpawner.cpp
--- a/Engine/Source/Editor/BlueprintGraph/Private/BlueprintDelegateNodeSpawner.cpp
+++ b/Engine/Source/Editor/BlueprintGraph/Private/BlueprintDelegateNodeSpawner.cpp
@@ -19,7 +19,7 @@ namespace BlueprintDelegateNodeSpawnerImpl
 {
 	static FText GetDefaultMenuName(UMulticastDelegateProperty const* Delegate);
 	static FText GetDefaultMenuCategory(UMulticastDelegateProperty const* Delegate);
-	static FName GetDefaultMenuIcon(UMulticastDelegateProperty const* Delegate, FLinearColor& ColorOut);
+	static FSlateIcon GetDefaultMenuIcon(UMulticastDelegateProperty const* Delegate, FLinearColor& ColorOut);
 }
 
 //------------------------------------------------------------------------------
@@ -41,7 +41,7 @@ static FText BlueprintDelegateNodeSpawnerImpl::GetDefaultMenuCategory(UMulticast
 }
 
 //------------------------------------------------------------------------------
-static FName BlueprintDelegateNodeSpawnerImpl::GetDefaultMenuIcon(UMulticastDelegateProperty const* Delegate, FLinearColor& ColorOut)
+static FSlateIcon BlueprintDelegateNodeSpawnerImpl::GetDefaultMenuIcon(UMulticastDelegateProperty const* Delegate, FLinearColor& ColorOut)
 {
 	FName    const PropertyName = Delegate->GetFName();
 	UStruct* const PropertyOwner = CastChecked<UStruct>(Delegate->GetOuterUField());
@@ -79,7 +79,7 @@ UBlueprintDelegateNodeSpawner* UBlueprintDelegateNodeSpawner::Create(TSubclassOf
 	MenuSignature.Category = BlueprintDelegateNodeSpawnerImpl::GetDefaultMenuCategory(Property);
 	//MenuSignature.Tooltip,  will be pulled from the node template
 	//MenuSignature.Keywords, will be pulled from the node template
-	MenuSignature.IconName = BlueprintDelegateNodeSpawnerImpl::GetDefaultMenuIcon(Property, MenuSignature.IconTint);
+	MenuSignature.Icon = BlueprintDelegateNodeSpawnerImpl::GetDefaultMenuIcon(Property, MenuSignature.IconTint);
 
 	//--------------------------------------
 	// Post-Spawn Setup
diff --git a/Engine/Source/Editor/BlueprintGraph/Private/BlueprintEventNodeSpawner.cpp b/Engine/Source/Editor/BlueprintGraph/Private/BlueprintEventNodeSpawner.cpp
--- a/Engine/Source/Editor/BlueprintGraph/Private/BlueprintEventNodeSpawner.cpp
+++ b/Engine/Source/Editor/BlueprintGraph/Private/BlueprintEventNodeSpawner.cpp
@@ -73,7 +73,7 @@ UBlueprintEventNodeSpawner* UBlueprintEventNodeSpawner::Create(UFunction const*
 	{
 		MenuSignature.Keywords = FText::FromString(TEXT(" "));
 	}
-	MenuSignature.IconName = TEXT("GraphEditor.Event_16x");
+	MenuSignature.Icon = FSlateIcon("EditorStyle", "GraphEditor.Event_16x");
 
 	return NodeSpawner;
 }
@@ -94,13 +94,13 @@ UBlueprintEventNodeSpawner* UBlueprintEventNodeSpawner::Create(TSubclassOf<UK2No
 	if (CustomEventName.IsNone())
 	{
 		MenuSignature.MenuName = LOCTEXT("AddCustomEvent", "Add Custom Event...");
-		MenuSignature.IconName = TEXT("GraphEditor.CustomEvent_16x");
+		MenuSignature.Icon = FSlateIcon("EditorStyle", "GraphEditor.CustomEvent_16x");
 	}
 	else
 	{
 		FText const EventName = FText::FromName(CustomEventName);
 		MenuSignature.MenuName = FText::Format(LOCTEXT("EventWithSignatureName", "Event {0}"), EventName);
-		MenuSignature.IconName = TEXT("GraphEditor.Event_16x");
+		MenuSignature.Icon = FSlateIcon("EditorStyle", "GraphEditor.Event_16x");
 	}
 	//MenuSignature.Category, will be pulled from the node template
 	//MenuSignature.Tooltip,  will be pulled from the node template 
diff --git a/Engine/Source/Editor/BlueprintGraph/Private/BlueprintFunctionNodeSpawner.cpp b/Engine/Source/Editor/BlueprintGraph/Private/BlueprintFunctionNodeSpawner.cpp
--- a/Engine/Source/Editor/BlueprintGraph/Private/BlueprintFunctionNodeSpawner.cpp
+++ b/Engine/Source/Editor/BlueprintGraph/Private/BlueprintFunctionNodeSpawner.cpp
@@ -252,7 +252,7 @@ UBlueprintFunctionNodeSpawner* UBlueprintFunctionNodeSpawner::Create(TSubclassOf
 		MenuSignature.Keywords = FText::FromString(TEXT(" "));
 	}
 
-	MenuSignature.IconName = UK2Node_CallFunction::GetPaletteIconForFunction(Function, MenuSignature.IconTint);
+	MenuSignature.Icon = UK2Node_CallFunction::GetPaletteIconForFunction(Function, MenuSignature.IconTint);
 
 	if (MenuSignature.Category.IsEmpty())
 	{
diff --git a/Engine/Source/Editor/BlueprintGraph/Private/BlueprintNodeSpawner.cpp b/Engine/Source/Editor/BlueprintGraph/Private/BlueprintNodeSpawner.cpp
--- a/Engine/Source/Editor/BlueprintGraph/Private/BlueprintNodeSpawner.cpp
+++ b/Engine/Source/Editor/BlueprintGraph/Private/BlueprintNodeSpawner.cpp
@@ -191,19 +191,19 @@ FBlueprintActionUiSpec const& UBlueprintNodeSpawner::PrimeDefaultUiSpec(UEdGraph
 	// Verify Icon Brush Name
 	//--------------------------------------
 
-	if (MenuSignature.IconName.IsNone())
+	if (!MenuSignature.Icon.IsSet())
 	{
 		NodeTemplate = bTemplateNodeFetched ? NodeTemplate : GetTemplateNode(TargetGraph);
 		if (NodeTemplate != nullptr)
 		{
-			MenuSignature.IconName = NodeTemplate->GetPaletteIcon(MenuSignature.IconTint);
+			MenuSignature.Icon = NodeTemplate->GetIconAndTint(MenuSignature.IconTint);
 		}
 		// if a target graph was provided, then we've done all we can to spawn a
 		// template node... we have to default to something
-		if (MenuSignature.IconName.IsNone() && (TargetGraph != nullptr))
+		if (!MenuSignature.Icon.IsSet() && (TargetGraph != nullptr))
 		{
 			// want to set it to something so we won't end up back in this condition
-			MenuSignature.IconName = TEXT("GraphEditor.Default_16x");
+			MenuSignature.Icon = FSlateIcon("EditorStyle", "GraphEditor.Default_16x");
 		}
 		bTemplateNodeFetched = true;
 	}
diff --git a/Engine/Source/Editor/BlueprintGraph/Private/BlueprintVariableNodeSpawner.cpp b/Engine/Source/Editor/BlueprintGraph/Private/BlueprintVariableNodeSpawner.cpp
--- a/Engine/Source/Editor/BlueprintGraph/Private/BlueprintVariableNodeSpawner.cpp
+++ b/Engine/Source/Editor/BlueprintGraph/Private/BlueprintVariableNodeSpawner.cpp
@@ -65,7 +65,7 @@ UBlueprintVariableNodeSpawner* UBlueprintVariableNodeSpawner::Create(TSubclassOf
 		// want to set it to something so we won't end up back in this condition
 		MenuSignature.Keywords = FText::FromString(TEXT(" "));
 	}
-	MenuSignature.IconName = UK2Node_Variable::GetVarIconFromPinType(NodeSpawner->GetVarType(), MenuSignature.IconTint);
+	MenuSignature.Icon = UK2Node_Variable::GetVarIconFromPinType(NodeSpawner->GetVarType(), MenuSignature.IconTint);
 
 	//--------------------------------------
 	// Post-Spawn Setup
@@ -137,7 +137,7 @@ UBlueprintVariableNodeSpawner* UBlueprintVariableNodeSpawner::Create(TSubclassOf
 		// want to set it to something so we won't end up back in this condition
 		MenuSignature.Keywords = FText::FromString(TEXT(" "));
 	}
-	MenuSignature.IconName = UK2Node_Variable::GetVarIconFromPinType(NodeSpawner->GetVarType(), MenuSignature.IconTint);
+	MenuSignature.Icon = UK2Node_Variable::GetVarIconFromPinType(NodeSpawner->GetVarType(), MenuSignature.IconTint);
 
 	return NodeSpawner;
 }
diff --git a/Engine/Source/Editor/BlueprintGraph/Private/EdGraphSchema_K2.cpp b/Engine/Source/Editor/BlueprintGraph/Private/EdGraphSchema_K2.cpp
--- a/Engine/Source/Editor/BlueprintGraph/Private/EdGraphSchema_K2.cpp
+++ b/Engine/Source/Editor/BlueprintGraph/Private/EdGraphSchema_K2.cpp
@@ -135,8 +135,11 @@ struct FUnloadedAssetData
 		, AssetFriendlyName(FText::FromString(FName::NameToDisplayString(InAsset.AssetName.ToString(), false)))
 		, PossibleObjectReferenceTypes(InPossibleObjectReferenceTypes)
 	{
-		const FString* TooltipPtr = InAsset.TagsAndValues.Find(TEXT("Tooltip"));
-		Tooltip = FText::FromString((TooltipPtr && !TooltipPtr->IsEmpty()) ? *TooltipPtr : InAsset.ObjectPath.ToString());
+		InAsset.GetTagValue("Tooltip", Tooltip);
+		if (Tooltip.IsEmpty())
+		{
+			Tooltip = FText::FromString(InAsset.ObjectPath.ToString());
+		}
 	}
 };
 
@@ -310,19 +313,13 @@ public:
 
 					if (Asset.IsValid() && !Asset.IsAssetLoaded())
 					{
-						const FString* BlueprintTypeStr = Asset.TagsAndValues.Find("BlueprintType");
-						const bool bNormalBP = BlueprintTypeStr && (*BlueprintTypeStr == BPNormalTypeAllowed);
-						const bool bInterfaceBP = BlueprintTypeStr && (*BlueprintTypeStr == BPInterfaceTypeAllowed);
+						const FString BlueprintTypeStr = Asset.GetTagValueRef<FString>("BlueprintType");
+						const bool bNormalBP = BlueprintTypeStr == BPNormalTypeAllowed;
+						const bool bInterfaceBP = BlueprintTypeStr == BPInterfaceTypeAllowed;
 
 						if (bNormalBP || bInterfaceBP)
 						{
-							uint32 ClassFlags = 0;
-							const FString* ClassFlagsStr = Asset.TagsAndValues.Find("ClassFlags");
-							if (ClassFlagsStr)
-							{
-								ClassFlags = FCString::Atoi(**ClassFlagsStr);
-							}
-
+							const uint32 ClassFlags = Asset.GetTagValueRef<uint32>("ClassFlags");
 							if (!(ClassFlags & CLASS_Deprecated))
 							{
 								if (bNormalBP)
diff --git a/Engine/Source/Editor/BlueprintGraph/Private/K2Node_BreakStruct.cpp b/Engine/Source/Editor/BlueprintGraph/Private/K2Node_BreakStruct.cpp
--- a/Engine/Source/Editor/BlueprintGraph/Private/K2Node_BreakStruct.cpp
+++ b/Engine/Source/Editor/BlueprintGraph/Private/K2Node_BreakStruct.cpp
@@ -287,6 +287,12 @@ void UK2Node_BreakStruct::ValidateNodeDuringCompilation(class FCompilerResultsLo
 	}
 }
 
+FSlateIcon UK2Node_BreakStruct::GetIconAndTint(FLinearColor& OutColor) const
+{
+	static FSlateIcon Icon("EditorStyle", "GraphEditor.BreakStruct_16x");
+	return Icon;
+}
+
 FLinearColor UK2Node_BreakStruct::GetNodeTitleColor() const 
 {
 	if(const UEdGraphSchema_K2* K2Schema = GetDefault<UEdGraphSchema_K2>())
diff --git a/Engine/Source/Editor/BlueprintGraph/Private/K2Node_CallFunction.cpp b/Engine/Source/Editor/BlueprintGraph/Private/K2Node_CallFunction.cpp
--- a/Engine/Source/Editor/BlueprintGraph/Private/K2Node_CallFunction.cpp
+++ b/Engine/Source/Editor/BlueprintGraph/Private/K2Node_CallFunction.cpp
@@ -1157,28 +1157,33 @@ static FLinearColor GetPalletteIconColor(UFunction const* Function)
 	return GetDefault<UGraphEditorSettings>()->FunctionCallNodeTitleColor;
 }
 
-FName UK2Node_CallFunction::GetPaletteIconForFunction(UFunction const* Function, FLinearColor& OutColor)
+FSlateIcon UK2Node_CallFunction::GetPaletteIconForFunction(UFunction const* Function, FLinearColor& OutColor)
 {
 	static const FName NativeMakeFunc(TEXT("NativeMakeFunc"));
 	static const FName NativeBrakeFunc(TEXT("NativeBreakFunc"));
 
 	if (Function && Function->HasMetaData(NativeMakeFunc))
 	{
-		return TEXT("GraphEditor.MakeStruct_16x");
+		static FSlateIcon Icon("EditorStyle", "GraphEditor.MakeStruct_16x");
+		return Icon;
 	}
 	else if (Function && Function->HasMetaData(NativeBrakeFunc))
 	{
-		return TEXT("GraphEditor.BreakStruct_16x");
+		static FSlateIcon Icon("EditorStyle", "GraphEditor.BreakStruct_16x");
+		return Icon;
 	}
 	// Check to see if the function is calling an function that could be an event, display the event icon instead.
 	else if (Function && UEdGraphSchema_K2::FunctionCanBePlacedAsEvent(Function))
 	{
-		return TEXT("GraphEditor.Event_16x");
+		static FSlateIcon Icon("EditorStyle", "GraphEditor.Event_16x");
+		return Icon;
 	}
 	else
 	{
 		OutColor = GetPalletteIconColor(Function);
-		return TEXT("Kismet.AllClasses.FunctionIcon");
+
+		static FSlateIcon Icon("EditorStyle", "Kismet.AllClasses.FunctionIcon");
+		return Icon;
 	}
 }
 
@@ -2217,7 +2222,7 @@ FName UK2Node_CallFunction::GetCornerIcon() const
 	return Super::GetCornerIcon();
 }
 
-FName UK2Node_CallFunction::GetPaletteIcon(FLinearColor& OutColor) const
+FSlateIcon UK2Node_CallFunction::GetIconAndTint(FLinearColor& OutColor) const
 {
 	return GetPaletteIconForFunction(GetTargetFunction(), OutColor);
 }
diff --git a/Engine/Source/Editor/BlueprintGraph/Private/K2Node_CastByteToEnum.cpp b/Engine/Source/Editor/BlueprintGraph/Private/K2Node_CastByteToEnum.cpp
--- a/Engine/Source/Editor/BlueprintGraph/Private/K2Node_CastByteToEnum.cpp
+++ b/Engine/Source/Editor/BlueprintGraph/Private/K2Node_CastByteToEnum.cpp
@@ -55,6 +55,12 @@ FText UK2Node_CastByteToEnum::GetNodeTitle(ENodeTitleType::Type TitleType) const
 	return GetTooltipText();
 }
 
+FSlateIcon UK2Node_CastByteToEnum::GetIconAndTint(FLinearColor& OutColor) const
+{
+	static FSlateIcon Icon("EditorStyle", "GraphEditor.Enum_16x");
+	return Icon;
+}
+
 FText UK2Node_CastByteToEnum::GetCompactNodeTitle() const
 {
 	return NSLOCTEXT("K2Node", "CastSymbol", "\x2022");
diff --git a/Engine/Source/Editor/BlueprintGraph/Private/K2Node_CustomEvent.cpp b/Engine/Source/Editor/BlueprintGraph/Private/K2Node_CustomEvent.cpp
--- a/Engine/Source/Editor/BlueprintGraph/Private/K2Node_CustomEvent.cpp
+++ b/Engine/Source/Editor/BlueprintGraph/Private/K2Node_CustomEvent.cpp
@@ -471,9 +471,9 @@ FString UK2Node_CustomEvent::GetDocumentationExcerptName() const
 	return TEXT("UK2Node_CustomEvent");
 }
 
-FName UK2Node_CustomEvent::GetPaletteIcon(FLinearColor& OutColor) const
+FSlateIcon UK2Node_CustomEvent::GetIconAndTint(FLinearColor& OutColor) const
 {
-	return bCallInEditor ? TEXT("GraphEditor.CallInEditorEvent_16x") : TEXT("GraphEditor.CustomEvent_16x");
+	return FSlateIcon("EditorStyle", bCallInEditor ? "GraphEditor.CallInEditorEvent_16x" : "GraphEditor.CustomEvent_16x");
 }
 
 void UK2Node_CustomEvent::AutowireNewNode(UEdGraphPin* FromPin)
diff --git a/Engine/Source/Editor/BlueprintGraph/Private/K2Node_DynamicCast.cpp b/Engine/Source/Editor/BlueprintGraph/Private/K2Node_DynamicCast.cpp
--- a/Engine/Source/Editor/BlueprintGraph/Private/K2Node_DynamicCast.cpp
+++ b/Engine/Source/Editor/BlueprintGraph/Private/K2Node_DynamicCast.cpp
@@ -76,6 +76,12 @@ FLinearColor UK2Node_DynamicCast::GetNodeTitleColor() const
 	return FLinearColor(0.0f, 0.55f, 0.62f);
 }
 
+FSlateIcon UK2Node_DynamicCast::GetIconAndTint(FLinearColor& OutColor) const
+{
+	static FSlateIcon Icon("EditorStyle", "GraphEditor.Cast_16x");
+	return Icon;
+}
+
 FText UK2Node_DynamicCast::GetNodeTitle(ENodeTitleType::Type TitleType) const
 {
 	if (TargetType == nullptr)
diff --git a/Engine/Source/Editor/BlueprintGraph/Private/K2Node_EaseFunction.cpp b/Engine/Source/Editor/BlueprintGraph/Private/K2Node_EaseFunction.cpp
--- a/Engine/Source/Editor/BlueprintGraph/Private/K2Node_EaseFunction.cpp
+++ b/Engine/Source/Editor/BlueprintGraph/Private/K2Node_EaseFunction.cpp
@@ -129,11 +129,12 @@ FText UK2Node_EaseFunction::GetTooltipText() const
 	return NodeTooltip;
 }
 
-FName UK2Node_EaseFunction::GetPaletteIcon(FLinearColor& OutColor) const
+FSlateIcon UK2Node_EaseFunction::GetIconAndTint(FLinearColor& OutColor) const
 {
 	// The function icon seams the best choice!
 	OutColor = GetNodeTitleColor();
-	return TEXT("Kismet.AllClasses.FunctionIcon");
+	static FSlateIcon Icon("EditorStyle", "Kismet.AllClasses.FunctionIcon");
+	return Icon;
 }
 
 void UK2Node_EaseFunction::SetPinToolTip(UEdGraphPin& MutatablePin, const FText& PinDescription) const
diff --git a/Engine/Source/Editor/BlueprintGraph/Private/K2Node_EnumLiteral.cpp b/Engine/Source/Editor/BlueprintGraph/Private/K2Node_EnumLiteral.cpp
--- a/Engine/Source/Editor/BlueprintGraph/Private/K2Node_EnumLiteral.cpp
+++ b/Engine/Source/Editor/BlueprintGraph/Private/K2Node_EnumLiteral.cpp
@@ -40,10 +40,10 @@ void UK2Node_EnumLiteral::AllocateDefaultPins()
 	Super::AllocateDefaultPins();
 }
 
-FName UK2Node_EnumLiteral::GetPaletteIcon(FLinearColor& OutColor) const
-{ 
-	static const FName PaletteIconName(TEXT("GraphEditor.Enum_16x"));
-	return PaletteIconName;
+FSlateIcon UK2Node_EnumLiteral::GetIconAndTint(FLinearColor& OutColor) const
+{
+	static FSlateIcon Icon("EditorStyle", "GraphEditor.Enum_16x");
+	return Icon;
 }
 
 FText UK2Node_EnumLiteral::GetTooltipText() const
diff --git a/Engine/Source/Editor/BlueprintGraph/Private/K2Node_Event.cpp b/Engine/Source/Editor/BlueprintGraph/Private/K2Node_Event.cpp
--- a/Engine/Source/Editor/BlueprintGraph/Private/K2Node_Event.cpp
+++ b/Engine/Source/Editor/BlueprintGraph/Private/K2Node_Event.cpp
@@ -810,6 +810,12 @@ UObject* UK2Node_Event::GetJumpTargetForDoubleClick() const
 	return NULL;
 }
 
+FSlateIcon UK2Node_Event::GetIconAndTint(FLinearColor& OutColor) const
+{
+	static FSlateIcon Icon("EditorStyle", "GraphEditor.Event_16x");
+	return Icon;
+}
+
 FString UK2Node_Event::GetFindReferenceSearchString() const
 {
 	FString FunctionName = EventReference.GetMemberName().ToString(); // If we fail to find the function, still want to search for its expected name.
diff --git a/Engine/Source/Editor/BlueprintGraph/Private/K2Node_ExecutionSequence.cpp b/Engine/Source/Editor/BlueprintGraph/Private/K2Node_ExecutionSequence.cpp
--- a/Engine/Source/Editor/BlueprintGraph/Private/K2Node_ExecutionSequence.cpp
+++ b/Engine/Source/Editor/BlueprintGraph/Private/K2Node_ExecutionSequence.cpp
@@ -174,6 +174,12 @@ FText UK2Node_ExecutionSequence::GetNodeTitle(ENodeTitleType::Type TitleType) co
 	return NSLOCTEXT("K2Node", "Sequence", "Sequence");
 }
 
+FSlateIcon UK2Node_ExecutionSequence::GetIconAndTint(FLinearColor& OutColor) const
+{
+	static FSlateIcon Icon("EditorStyle", "GraphEditor.Sequence_16x");
+	return Icon;
+}
+
 FLinearColor UK2Node_ExecutionSequence::GetNodeTitleColor() const
 {
 	return FLinearColor::White;
diff --git a/Engine/Source/Editor/BlueprintGraph/Private/K2Node_ForEachElementInEnum.cpp b/Engine/Source/Editor/BlueprintGraph/Private/K2Node_ForEachElementInEnum.cpp
--- a/Engine/Source/Editor/BlueprintGraph/Private/K2Node_ForEachElementInEnum.cpp
+++ b/Engine/Source/Editor/BlueprintGraph/Private/K2Node_ForEachElementInEnum.cpp
@@ -182,6 +182,12 @@ FText UK2Node_ForEachElementInEnum::GetNodeTitle(ENodeTitleType::Type TitleType)
 	return CachedNodeTitle;
 }
 
+FSlateIcon UK2Node_ForEachElementInEnum::GetIconAndTint(FLinearColor& OutColor) const
+{
+	static FSlateIcon Icon("EditorStyle", "GraphEditor.Macro.Loop_16x");
+	return Icon;
+}
+
 void UK2Node_ForEachElementInEnum::ExpandNode(class FKismetCompilerContext& CompilerContext, UEdGraph* SourceGraph)
 {
 	Super::ExpandNode(CompilerContext, SourceGraph);
diff --git a/Engine/Source/Editor/BlueprintGraph/Private/K2Node_GetDataTableRow.cpp b/Engine/Source/Editor/BlueprintGraph/Private/K2Node_GetDataTableRow.cpp
--- a/Engine/Source/Editor/BlueprintGraph/Private/K2Node_GetDataTableRow.cpp
+++ b/Engine/Source/Editor/BlueprintGraph/Private/K2Node_GetDataTableRow.cpp
@@ -374,6 +374,13 @@ void UK2Node_GetDataTableRow::ExpandNode(class FKismetCompilerContext& CompilerC
 	BreakAllNodeLinks();
 }
 
+FSlateIcon UK2Node_GetDataTableRow::GetIconAndTint(FLinearColor& OutColor) const
+{
+	OutColor = GetNodeTitleColor();
+	static FSlateIcon Icon("EditorStyle", "Kismet.AllClasses.FunctionIcon");
+	return Icon;
+}
+
 void UK2Node_GetDataTableRow::EarlyValidation(class FCompilerResultsLog& MessageLog) const
 {
 	const auto DataTablePin = GetDataTablePin();
diff --git a/Engine/Source/Editor/BlueprintGraph/Private/K2Node_GetEnumeratorName.cpp b/Engine/Source/Editor/BlueprintGraph/Private/K2Node_GetEnumeratorName.cpp
--- a/Engine/Source/Editor/BlueprintGraph/Private/K2Node_GetEnumeratorName.cpp
+++ b/Engine/Source/Editor/BlueprintGraph/Private/K2Node_GetEnumeratorName.cpp
@@ -58,6 +58,12 @@ void UK2Node_GetEnumeratorName::ValidateNodeDuringCompilation(class FCompilerRes
 	}
 }
 
+FSlateIcon UK2Node_GetEnumeratorName::GetIconAndTint(FLinearColor& OutColor) const
+{
+	static FSlateIcon Icon("EditorStyle", "GraphEditor.Enum_16x");
+	return Icon;
+}
+
 void UK2Node_GetEnumeratorName::EarlyValidation(class FCompilerResultsLog& MessageLog) const
 {
 	const UEnum* Enum = GetEnum();
diff --git a/Engine/Source/Editor/BlueprintGraph/Private/K2Node_GetEnumeratorNameAsString.cpp b/Engine/Source/Editor/BlueprintGraph/Private/K2Node_GetEnumeratorNameAsString.cpp
--- a/Engine/Source/Editor/BlueprintGraph/Private/K2Node_GetEnumeratorNameAsString.cpp
+++ b/Engine/Source/Editor/BlueprintGraph/Private/K2Node_GetEnumeratorNameAsString.cpp
@@ -32,6 +32,12 @@ FText UK2Node_GetEnumeratorNameAsString::GetNodeTitle(ENodeTitleType::Type Title
 	return NSLOCTEXT("K2Node", "GetEnumeratorNameAsString_Title", "Enum to String");
 }
 
+FSlateIcon UK2Node_GetEnumeratorNameAsString::GetIconAndTint(FLinearColor& OutColor) const
+{
+	static FSlateIcon Icon("EditorStyle", "GraphEditor.Enum_16x");
+	return Icon;
+}
+
 void UK2Node_GetEnumeratorNameAsString::GetMenuActions(FBlueprintActionDatabaseRegistrar& ActionRegistrar) const
 {
 	// actions get registered under specific object-keys; the idea is that 
diff --git a/Engine/Source/Editor/BlueprintGraph/Private/K2Node_GetInputAxisKeyValue.cpp b/Engine/Source/Editor/BlueprintGraph/Private/K2Node_GetInputAxisKeyValue.cpp
--- a/Engine/Source/Editor/BlueprintGraph/Private/K2Node_GetInputAxisKeyValue.cpp
+++ b/Engine/Source/Editor/BlueprintGraph/Private/K2Node_GetInputAxisKeyValue.cpp
@@ -109,9 +109,9 @@ void UK2Node_GetInputAxisKeyValue::RegisterDynamicBinding(UDynamicBlueprintBindi
 	InputAxisKeyBindingObject->InputAxisKeyDelegateBindings.Add(Binding);
 }
 
-FName UK2Node_GetInputAxisKeyValue::GetPaletteIcon(FLinearColor& OutColor) const
+FSlateIcon UK2Node_GetInputAxisKeyValue::GetIconAndTint(FLinearColor& OutColor) const
 {
-	return EKeys::GetMenuCategoryPaletteIcon(InputAxisKey.GetMenuCategory());
+	return FSlateIcon("EditorStyle", EKeys::GetMenuCategoryPaletteIcon(InputAxisKey.GetMenuCategory()));
 }
 
 void UK2Node_GetInputAxisKeyValue::GetMenuActions(FBlueprintActionDatabaseRegistrar& ActionRegistrar) const
diff --git a/Engine/Source/Editor/BlueprintGraph/Private/K2Node_GetNumEnumEntries.cpp b/Engine/Source/Editor/BlueprintGraph/Private/K2Node_GetNumEnumEntries.cpp
--- a/Engine/Source/Editor/BlueprintGraph/Private/K2Node_GetNumEnumEntries.cpp
+++ b/Engine/Source/Editor/BlueprintGraph/Private/K2Node_GetNumEnumEntries.cpp
@@ -53,6 +53,12 @@ FText UK2Node_GetNumEnumEntries::GetNodeTitle(ENodeTitleType::Type TitleType) co
 	return CachedNodeTitle;
 }
 
+FSlateIcon UK2Node_GetNumEnumEntries::GetIconAndTint(FLinearColor& OutColor) const
+{
+	static FSlateIcon Icon("EditorStyle", "GraphEditor.Enum_16x");
+	return Icon;
+}
+
 void UK2Node_GetNumEnumEntries::ExpandNode(class FKismetCompilerContext& CompilerContext, UEdGraph* SourceGraph)
 {
 	Super::ExpandNode(CompilerContext, SourceGraph);
diff --git a/Engine/Source/Editor/BlueprintGraph/Private/K2Node_IfThenElse.cpp b/Engine/Source/Editor/BlueprintGraph/Private/K2Node_IfThenElse.cpp
--- a/Engine/Source/Editor/BlueprintGraph/Private/K2Node_IfThenElse.cpp
+++ b/Engine/Source/Editor/BlueprintGraph/Private/K2Node_IfThenElse.cpp
@@ -112,6 +112,12 @@ FLinearColor UK2Node_IfThenElse::GetNodeTitleColor() const
 	return GetDefault<UGraphEditorSettings>()->ExecBranchNodeTitleColor;
 }
 
+FSlateIcon UK2Node_IfThenElse::GetIconAndTint(FLinearColor& OutColor) const
+{
+	static FSlateIcon Icon("EditorStyle", "GraphEditor.Branch_16x");
+	return Icon;
+}
+
 FText UK2Node_IfThenElse::GetTooltipText() const
 {
 	return LOCTEXT("BrancStatement_Tooltip", "Branch Statement\nIf Condition is true, execution goes to True, otherwise it goes to False");
diff --git a/Engine/Source/Editor/BlueprintGraph/Private/K2Node_InputAction.cpp b/Engine/Source/Editor/BlueprintGraph/Private/K2Node_InputAction.cpp
--- a/Engine/Source/Editor/BlueprintGraph/Private/K2Node_InputAction.cpp
+++ b/Engine/Source/Editor/BlueprintGraph/Private/K2Node_InputAction.cpp
@@ -76,6 +76,12 @@ FText UK2Node_InputAction::GetTooltipText() const
 	return CachedTooltip;
 }
 
+FSlateIcon UK2Node_InputAction::GetIconAndTint(FLinearColor& OutColor) const
+{
+	static FSlateIcon Icon("EditorStyle", "GraphEditor.Event_16x");
+	return Icon;
+}
+
 bool UK2Node_InputAction::IsCompatibleWithGraph(UEdGraph const* Graph) const
 {
 	// This node expands into event nodes and must be placed in a Ubergraph
diff --git a/Engine/Source/Editor/BlueprintGraph/Private/K2Node_InputAxisKeyEvent.cpp b/Engine/Source/Editor/BlueprintGraph/Private/K2Node_InputAxisKeyEvent.cpp
--- a/Engine/Source/Editor/BlueprintGraph/Private/K2Node_InputAxisKeyEvent.cpp
+++ b/Engine/Source/Editor/BlueprintGraph/Private/K2Node_InputAxisKeyEvent.cpp
@@ -79,9 +79,9 @@ UClass* UK2Node_InputAxisKeyEvent::GetDynamicBindingClass() const
 	return UInputAxisKeyDelegateBinding::StaticClass();
 }
 
-FName UK2Node_InputAxisKeyEvent::GetPaletteIcon(FLinearColor& OutColor) const
+FSlateIcon UK2Node_InputAxisKeyEvent::GetIconAndTint(FLinearColor& OutColor) const
 {
-	return EKeys::GetMenuCategoryPaletteIcon(AxisKey.GetMenuCategory());
+	return FSlateIcon("EditorStyle", EKeys::GetMenuCategoryPaletteIcon(AxisKey.GetMenuCategory()));
 }
 
 void UK2Node_InputAxisKeyEvent::RegisterDynamicBinding(UDynamicBlueprintBinding* BindingObject) const
diff --git a/Engine/Source/Editor/BlueprintGraph/Private/K2Node_InputKey.cpp b/Engine/Source/Editor/BlueprintGraph/Private/K2Node_InputKey.cpp
--- a/Engine/Source/Editor/BlueprintGraph/Private/K2Node_InputKey.cpp
+++ b/Engine/Source/Editor/BlueprintGraph/Private/K2Node_InputKey.cpp
@@ -192,9 +192,9 @@ FText UK2Node_InputKey::GetTooltipText() const
 	return CachedTooltip;
 }
 
-FName UK2Node_InputKey::GetPaletteIcon(FLinearColor& OutColor) const
+FSlateIcon UK2Node_InputKey::GetIconAndTint(FLinearColor& OutColor) const
 {
-	return EKeys::GetMenuCategoryPaletteIcon(InputKey.GetMenuCategory());
+	return FSlateIcon("EditorStyle", EKeys::GetMenuCategoryPaletteIcon(InputKey.GetMenuCategory()));
 }
 
 bool UK2Node_InputKey::IsCompatibleWithGraph(UEdGraph const* Graph) const
diff --git a/Engine/Source/Editor/BlueprintGraph/Private/K2Node_InputTouch.cpp b/Engine/Source/Editor/BlueprintGraph/Private/K2Node_InputTouch.cpp
--- a/Engine/Source/Editor/BlueprintGraph/Private/K2Node_InputTouch.cpp
+++ b/Engine/Source/Editor/BlueprintGraph/Private/K2Node_InputTouch.cpp
@@ -75,6 +75,12 @@ FText UK2Node_InputTouch::GetTooltipText() const
 	return NSLOCTEXT("K2Node", "InputTouch_Tooltip", "Event for when a finger presses, releases or is moved on a touch device.");
 }
 
+FSlateIcon UK2Node_InputTouch::GetIconAndTint(FLinearColor& OutColor) const
+{
+	static FSlateIcon Icon("EditorStyle", "GraphEditor.TouchEvent_16x");
+	return Icon;
+}
+
 void UK2Node_InputTouch::GetMenuActions(FBlueprintActionDatabaseRegistrar& ActionRegistrar) const
 {
 	// actions get registered under specific object-keys; the idea is that 
diff --git a/Engine/Source/Editor/BlueprintGraph/Private/K2Node_Literal.cpp b/Engine/Source/Editor/BlueprintGraph/Private/K2Node_Literal.cpp
--- a/Engine/Source/Editor/BlueprintGraph/Private/K2Node_Literal.cpp
+++ b/Engine/Source/Editor/BlueprintGraph/Private/K2Node_Literal.cpp
@@ -188,10 +188,10 @@ void UK2Node_Literal::GetMenuActions(FBlueprintActionDatabaseRegistrar& ActionRe
 			UiSpecOut->MenuName = FText::Format( NSLOCTEXT("K2Node", "LiteralTitle", "Create a Reference to {0}"), 
 				FText::FromString(ActorObj->GetActorLabel()) );
 
-			const FName IconName = FClassIconFinder::FindIconNameForActor(ActorObj);
-			if (!IconName.IsNone())
+			FSlateIcon Icon = FClassIconFinder::FindSlateIconForActor(ActorObj);
+			if (Icon.IsSet())
 			{
-				UiSpecOut->IconName = IconName;
+				UiSpecOut->Icon = Icon;
 			}
 		}
 		else if (Bindings.Num() > 1)
@@ -211,10 +211,10 @@ void UK2Node_Literal::GetMenuActions(FBlueprintActionDatabaseRegistrar& ActionRe
 				}
 			}
 
-			const FName IconName = FClassIconFinder::FindIconNameForClass(CommonClass);
-			if (!IconName.IsNone())
+			const FSlateIcon Icon = FSlateIconFinder::FindIconForClass(CommonClass);
+			if (Icon.IsSet())
 			{
-				UiSpecOut->IconName = IconName;
+				UiSpecOut->Icon = Icon;
 			}
 		}
 		else
@@ -279,21 +279,21 @@ FNodeHandlingFunctor* UK2Node_Literal::CreateNodeHandler(FKismetCompilerContext&
 	return new FKCHandler_LiteralStatement(CompilerContext);
 }
 
-FName UK2Node_Literal::GetPaletteIcon(FLinearColor& OutColor) const
+FSlateIcon UK2Node_Literal::GetIconAndTint(FLinearColor& OutColor) const
 {
 	if(ObjectRef != NULL)
 	{
 		AActor* Actor = Cast<AActor>(ObjectRef);
 		if(Actor != NULL)
 		{
-			return FClassIconFinder::FindIconNameForActor(Actor);
+			return FClassIconFinder::FindSlateIconForActor(Actor);
 		}
 		else
 		{
-			return FClassIconFinder::FindIconNameForClass(ObjectRef->GetClass());
+			return FSlateIconFinder::FindIconForClass(ObjectRef->GetClass());
 		}	
 	}
-	return Super::GetPaletteIcon(OutColor);
+	return Super::GetIconAndTint(OutColor);
 }
 
 #undef LOCTEXT_NAMESPACE
diff --git a/Engine/Source/Editor/BlueprintGraph/Private/K2Node_MacroInstance.cpp b/Engine/Source/Editor/BlueprintGraph/Private/K2Node_MacroInstance.cpp
--- a/Engine/Source/Editor/BlueprintGraph/Private/K2Node_MacroInstance.cpp
+++ b/Engine/Source/Editor/BlueprintGraph/Private/K2Node_MacroInstance.cpp
@@ -319,8 +319,10 @@ void UK2Node_MacroInstance::PostReconstructNode()
 	Super::PostReconstructNode();
 }
 
-FName UK2Node_MacroInstance::GetPaletteIcon(FLinearColor& OutColor) const
+FSlateIcon UK2Node_MacroInstance::GetIconAndTint(FLinearColor& OutColor) const
 {
+	const char* IconName = "GraphEditor.Macro_16x";
+
 	// Special case handling for standard macros
 	// @TODO Change this to a SlateBurushAsset pointer on the graph or something similar, to allow any macro to have an icon
 	UEdGraph* MacroGraph = MacroGraphReference.GetGraph();
@@ -331,36 +333,36 @@ FName UK2Node_MacroInstance::GetPaletteIcon(FLinearColor& OutColor) const
 			MacroName == TEXT("ForLoopWithBreak") ||
 			MacroName == TEXT("WhileLoop") )
 		{
-			return TEXT("GraphEditor.Macro.Loop_16x");
+			IconName = "GraphEditor.Macro.Loop_16x";
 		}
 		else if( MacroName == TEXT("Gate") )
 		{
-			return TEXT("GraphEditor.Macro.Gate_16x");
+			IconName = "GraphEditor.Macro.Gate_16x";
 		}
 		else if( MacroName == TEXT("Do N") )
 		{
-			return TEXT("GraphEditor.Macro.DoN_16x");
+			IconName = "GraphEditor.Macro.DoN_16x";
 		}
 		else if (MacroName == TEXT("DoOnce"))
 		{
-			return TEXT("GraphEditor.Macro.DoOnce_16x");
+			IconName = "GraphEditor.Macro.DoOnce_16x";
 		}
 		else if (MacroName == TEXT("IsValid"))
 		{
-			return TEXT("GraphEditor.Macro.IsValid_16x");
+			IconName = "GraphEditor.Macro.IsValid_16x";
 		}
 		else if (MacroName == TEXT("FlipFlop"))
 		{
-			return TEXT("GraphEditor.Macro.FlipFlop_16x");
+			IconName = "GraphEditor.Macro.FlipFlop_16x";
 		}
 		else if ( MacroName == TEXT("ForEachLoop") ||
 				  MacroName == TEXT("ForEachLoopWithBreak") )
 		{
-			return TEXT("GraphEditor.Macro.ForEach_16x");
+			IconName = "GraphEditor.Macro.ForEach_16x";
 		}
 	}
 
-	return TEXT("GraphEditor.Macro_16x");
+	return FSlateIcon("EditorStyle", IconName);
 }
 
 FText UK2Node_MacroInstance::GetCompactNodeTitle() const
diff --git a/Engine/Source/Editor/BlueprintGraph/Private/K2Node_MakeArray.cpp b/Engine/Source/Editor/BlueprintGraph/Private/K2Node_MakeArray.cpp
--- a/Engine/Source/Editor/BlueprintGraph/Private/K2Node_MakeArray.cpp
+++ b/Engine/Source/Editor/BlueprintGraph/Private/K2Node_MakeArray.cpp
@@ -322,6 +322,12 @@ FText UK2Node_MakeArray::GetTooltipText() const
 	return LOCTEXT("MakeArrayTooltip", "Create an array from a series of items.");
 }
 
+FSlateIcon UK2Node_MakeArray::GetIconAndTint(FLinearColor& OutColor) const
+{
+	static FSlateIcon Icon("EditorStyle", "GraphEditor.MakeArray_16x");
+	return Icon;
+}
+
 void UK2Node_MakeArray::AddInputPin()
 {
 	FScopedTransaction Transaction( LOCTEXT("AddPinTx", "AddPin") );
diff --git a/Engine/Source/Editor/BlueprintGraph/Private/K2Node_MakeStruct.cpp b/Engine/Source/Editor/BlueprintGraph/Private/K2Node_MakeStruct.cpp
--- a/Engine/Source/Editor/BlueprintGraph/Private/K2Node_MakeStruct.cpp
+++ b/Engine/Source/Editor/BlueprintGraph/Private/K2Node_MakeStruct.cpp
@@ -225,6 +225,12 @@ FText UK2Node_MakeStruct::GetTooltipText() const
 	return CachedTooltip;
 }
 
+FSlateIcon UK2Node_MakeStruct::GetIconAndTint(FLinearColor& OutColor) const
+{
+	static FSlateIcon Icon("EditorStyle", "GraphEditor.MakeStruct_16x");
+	return Icon;
+}
+
 FLinearColor UK2Node_MakeStruct::GetNodeTitleColor() const
 {
 	if(const UEdGraphSchema_K2* K2Schema = GetDefault<UEdGraphSchema_K2>())
diff --git a/Engine/Source/Editor/BlueprintGraph/Private/K2Node_Select.cpp b/Engine/Source/Editor/BlueprintGraph/Private/K2Node_Select.cpp
--- a/Engine/Source/Editor/BlueprintGraph/Private/K2Node_Select.cpp
+++ b/Engine/Source/Editor/BlueprintGraph/Private/K2Node_Select.cpp
@@ -1016,6 +1016,12 @@ void UK2Node_Select::PostPasteNode()
 	IndexPin->DefaultValue = OldDefaultValue;
 }
 
+FSlateIcon UK2Node_Select::GetIconAndTint(FLinearColor& OutColor) const
+{
+	static FSlateIcon Icon("EditorStyle", "GraphEditor.Select_16x");
+	return Icon;
+}
+
 bool UK2Node_Select::IsConnectionDisallowed(const UEdGraphPin* MyPin, const UEdGraphPin* OtherPin, FString& OutReason) const
 {
 	const UEdGraphSchema_K2* K2Schema = GetDefault<UEdGraphSchema_K2>();
diff --git a/Engine/Source/Editor/BlueprintGraph/Private/K2Node_SetFieldsInStruct.cpp b/Engine/Source/Editor/BlueprintGraph/Private/K2Node_SetFieldsInStruct.cpp
--- a/Engine/Source/Editor/BlueprintGraph/Private/K2Node_SetFieldsInStruct.cpp
+++ b/Engine/Source/Editor/BlueprintGraph/Private/K2Node_SetFieldsInStruct.cpp
@@ -181,9 +181,9 @@ FText UK2Node_SetFieldsInStruct::GetTooltipText() const
 	return CachedTooltip;
 }
 
-FName UK2Node_SetFieldsInStruct::GetPaletteIcon(FLinearColor& OutColor) const
-{ 
-	return UK2Node_Variable::GetPaletteIcon(OutColor);
+FSlateIcon UK2Node_SetFieldsInStruct::GetIconAndTint(FLinearColor& OutColor) const
+{
+	return UK2Node_Variable::GetIconAndTint(OutColor);
 }
 
 void UK2Node_SetFieldsInStruct::ValidateNodeDuringCompilation(FCompilerResultsLog& MessageLog) const
diff --git a/Engine/Source/Editor/BlueprintGraph/Private/K2Node_SpawnActor.cpp b/Engine/Source/Editor/BlueprintGraph/Private/K2Node_SpawnActor.cpp
--- a/Engine/Source/Editor/BlueprintGraph/Private/K2Node_SpawnActor.cpp
+++ b/Engine/Source/Editor/BlueprintGraph/Private/K2Node_SpawnActor.cpp
@@ -492,4 +492,10 @@ FString UK2Node_SpawnActor::GetDeprecationMessage() const
 	return LOCTEXT("SpawnActorNodeOnlyDefaultBlueprint_Deprecatio", "Spawn Actor @@ is DEPRECATED and should be replaced by SpawnActorFromClass").ToString();
 }
 
+FSlateIcon UK2Node_SpawnActor::GetIconAndTint(FLinearColor& OutColor) const
+{
+	static FSlateIcon Icon("EditorStyle", "GraphEditor.SpawnActor_16x");
+	return Icon;
+}
+
 #undef LOCTEXT_NAMESPACE
diff --git a/Engine/Source/Editor/BlueprintGraph/Private/K2Node_SpawnActorFromClass.cpp b/Engine/Source/Editor/BlueprintGraph/Private/K2Node_SpawnActorFromClass.cpp
--- a/Engine/Source/Editor/BlueprintGraph/Private/K2Node_SpawnActorFromClass.cpp
+++ b/Engine/Source/Editor/BlueprintGraph/Private/K2Node_SpawnActorFromClass.cpp
@@ -412,6 +412,12 @@ FText UK2Node_SpawnActorFromClass::GetTooltipText() const
 	return NodeTooltip;
 }
 
+FSlateIcon UK2Node_SpawnActorFromClass::GetIconAndTint(FLinearColor& OutColor) const
+{
+	static FSlateIcon Icon("EditorStyle", "GraphEditor.SpawnActor_16x");
+	return Icon;
+}
+
 UEdGraphPin* UK2Node_SpawnActorFromClass::GetThenPin()const
 {
 	const UEdGraphSchema_K2* K2Schema = GetDefault<UEdGraphSchema_K2>();
diff --git a/Engine/Source/Editor/BlueprintGraph/Private/K2Node_Switch.cpp b/Engine/Source/Editor/BlueprintGraph/Private/K2Node_Switch.cpp
--- a/Engine/Source/Editor/BlueprintGraph/Private/K2Node_Switch.cpp
+++ b/Engine/Source/Editor/BlueprintGraph/Private/K2Node_Switch.cpp
@@ -222,6 +222,12 @@ FLinearColor UK2Node_Switch::GetNodeTitleColor() const
 	return FLinearColor(255.0f, 255.0f, 0.0f);
 }
 
+FSlateIcon UK2Node_Switch::GetIconAndTint(FLinearColor& OutColor) const
+{
+	static FSlateIcon Icon("EditorStyle", "GraphEditor.Switch_16x");
+	return Icon;
+}
+
 void UK2Node_Switch::AddPinToSwitchNode()
 {
 	const UEdGraphSchema_K2* K2Schema = GetDefault<UEdGraphSchema_K2>();
diff --git a/Engine/Source/Editor/BlueprintGraph/Private/K2Node_Timeline.cpp b/Engine/Source/Editor/BlueprintGraph/Private/K2Node_Timeline.cpp
--- a/Engine/Source/Editor/BlueprintGraph/Private/K2Node_Timeline.cpp
+++ b/Engine/Source/Editor/BlueprintGraph/Private/K2Node_Timeline.cpp
@@ -560,6 +560,12 @@ FText UK2Node_Timeline::GetTooltipText() const
 	return LOCTEXT("TimelineTooltip", "Timeline node allows values to be keyframed over time.\nDouble click to open timeline editor.");
 }
 
+FSlateIcon UK2Node_Timeline::GetIconAndTint(FLinearColor& OutColor) const
+{
+	static FSlateIcon Icon("EditorStyle", "GraphEditor.Timeline_16x");
+	return Icon;
+}
+
 FString UK2Node_Timeline::GetDocumentationExcerptName() const
 {
 	return TEXT("UK2Node_Timeline");
diff --git a/Engine/Source/Editor/BlueprintGraph/Private/K2Node_Variable.cpp b/Engine/Source/Editor/BlueprintGraph/Private/K2Node_Variable.cpp
--- a/Engine/Source/Editor/BlueprintGraph/Private/K2Node_Variable.cpp
+++ b/Engine/Source/Editor/BlueprintGraph/Private/K2Node_Variable.cpp
@@ -4,7 +4,7 @@
 #include "BlueprintGraphPrivatePCH.h"
 
 #include "CompilerResultsLog.h"
-#include "ClassIconFinder.h"
+#include "SlateIconFinder.h"
 #include "MessageLog.h"
 
 #define LOCTEXT_NAMESPACE "K2Node"
@@ -427,42 +427,33 @@ void UK2Node_Variable::ValidateNodeDuringCompilation(class FCompilerResultsLog&
 	}
 }
 
-FName UK2Node_Variable::GetPaletteIcon(FLinearColor& ColorOut) const
+FSlateIcon UK2Node_Variable::GetIconAndTint(FLinearColor& ColorOut) const
 {
-	FName ReturnIconName;
+	const UStruct* VarScope = VariableReference.IsLocalScope() ?
+		VariableReference.GetMemberScope(GetBlueprintClassFromNode()) :
+		GetVariableSourceClass();
 
-	if(VariableReference.IsLocalScope())
-	{
-		ReturnIconName = GetVariableIconAndColor(VariableReference.GetMemberScope(GetBlueprintClassFromNode()), GetVarName(), ColorOut);
-	}
-	else
-	{
-		ReturnIconName = GetVariableIconAndColor(GetVariableSourceClass(), GetVarName(), ColorOut);
-	}
-
-	return ReturnIconName;
+	return GetVariableIconAndColor(VarScope, GetVarName(), ColorOut);
 }
 
-FName UK2Node_Variable::GetVarIconFromPinType(const FEdGraphPinType& InPinType, FLinearColor& IconColorOut)
+FSlateIcon UK2Node_Variable::GetVarIconFromPinType(const FEdGraphPinType& InPinType, FLinearColor& IconColorOut)
 {
-	FName IconBrush = TEXT("Kismet.AllClasses.VariableIcon");
-
 	const UEdGraphSchema_K2* K2Schema = GetDefault<UEdGraphSchema_K2>();
 	IconColorOut = K2Schema->GetPinTypeColor(InPinType);
 
 	if(InPinType.bIsArray)
 	{
-		IconBrush = TEXT("Kismet.AllClasses.ArrayVariableIcon");
+		return FSlateIcon("EditorStyle", "Kismet.AllClasses.ArrayVariableIcon");
 	}
 	else if(InPinType.PinSubCategoryObject.IsValid())
 	{
 		if(UClass* Class = Cast<UClass>(InPinType.PinSubCategoryObject.Get()))
 		{
-			IconBrush = FClassIconFinder::FindIconNameForClass( Class );
+			return FSlateIconFinder::FindIconForClass( Class );
 		}
 	}
 
-	return IconBrush;
+	return FSlateIcon("EditorStyle", "Kismet.AllClasses.VariableIcon");
 }
 
 FText UK2Node_Variable::GetToolTipHeading() const
@@ -507,10 +498,8 @@ void UK2Node_Variable::GetNodeAttributes( TArray<TKeyValuePair<FString, FString>
 	OutNodeAttributes.Add( TKeyValuePair<FString, FString>( TEXT( "Name" ), VariableName ));
 }
 
-FName UK2Node_Variable::GetVariableIconAndColor(const UStruct* VarScope, FName VarName, FLinearColor& IconColorOut)
+FSlateIcon UK2Node_Variable::GetVariableIconAndColor(const UStruct* VarScope, FName VarName, FLinearColor& IconColorOut)
 {
-	FName IconBrush = TEXT("Kismet.AllClasses.VariableIcon");
-
 	if(VarScope != NULL)
 	{
 		UProperty* Property = FindField<UProperty>(VarScope, VarName);
@@ -521,12 +510,12 @@ FName UK2Node_Variable::GetVariableIconAndColor(const UStruct* VarScope, FName V
 			FEdGraphPinType PinType;
 			if(K2Schema->ConvertPropertyToPinType(Property,  PinType)) // use schema to get the color
 			{
-				IconBrush = GetVarIconFromPinType(PinType, IconColorOut);
+				return GetVarIconFromPinType(PinType, IconColorOut);
 			}
 		}
 	}
 
-	return IconBrush;
+	return FSlateIcon("EditorStyle", "Kismet.AllClasses.VariableIcon");
 }
 
 
diff --git a/Engine/Source/Editor/BlueprintGraph/Public/BlueprintNodeSpawner.h b/Engine/Source/Editor/BlueprintGraph/Public/BlueprintNodeSpawner.h
--- a/Engine/Source/Editor/BlueprintGraph/Public/BlueprintNodeSpawner.h
+++ b/Engine/Source/Editor/BlueprintGraph/Public/BlueprintNodeSpawner.h
@@ -55,7 +55,7 @@ struct FBlueprintActionUiSpec
 	FText   Category;
 	FText   Tooltip;
 	FText	Keywords;
-	FName   IconName;
+	FSlateIcon Icon;
 	FLinearColor IconTint;
 	FString DocLink;
 	FString DocExcerptTag;
@@ -67,7 +67,7 @@ private:
 		Category = MoveTemp(Other.Category);
 		Tooltip = MoveTemp(Other.Tooltip);
 		Keywords = MoveTemp(Other.Keywords);
-		IconName = MoveTemp(Other.IconName);
+		Icon = MoveTemp(Other.Icon);
 		IconTint = MoveTemp(Other.IconTint);
 		DocLink = MoveTemp(Other.DocLink);
 		DocExcerptTag = MoveTemp(Other.DocExcerptTag);
@@ -79,7 +79,7 @@ private:
 		Category = Other.Category;
 		Tooltip = Other.Tooltip;
 		Keywords = Other.Keywords;
-		IconName = Other.IconName;
+		Icon = Other.Icon;
 		IconTint = Other.IconTint;
 		DocLink = Other.DocLink;
 		DocExcerptTag = Other.DocExcerptTag;
diff --git a/Engine/Source/Editor/Cascade/Private/CascadeActions.cpp b/Engine/Source/Editor/Cascade/Private/CascadeActions.cpp
--- a/Engine/Source/Editor/Cascade/Private/CascadeActions.cpp
+++ b/Engine/Source/Editor/Cascade/Private/CascadeActions.cpp
@@ -4,7 +4,7 @@
 #include "CascadeActions.h"
 #include "Cascade.h"
 
-#define LOCTEXT_NAMESPACE ""
+#define LOCTEXT_NAMESPACE "CascadeCommands"
 
 /** UI_COMMAND takes long for the compile to optimize */
 PRAGMA_DISABLE_OPTIMIZATION
diff --git a/Engine/Source/Editor/ClassViewer/Private/SClassViewer.cpp b/Engine/Source/Editor/ClassViewer/Private/SClassViewer.cpp
--- a/Engine/Source/Editor/ClassViewer/Private/SClassViewer.cpp
+++ b/Engine/Source/Editor/ClassViewer/Private/SClassViewer.cpp
@@ -6,7 +6,7 @@
 #include "Editor/ClassViewer/Private/SClassViewer.h"
 #include "Editor/UnrealEd/Public/DragAndDrop/ClassDragDropOp.h"
 #include "Editor/UnrealEd/Public/DragAndDrop/AssetDragDropOp.h"
-#include "Editor/UnrealEd/Public/ClassIconFinder.h"
+
 #include "Editor/ContentBrowser/Public/ContentBrowserModule.h"
 
 #include "Kismet2/KismetEditorUtilities.h"
@@ -18,6 +18,8 @@
 #include "AssetSelection.h"
 #include "AssetToolsModule.h"
 
+#include "SlateIconFinder.h"
+
 #include "ClassViewerNode.h"
 
 #include "UnloadedBlueprintData.h"
@@ -1297,7 +1299,7 @@ public:
 
 		bool bIsRestricted = AssociatedNode->IsRestricted();
 
-		const FSlateBrush* ClassIcon = FClassIconFinder::FindIconForClass(AssociatedNode->Class.Get());
+		const FSlateBrush* ClassIcon = FSlateIconFinder::FindIconBrushForClass(AssociatedNode->Class.Get());
 		
 		this->ChildSlot
 		[
@@ -1520,10 +1522,7 @@ static TSharedPtr< FClassViewerNode > CreateNodeForClass(UClass* Class, const TM
 	if ( BlueprintList.Num() )
 	{
 		// Grab the generated class name and check it before assigning. Objects that haven't been saved since this has started to be exported do not have the information.
-		if (auto GeneratedClassnamePtr = BlueprintList[0].TagsAndValues.Find(FName("GeneratedClass")))
-		{
-			NewNode->GeneratedClassname = FName(**GeneratedClassnamePtr);
-		}
+		BlueprintList[0].GetTagValue("GeneratedClass", NewNode->GeneratedClassname);
 	}
 
 	return NewNode;
@@ -1790,29 +1789,26 @@ void FClassHierarchy::FindClass(TSharedPtr< FClassViewerNode > InOutClassNode)
 
 void FClassHierarchy::LoadUnloadedTagData(TSharedPtr<FClassViewerNode>& InOutClassViewerNode, const FAssetData& InAssetData)
 {
-	const FString GeneratedClassPackage = InAssetData.PackageName.ToString();
-
-	const FString* GeneratedClassname = InAssetData.TagsAndValues.Find("GeneratedClass");
-	const FString* ParentClassname = InAssetData.TagsAndValues.Find("ParentClass");
-	const FString* BlueprintType = InAssetData.TagsAndValues.Find("BlueprintType");
-	const FString AssetName = InAssetData.AssetName.ToString();
-
 	// Create the viewer node.
+	const FString AssetName = InAssetData.AssetName.ToString();
 	InOutClassViewerNode = MakeShareable(new FClassViewerNode(AssetName, AssetName));
+
+	InOutClassViewerNode->AssetName = AssetName;
+	InOutClassViewerNode->GeneratedClassPackage = InAssetData.PackageName.ToString();
+	InAssetData.GetTagValue("GeneratedClass", InOutClassViewerNode->GeneratedClassname);
+	InAssetData.GetTagValue("ParentClass", InOutClassViewerNode->ParentClassname);
+	InOutClassViewerNode->bIsBPNormalType = InAssetData.GetTagValueRef<FString>("BlueprintType") == TEXT("BPType_Normal");
 			
 	// It is an unloaded blueprint, so we need to create the structure that will hold the data.
 	TSharedPtr<FUnloadedBlueprintData> UnloadedBlueprintData = MakeShareable( new FUnloadedBlueprintData(InOutClassViewerNode) );
 	InOutClassViewerNode->UnloadedBlueprintData = UnloadedBlueprintData;
 
 	// Get the class flags.
-	const FString* ClassFlags = InAssetData.TagsAndValues.Find("ClassFlags");
-	if(ClassFlags)
-	{
-		InOutClassViewerNode->UnloadedBlueprintData->SetClassFlags(FCString::Atoi(**ClassFlags));
-	}
+	const uint32 ClassFlags = InAssetData.GetTagValueRef<uint32>("ClassFlags");
+	InOutClassViewerNode->UnloadedBlueprintData->SetClassFlags(ClassFlags);
 
-	const FString* ImplementedInterfaces = InAssetData.TagsAndValues.Find("ImplementedInterfaces");
-	if(ImplementedInterfaces)
+	const FString ImplementedInterfaces = InAssetData.GetTagValueRef<FString>("ImplementedInterfaces");
+	if(!ImplementedInterfaces.IsEmpty())
 	{
 		FString FullInterface;
 		FString RemainingString;
@@ -1833,26 +1829,6 @@ void FClassHierarchy::LoadUnloadedTagData(TSharedPtr<FClassViewerNode>& InOutCla
 			CurrentString = RemainingString;
 		}
 	}
-
-	InOutClassViewerNode->GeneratedClassPackage = GeneratedClassPackage;
-
-	InOutClassViewerNode->ParentClassname = NAME_None;
-	if(ParentClassname)
-	{
-		InOutClassViewerNode->ParentClassname = FName(**ParentClassname);
-	}
-
-	if(GeneratedClassname)
-	{
-		InOutClassViewerNode->GeneratedClassname = FName(**GeneratedClassname);
-	}
-
-	if(BlueprintType && *BlueprintType == TEXT("BPType_Normal"))
-	{
-		InOutClassViewerNode->bIsBPNormalType = true;
-	}
-
-	InOutClassViewerNode->AssetName = AssetName;
 }
 
 void FClassHierarchy::PopulateClassHierarchy()
diff --git a/Engine/Source/Editor/ContentBrowser/Private/AssetContextMenu.cpp b/Engine/Source/Editor/ContentBrowser/Private/AssetContextMenu.cpp
--- a/Engine/Source/Editor/ContentBrowser/Private/AssetContextMenu.cpp
+++ b/Engine/Source/Editor/ContentBrowser/Private/AssetContextMenu.cpp
@@ -870,10 +870,10 @@ bool FAssetContextMenu::AddDocumentationMenuOptions(FMenuBuilder& MenuBuilder)
 		const bool bIsBlueprint = SelectedClass->IsChildOf<UBlueprint>();
 		if (bIsBlueprint)
 		{
-			auto ParentClassPath = SelectedAssets[0].TagsAndValues.Find(GET_MEMBER_NAME_CHECKED(UBlueprint,ParentClass));
-			if (ParentClassPath)
+			const FString ParentClassPath = SelectedAssets[0].GetTagValueRef<FString>(GET_MEMBER_NAME_CHECKED(UBlueprint,ParentClass));
+			if (!ParentClassPath.IsEmpty())
 			{
-				SelectedClass = FindObject<UClass>(nullptr,**ParentClassPath);
+				SelectedClass = FindObject<UClass>(nullptr,*ParentClassPath);
 			}
 		}
 
@@ -919,8 +919,8 @@ bool FAssetContextMenu::AddDocumentationMenuOptions(FMenuBuilder& MenuBuilder)
 						}
 
 						UEnum* BlueprintTypeEnum = FindObject<UEnum>(ANY_PACKAGE, TEXT("EBlueprintType"), true);
-						auto EnumString = SelectedAssets[0].TagsAndValues.Find(GET_MEMBER_NAME_CHECKED(UBlueprint,BlueprintType));
-						EBlueprintType BlueprintType = (EnumString ? (EBlueprintType)BlueprintTypeEnum->GetValueByName(**EnumString) : BPTYPE_Normal);
+						const FString EnumString = SelectedAssets[0].GetTagValueRef<FString>(GET_MEMBER_NAME_CHECKED(UBlueprint,BlueprintType));
+						EBlueprintType BlueprintType = (!EnumString.IsEmpty() ? (EBlueprintType)BlueprintTypeEnum->GetValueByName(*EnumString) : BPTYPE_Normal);
 
 						switch (BlueprintType)
 						{
@@ -1446,8 +1446,8 @@ void FAssetContextMenu::ExecuteFindInExplorer()
 				FString ModulePath;
 				if (FSourceCodeNavigation::FindModulePath(ModuleName, ModulePath))
 				{
-					const FString* RelativePath = AssetData.TagsAndValues.Find("ModuleRelativePath");
-					if (RelativePath)
+					FString RelativePath;
+					if (AssetData.GetTagValue("ModuleRelativePath", RelativePath))
 					{
 						const FString FullFilePath = FPaths::ConvertRelativePathToFull(ModulePath / (*RelativePath));
 						FPlatformProcess::ExploreFolder(*FullFilePath);
diff --git a/Engine/Source/Editor/ContentBrowser/Private/AssetViewSortManager.cpp b/Engine/Source/Editor/ContentBrowser/Private/AssetViewSortManager.cpp
--- a/Engine/Source/Editor/ContentBrowser/Private/AssetViewSortManager.cpp
+++ b/Engine/Source/Editor/ContentBrowser/Private/AssetViewSortManager.cpp
@@ -182,19 +182,21 @@ protected:
 	FORCEINLINE virtual bool Compare(const TSharedPtr<FAssetViewItem>& A, const TSharedPtr<FAssetViewItem>& B) const override
 	{
 		// Depending if we're sorting ascending or descending it's quicker to flip the compares incase tags are missing
-		const FString* Value1 = bAscending ? StaticCastSharedPtr<FAssetViewAsset>(A)->Data.TagsAndValues.Find(Tag) : StaticCastSharedPtr<FAssetViewAsset>(B)->Data.TagsAndValues.Find(Tag);
-		if (!Value1)
+		FString Value1;
+		const bool bFoundValue1 = bAscending ? StaticCastSharedPtr<FAssetViewAsset>(A)->Data.GetTagValue(Tag, Value1) : StaticCastSharedPtr<FAssetViewAsset>(B)->Data.GetTagValue(Tag, Value1);
+		if (!bFoundValue1)
 		{
 			return true;
 		}
 
-		const FString* Value2 = bAscending ? StaticCastSharedPtr<FAssetViewAsset>(B)->Data.TagsAndValues.Find(Tag) : StaticCastSharedPtr<FAssetViewAsset>(A)->Data.TagsAndValues.Find(Tag);
-		if (!Value2)
+		FString Value2;
+		const bool bFoundValue2 = bAscending ? StaticCastSharedPtr<FAssetViewAsset>(B)->Data.GetTagValue(Tag, Value2) : StaticCastSharedPtr<FAssetViewAsset>(A)->Data.GetTagValue(Tag, Value2);
+		if (!bFoundValue2)
 		{
 			return false;
 		}
 
-		const int32 Result = Value1->Compare(*Value2, ESearchCase::IgnoreCase);
+		const int32 Result = Value1.Compare(Value2, ESearchCase::IgnoreCase);
 		if (Result < 0)
 		{
 			return true;
@@ -217,25 +219,25 @@ protected:
 	FORCEINLINE virtual bool Compare(const TSharedPtr<FAssetViewItem>& A, const TSharedPtr<FAssetViewItem>& B) const override
 	{
 		// Depending if we're sorting ascending or descending it's quicker to flip the compares incase tags are missing
-		const FString* Value1 = bAscending ? StaticCastSharedPtr<FAssetViewAsset>(A)->Data.TagsAndValues.Find(Tag) : StaticCastSharedPtr<FAssetViewAsset>(B)->Data.TagsAndValues.Find(Tag);
-		if (!Value1)
+		float Value1 = 0.0f;
+		const bool bFoundValue1 = bAscending ? StaticCastSharedPtr<FAssetViewAsset>(A)->Data.GetTagValue(Tag, Value1) : StaticCastSharedPtr<FAssetViewAsset>(B)->Data.GetTagValue(Tag, Value1);
+		if (!bFoundValue1)
 		{
 			return true;
 		}
 
-		const FString* Value2 = bAscending ? StaticCastSharedPtr<FAssetViewAsset>(B)->Data.TagsAndValues.Find(Tag) : StaticCastSharedPtr<FAssetViewAsset>(A)->Data.TagsAndValues.Find(Tag);
-		if (!Value2)
+		float Value2 = 0.0f;
+		const bool bFoundValue2 = bAscending ? StaticCastSharedPtr<FAssetViewAsset>(B)->Data.GetTagValue(Tag, Value2) : StaticCastSharedPtr<FAssetViewAsset>(A)->Data.GetTagValue(Tag, Value2);
+		if (!bFoundValue2)
 		{
 			return false;
 		}
 
-		const float Num1 = FCString::Atof(**Value1);
-		const float Num2 = FCString::Atof(**Value2);
-		if (Num1 < Num2)
+		if (Value1 < Value2)
 		{
 			return true;
 		}
-		else if (Num1 > Num2)
+		else if (Value1 > Value2)
 		{
 			return false;
 		}
@@ -253,14 +255,16 @@ protected:
 	FORCEINLINE virtual bool Compare(const TSharedPtr<FAssetViewItem>& A, const TSharedPtr<FAssetViewItem>& B) const override
 	{
 		// Depending if we're sorting ascending or descending it's quicker to flip the compares incase tags are missing
-		const FString* Value1 = bAscending ? StaticCastSharedPtr<FAssetViewAsset>(A)->Data.TagsAndValues.Find(Tag) : StaticCastSharedPtr<FAssetViewAsset>(B)->Data.TagsAndValues.Find(Tag);
-		if (!Value1)
+		FString Value1;
+		const bool bHasFoundValue1 = bAscending ? StaticCastSharedPtr<FAssetViewAsset>(A)->Data.GetTagValue(Tag, Value1) : StaticCastSharedPtr<FAssetViewAsset>(B)->Data.GetTagValue(Tag, Value1);
+		if (!bHasFoundValue1)
 		{
 			return true;
 		}
 
-		const FString* Value2 = bAscending ? StaticCastSharedPtr<FAssetViewAsset>(B)->Data.TagsAndValues.Find(Tag) : StaticCastSharedPtr<FAssetViewAsset>(A)->Data.TagsAndValues.Find(Tag);
-		if (!Value2)
+		FString Value2;
+		const bool bHasFoundValue2 = bAscending ? StaticCastSharedPtr<FAssetViewAsset>(B)->Data.GetTagValue(Tag, Value2) : StaticCastSharedPtr<FAssetViewAsset>(A)->Data.GetTagValue(Tag, Value2);
+		if (!bHasFoundValue2)
 		{
 			return false;
 		}
@@ -268,7 +272,7 @@ protected:
 		float Num1 = 1.f;
 		{
 			TArray<FString> Tokens1;
-			Value1->ParseIntoArray(Tokens1, TEXT("x"), true);
+			Value1.ParseIntoArray(Tokens1, TEXT("x"), true);
 			for (auto TokenIt1 = Tokens1.CreateConstIterator(); TokenIt1; ++TokenIt1)
 			{
 				Num1 *= FCString::Atof(**TokenIt1);
@@ -277,7 +281,7 @@ protected:
 		float Num2 = 1.f;
 		{
 			TArray<FString> Tokens2;
-			Value2->ParseIntoArray(Tokens2, TEXT("x"), true);
+			Value2.ParseIntoArray(Tokens2, TEXT("x"), true);
 			for (auto TokenIt2 = Tokens2.CreateConstIterator(); TokenIt2; ++TokenIt2)
 			{
 				Num2 *= FCString::Atof(**TokenIt2);
diff --git a/Engine/Source/Editor/ContentBrowser/Private/AssetViewWidgets.cpp b/Engine/Source/Editor/ContentBrowser/Private/AssetViewWidgets.cpp
--- a/Engine/Source/Editor/ContentBrowser/Private/AssetViewWidgets.cpp
+++ b/Engine/Source/Editor/ContentBrowser/Private/AssetViewWidgets.cpp
@@ -1785,11 +1785,7 @@ FText SAssetColumnItem::GetAssetTagText(FName AssetRegistryTag) const
 	{
 		if(AssetItem->GetType() != EAssetItemType::Folder)
 		{
-			const FString* TagValue = StaticCastSharedPtr<FAssetViewAsset>(AssetItem)->Data.TagsAndValues.Find(AssetRegistryTag);
-			if ( TagValue != NULL )
-			{
-				return FText::FromString(*TagValue);
-			}
+			return StaticCastSharedPtr<FAssetViewAsset>(AssetItem)->Data.GetTagValueRef<FText>(AssetRegistryTag);
 		}
 	}
 	
diff --git a/Engine/Source/Editor/ContentBrowser/Private/CollectionAssetRegistryBridge.cpp b/Engine/Source/Editor/ContentBrowser/Private/CollectionAssetRegistryBridge.cpp
--- a/Engine/Source/Editor/ContentBrowser/Private/CollectionAssetRegistryBridge.cpp
+++ b/Engine/Source/Editor/ContentBrowser/Private/CollectionAssetRegistryBridge.cpp
@@ -70,10 +70,9 @@ public:
 
 				// Get the destination object from the meta-data rather than load the redirector object, as 
 				// loading a redirector will also load the object it points to, which could cause a large hitch
-				const FString* DestinationObjectStrPtr = ObjectAssetData.TagsAndValues.Find("DestinationObject");
-				if (DestinationObjectStrPtr)
+				FString DestinationObjectPath;
+				if (ObjectAssetData.GetTagValue("DestinationObject", DestinationObjectPath))
 				{
-					FString DestinationObjectPath = *DestinationObjectStrPtr;
 					ConstructorHelpers::StripObjectClass(DestinationObjectPath);
 					ObjectAssetData = AssetRegistryModule.Get().GetAssetByObjectPath(*DestinationObjectPath);
 				}
diff --git a/Engine/Source/Editor/ContentBrowser/Private/FrontendFilters.cpp b/Engine/Source/Editor/ContentBrowser/Private/FrontendFilters.cpp
--- a/Engine/Source/Editor/ContentBrowser/Private/FrontendFilters.cpp
+++ b/Engine/Source/Editor/ContentBrowser/Private/FrontendFilters.cpp
@@ -357,10 +357,10 @@ public:
 		}
 
 		// Generic handling for anything in the asset meta-data
-		const FString* const MetaDataValue = AssetPtr->TagsAndValues.Find(InKey);
-		if (MetaDataValue)
+		FString MetaDataValue;
+		if (AssetPtr->GetTagValue(InKey, MetaDataValue))
 		{
-			return TextFilterUtils::TestComplexExpression(*MetaDataValue, InValue, InComparisonOperation, InTextComparisonMode);
+			return TextFilterUtils::TestComplexExpression(MetaDataValue, InValue, InComparisonOperation, InTextComparisonMode);
 		}
 
 		return false;
@@ -596,8 +596,8 @@ void FFrontendFilter_Modified::OnPackageDirtyStateUpdated(UPackage* Package)
 
 bool FFrontendFilter_ReplicatedBlueprint::PassesFilter(FAssetFilterType InItem) const
 {
-	FString TagValue = InItem.TagsAndValues.FindRef("NumReplicatedProperties");
-	return !TagValue.IsEmpty() && (FCString::Atoi(*TagValue) > 0);
+	const int32 NumReplicatedProperties = InItem.GetTagValueRef<int32>("NumReplicatedProperties");
+	return NumReplicatedProperties > 0;
 }
 
 /////////////////////////////////////////
@@ -634,9 +634,10 @@ FText FFrontendFilter_ArbitraryComparisonOperation::GetToolTipText() const
 
 bool FFrontendFilter_ArbitraryComparisonOperation::PassesFilter(FAssetFilterType InItem) const
 {
-	if (const FString* pTagValue = InItem.TagsAndValues.Find(TagName))
+	FString TagValue;
+	if (InItem.GetTagValue(TagName, TagValue))
 	{
-		return TextFilterUtils::TestComplexExpression(*pTagValue, TargetTagValue, ComparisonOp, ETextFilterTextComparisonMode::Exact);
+		return TextFilterUtils::TestComplexExpression(TagValue, TargetTagValue, ComparisonOp, ETextFilterTextComparisonMode::Exact);
 	}
 	else
 	{
diff --git a/Engine/Source/Editor/DataTableEditor/Private/SRowEditor.cpp b/Engine/Source/Editor/DataTableEditor/Private/SRowEditor.cpp
--- a/Engine/Source/Editor/DataTableEditor/Private/SRowEditor.cpp
+++ b/Engine/Source/Editor/DataTableEditor/Private/SRowEditor.cpp
@@ -281,7 +281,20 @@ FReply SRowEditor::OnRemoveClicked()
 	if (DataTable.IsValid())
 	{
 		const FName RowToRemove = GetCurrentName();
-		FDataTableEditorUtils::RemoveRow(DataTable.Get(), RowToRemove);
+		const int32 RowToRemoveIndex = CachedRowNames.IndexOfByPredicate([&](const TSharedPtr<FName>& InRowName) -> bool
+		{
+			return *InRowName == RowToRemove;
+		});
+
+		if (FDataTableEditorUtils::RemoveRow(DataTable.Get(), RowToRemove))
+		{
+			// Try and keep the same row index selected
+			const int32 RowIndexToSelect = FMath::Clamp(RowToRemoveIndex, 0, CachedRowNames.Num() - 1);
+			if (CachedRowNames.IsValidIndex(RowIndexToSelect))
+			{
+				SelectRow(*CachedRowNames[RowIndexToSelect]);
+			}
+		}
 	}
 	return FReply::Handled();
 }
diff --git a/Engine/Source/Editor/DetailCustomizations/Private/AudioSettingsDetails.cpp b/Engine/Source/Editor/DetailCustomizations/Private/AudioSettingsDetails.cpp
--- a/Engine/Source/Editor/DetailCustomizations/Private/AudioSettingsDetails.cpp
+++ b/Engine/Source/Editor/DetailCustomizations/Private/AudioSettingsDetails.cpp
@@ -4,6 +4,8 @@
 #include "AudioSettingsDetails.h"
 #include "Runtime/Engine/Classes/Sound/AudioSettings.h"
 
+#define LOCTEXT_NAMESPACE "AudioSettingsDetails"
+
 TSharedRef<IDetailCustomization> FAudioSettingsDetails::MakeInstance()
 {
 	return MakeShareable(new FAudioSettingsDetails);
@@ -15,4 +17,50 @@ void FAudioSettingsDetails::CustomizeDetails(IDetailLayoutBuilder& DetailBuilder
 	{
 		DetailBuilder.HideProperty(GET_MEMBER_NAME_CHECKED(UAudioSettings, MaximumConcurrentStreams));
 	}
+
+	DialogueFilenameFormatProperty = DetailBuilder.GetProperty(GET_MEMBER_NAME_CHECKED(UAudioSettings, DialogueFilenameFormat));
+	DialogueFilenameFormatProperty->MarkHiddenByCustomization();
+
+	IDetailCategoryBuilder& DialogueCategoryBuilder = DetailBuilder.EditCategory("Dialogue", LOCTEXT("DialogueCategoryLabel", "Dialogue"));
+	DialogueCategoryBuilder.AddCustomRow(DialogueFilenameFormatProperty->GetPropertyDisplayName())
+		.NameContent()
+		[
+			DialogueFilenameFormatProperty->CreatePropertyNameWidget()
+		]
+		.ValueContent()
+		.HAlign(HAlign_Fill)
+		.MaxDesiredWidth(TOptional<float>())
+		[
+			SNew(SHorizontalBox)
+
+			+SHorizontalBox::Slot()
+			.AutoWidth()
+			[
+				DialogueFilenameFormatProperty->CreatePropertyValueWidget()
+			]
+
+			+SHorizontalBox::Slot()
+			.HAlign(HAlign_Right)
+			.VAlign(VAlign_Center)
+			.Padding(FMargin(4.0f, 0.0f, 30.0f, 0.0f))
+			[
+				SNew(STextBlock)
+				.Font(DetailBuilder.GetDetailFont())
+				.Text(this, &FAudioSettingsDetails::GetExampleDialogueFilename)
+			]
+		];
 }
+
+FText FAudioSettingsDetails::GetExampleDialogueFilename() const
+{
+	static const FGuid DummyGuid(0xa05875c2, 0xc7ca4601, 0x92e03564, 0x532674a3);
+	static const FString DummyAssetName = TEXT("Bob_CannotDo_Dialogue");
+	static const FString DummyContextId = TEXT("C174FFF6B897CD21");
+
+	FString FormatString;
+	DialogueFilenameFormatProperty->GetValue(FormatString);
+
+	return FText::FromString(UDialogueWave::BuildRecordedAudioFilename(FormatString, DummyGuid, DummyAssetName, DummyContextId, 2));
+}
+
+#undef LOCTEXT_NAMESPACE
diff --git a/Engine/Source/Editor/DetailCustomizations/Private/AudioSettingsDetails.h b/Engine/Source/Editor/DetailCustomizations/Private/AudioSettingsDetails.h
--- a/Engine/Source/Editor/DetailCustomizations/Private/AudioSettingsDetails.h
+++ b/Engine/Source/Editor/DetailCustomizations/Private/AudioSettingsDetails.h
@@ -11,4 +11,8 @@ public:
 private:
 	/** IDetailCustomization interface */
 	virtual void CustomizeDetails(IDetailLayoutBuilder& DetailBuilder) override;
+
+	FText GetExampleDialogueFilename() const;
+
+	TSharedPtr<IPropertyHandle> DialogueFilenameFormatProperty;
 };
\ No newline at end of file
diff --git a/Engine/Source/Editor/DetailCustomizations/Private/ColorStructCustomization.cpp b/Engine/Source/Editor/DetailCustomizations/Private/ColorStructCustomization.cpp
--- a/Engine/Source/Editor/DetailCustomizations/Private/ColorStructCustomization.cpp
+++ b/Engine/Source/Editor/DetailCustomizations/Private/ColorStructCustomization.cpp
@@ -387,7 +387,15 @@ FReply FColorStructCustomization::OnMouseButtonDownColorBlock(const FGeometry& M
 		return FReply::Unhandled();
 	}
 	
-	CreateColorPicker(true /*bUseAlpha*/);
+	bool CanShowColorPicker = true;
+	if (StructPropertyHandle.IsValid() && StructPropertyHandle->GetProperty() != nullptr)
+	{
+		CanShowColorPicker = !(StructPropertyHandle->GetProperty()->HasAllPropertyFlags(CPF_EditConst));
+	}
+	if (CanShowColorPicker)
+	{
+		CreateColorPicker(true /*bUseAlpha*/);
+	}
 
 	return FReply::Handled();
 }
diff --git a/Engine/Source/Editor/DetailCustomizations/Private/DialogueWaveDetails.cpp b/Engine/Source/Editor/DetailCustomizations/Private/DialogueWaveDetails.cpp
--- a/Engine/Source/Editor/DetailCustomizations/Private/DialogueWaveDetails.cpp
+++ b/Engine/Source/Editor/DetailCustomizations/Private/DialogueWaveDetails.cpp
@@ -258,7 +258,7 @@ void FDialogueWaveDetails::CustomizeDetails(class IDetailLayoutBuilder& DetailBu
 {
 	DetailLayoutBuilder = &DetailBuilder;
 
-	IDetailCategoryBuilder& ContextMappingsDetailCategoryBuilder = DetailBuilder.EditCategory("DialogueContexts");
+	IDetailCategoryBuilder& ContextMappingsDetailCategoryBuilder = DetailBuilder.EditCategory("DialogueContexts", FText::GetEmpty(), ECategoryPriority::Important);
 
 	// Add Context Button
 	ContextMappingsDetailCategoryBuilder.AddCustomRow(LOCTEXT("AddDialogueContext", "Add Dialogue Context"))
diff --git a/Engine/Source/Editor/DetailCustomizations/Private/SkeletalMeshComponentDetails.cpp b/Engine/Source/Editor/DetailCustomizations/Private/SkeletalMeshComponentDetails.cpp
--- a/Engine/Source/Editor/DetailCustomizations/Private/SkeletalMeshComponentDetails.cpp
+++ b/Engine/Source/Editor/DetailCustomizations/Private/SkeletalMeshComponentDetails.cpp
@@ -234,8 +234,8 @@ EVisibility FSkeletalMeshComponentDetails::VisibilityForAnimationMode(EAnimation
 
 bool FSkeletalMeshComponentDetails::OnShouldFilterAnimAsset( const FAssetData& AssetData )
 {
-	const FString* SkeletonName = AssetData.TagsAndValues.Find(TEXT("Skeleton"));
-	return *SkeletonName != SelectedSkeletonName;
+	const FString SkeletonName = AssetData.GetTagValueRef<FString>("Skeleton");
+	return SkeletonName != SelectedSkeletonName;
 }
 
 void FSkeletalMeshComponentDetails::SkeletalMeshPropertyChanged()
diff --git a/Engine/Source/Editor/DetailCustomizations/Private/TextCustomization.cpp b/Engine/Source/Editor/DetailCustomizations/Private/TextCustomization.cpp
--- a/Engine/Source/Editor/DetailCustomizations/Private/TextCustomization.cpp
+++ b/Engine/Source/Editor/DetailCustomizations/Private/TextCustomization.cpp
@@ -19,12 +19,12 @@ namespace
 
 		virtual bool IsMultiLineText() const override
 		{
-			return PropertyHandle->IsValidHandle() && PropertyHandle->GetBoolMetaData("MultiLine");
+			return PropertyHandle->IsValidHandle() && PropertyHandle->GetMetaDataProperty()->GetBoolMetaData("MultiLine");
 		}
 
 		virtual bool IsPassword() const override
 		{
-			return PropertyHandle->IsValidHandle() && PropertyHandle->GetBoolMetaData("PasswordField");
+			return PropertyHandle->IsValidHandle() && PropertyHandle->GetMetaDataProperty()->GetBoolMetaData("PasswordField");
 		}
 
 		virtual bool IsReadOnly() const override
diff --git a/Engine/Source/Editor/DistCurveEditor/Private/CurveEditorActions.cpp b/Engine/Source/Editor/DistCurveEditor/Private/CurveEditorActions.cpp
--- a/Engine/Source/Editor/DistCurveEditor/Private/CurveEditorActions.cpp
+++ b/Engine/Source/Editor/DistCurveEditor/Private/CurveEditorActions.cpp
@@ -3,7 +3,7 @@
 #include "DistCurveEditorModule.h"
 #include "CurveEditorActions.h"
 
-#define LOCTEXT_NAMESPACE ""
+#define LOCTEXT_NAMESPACE "CurveEditorCommands"
 
 void FCurveEditorCommands::RegisterCommands()
 {
diff --git a/Engine/Source/Editor/FoliageEdit/Private/FoliageEditActions.cpp b/Engine/Source/Editor/FoliageEdit/Private/FoliageEditActions.cpp
--- a/Engine/Source/Editor/FoliageEdit/Private/FoliageEditActions.cpp
+++ b/Engine/Source/Editor/FoliageEdit/Private/FoliageEditActions.cpp
@@ -6,7 +6,7 @@
 
 #include "FoliageEditActions.h"
 
-#define LOCTEXT_NAMESPACE ""
+#define LOCTEXT_NAMESPACE "FoliageEditCommands"
 
 void FFoliageEditCommands::RegisterCommands()
 {
diff --git a/Engine/Source/Editor/FoliageEdit/Private/FoliagePaletteCommands.cpp b/Engine/Source/Editor/FoliageEdit/Private/FoliagePaletteCommands.cpp
--- a/Engine/Source/Editor/FoliageEdit/Private/FoliagePaletteCommands.cpp
+++ b/Engine/Source/Editor/FoliageEdit/Private/FoliagePaletteCommands.cpp
@@ -3,7 +3,7 @@
 #include "UnrealEd.h"
 #include "FoliagePaletteCommands.h"
 
-#define LOCTEXT_NAMESPACE ""
+#define LOCTEXT_NAMESPACE "FoliagePaletteCommands"
 
 void FFoliagePaletteCommands::RegisterCommands()
 {
diff --git a/Engine/Source/Editor/FoliageEdit/Private/FoliageTypeCustomizationHelpers.cpp b/Engine/Source/Editor/FoliageEdit/Private/FoliageTypeCustomizationHelpers.cpp
--- a/Engine/Source/Editor/FoliageEdit/Private/FoliageTypeCustomizationHelpers.cpp
+++ b/Engine/Source/Editor/FoliageEdit/Private/FoliageTypeCustomizationHelpers.cpp
@@ -7,6 +7,7 @@
 #include "DetailLayoutBuilder.h"
 #include "DetailCategoryBuilder.h"
 #include "PropertyHandle.h"
+#include "ObjectEditorUtils.h"
 
 void FFoliageTypeCustomizationHelpers::ModifyFoliagePropertyRow(IDetailPropertyRow* PropertyRow, const TAttribute<EVisibility>& InVisibility, const TAttribute<bool>& InEnabled)
 {
@@ -23,6 +24,17 @@ void FFoliageTypeCustomizationHelpers::ModifyFoliagePropertyRow(IDetailPropertyR
 	}
 }
 
+void FFoliageTypeCustomizationHelpers::AddBodyInstanceProperties(IDetailLayoutBuilder& LayoutBuilder)
+{
+	static const FName PhysicsName("Physics");
+	IDetailCategoryBuilder& PhysicsCategory = LayoutBuilder.EditCategory(PhysicsName);
+
+	TSharedRef<IPropertyHandle> BodyInstanceHandle = LayoutBuilder.GetProperty(GET_MEMBER_NAME_CHECKED(UPrimitiveComponent, BodyInstance));
+
+	PhysicsCategory.AddProperty(BodyInstanceHandle->GetChildHandle(GET_MEMBER_NAME_CHECKED(FBodyInstance, WalkableSlopeOverride)));
+	PhysicsCategory.AddProperty(BodyInstanceHandle->GetChildHandle(GET_MEMBER_NAME_CHECKED(FBodyInstance, PhysMaterialOverride)));
+}
+
 void FFoliageTypeCustomizationHelpers::HideFoliageCategory(IDetailLayoutBuilder& DetailLayoutBuilder, FName CategoryName)
 {
 	TArray<TSharedRef<IPropertyHandle>> CategoryProperties;
diff --git a/Engine/Source/Editor/FoliageEdit/Private/FoliageTypeCustomizationHelpers.h b/Engine/Source/Editor/FoliageEdit/Private/FoliageTypeCustomizationHelpers.h
--- a/Engine/Source/Editor/FoliageEdit/Private/FoliageTypeCustomizationHelpers.h
+++ b/Engine/Source/Editor/FoliageEdit/Private/FoliageTypeCustomizationHelpers.h
@@ -17,6 +17,8 @@ public:
 	/** Modifies the visibility and enabled attributes of an existing property row */
 	static void ModifyFoliagePropertyRow(IDetailPropertyRow* PropertyRow, const TAttribute<EVisibility>& InVisibility, const TAttribute<bool>& InEnabled);
 
+	static void AddBodyInstanceProperties(IDetailLayoutBuilder& LayoutBuilder);
+		
 	/** Hides all properties in the given category */
 	static void HideFoliageCategory(IDetailLayoutBuilder& DetailLayoutBuilder, FName CategoryName);
 
diff --git a/Engine/Source/Editor/FoliageEdit/Private/FoliageTypeDetails.cpp b/Engine/Source/Editor/FoliageEdit/Private/FoliageTypeDetails.cpp
--- a/Engine/Source/Editor/FoliageEdit/Private/FoliageTypeDetails.cpp
+++ b/Engine/Source/Editor/FoliageEdit/Private/FoliageTypeDetails.cpp
@@ -99,4 +99,6 @@ void FFoliageTypeDetails::CustomizeDetails(IDetailLayoutBuilder& DetailBuilder)
 	}
 	
 	AddSubcategoryProperties(DetailBuilder, InstanceSettingsName);
+
+	FFoliageTypeCustomizationHelpers::AddBodyInstanceProperties(DetailBuilder);
 }
\ No newline at end of file
diff --git a/Engine/Source/Editor/FoliageEdit/Private/FoliageTypeObjectCustomization.cpp b/Engine/Source/Editor/FoliageEdit/Private/FoliageTypeObjectCustomization.cpp
--- a/Engine/Source/Editor/FoliageEdit/Private/FoliageTypeObjectCustomization.cpp
+++ b/Engine/Source/Editor/FoliageEdit/Private/FoliageTypeObjectCustomization.cpp
@@ -39,8 +39,8 @@ void FFoliageTypeObjectCustomization::CustomizeChildren(TSharedRef<IPropertyHand
 bool FFoliageTypeObjectCustomization::OnShouldFilterAsset(const FAssetData& AssetData) const
 {
 	// If the asset is a BP class that doesn't inherit from UFoliageType, hide it
-	const FString* ParentClassName = AssetData.TagsAndValues.Find(TEXT("ParentClass"));
-	if (ParentClassName && !ParentClassName->Contains(TEXT("FoliageType_InstancedStaticMesh")))
+	const FString ParentClassName = AssetData.GetTagValueRef<FString>("ParentClass");
+	if (!ParentClassName.IsEmpty() && !ParentClassName.Contains(TEXT("FoliageType_InstancedStaticMesh")))
 	{
 		return true;
 	}
diff --git a/Engine/Source/Editor/FoliageEdit/Private/FoliageTypePaintingCustomization.cpp b/Engine/Source/Editor/FoliageEdit/Private/FoliageTypePaintingCustomization.cpp
--- a/Engine/Source/Editor/FoliageEdit/Private/FoliageTypePaintingCustomization.cpp
+++ b/Engine/Source/Editor/FoliageEdit/Private/FoliageTypePaintingCustomization.cpp
@@ -36,6 +36,7 @@ void FFoliageTypePaintingCustomization::CustomizeDetails(IDetailLayoutBuilder& D
 	ShowFoliagePropertiesForCategory(DetailLayoutBuilder, "Painting", PropertyRowsByName);
 	ShowFoliagePropertiesForCategory(DetailLayoutBuilder, "Placement", PropertyRowsByName);
 	ShowFoliagePropertiesForCategory(DetailLayoutBuilder, "InstanceSettings", PropertyRowsByName);
+	FFoliageTypeCustomizationHelpers::AddBodyInstanceProperties(DetailLayoutBuilder);
 
 	// Density adjustment factor should only be visible when reapplying
 	FFoliageTypeCustomizationHelpers::ModifyFoliagePropertyRow(*PropertyRowsByName.Find(GET_MEMBER_NAME_CHECKED(UFoliageType, DensityAdjustmentFactor)),
diff --git a/Engine/Source/Editor/FontEditor/Private/SCompositeFontEditor.cpp b/Engine/Source/Editor/FontEditor/Private/SCompositeFontEditor.cpp
--- a/Engine/Source/Editor/FontEditor/Private/SCompositeFontEditor.cpp
+++ b/Engine/Source/Editor/FontEditor/Private/SCompositeFontEditor.cpp
@@ -1265,7 +1265,7 @@ void SCharacterRangeEditor::Construct(const FArguments& InArgs)
 			SNew(SEditableTextBox)
 			.Text(this, &SCharacterRangeEditor::GetRangeComponentAsTCHAR, 0)
 			.OnTextCommitted(this, &SCharacterRangeEditor::OnRangeComponentCommittedAsTCHAR, 0)
-			.ToolTipText(LOCTEXT("MinCharacterRangeEditCharTooltip", "Specifies the lower inclusive boundary of this character range as a character value"))
+			.ToolTipText(LOCTEXT("MinCharacterRangeEditCharTooltip", "Specifies the lower inclusive boundary of this character range as a literal unicode character.\nExample: If you wanted to use the range 'A-Z', this would be set to 'A'."))
 		]
 
 		+SGridPanel::Slot(0, 1)
@@ -1274,7 +1274,7 @@ void SCharacterRangeEditor::Construct(const FArguments& InArgs)
 			SNew(SEditableTextBox)
 			.Text(this, &SCharacterRangeEditor::GetRangeComponentAsHexString, 0)
 			.OnTextCommitted(this, &SCharacterRangeEditor::OnRangeComponentCommittedAsHexString, 0)
-			.ToolTipText(LOCTEXT("MinCharacterRangeEditHexTooltip", "Specifies the lower inclusive boundary of this character range as a hexadecimal value"))
+			.ToolTipText(LOCTEXT("MinCharacterRangeEditHexTooltip", "Specifies the lower inclusive boundary of this character range as the hexadecimal value of a unicode character.\nExample: If you wanted to use the range '0x41-0x5A' (A-Z), this would be set to '0x41'."))
 		]
 
 		+SGridPanel::Slot(0, 2)
@@ -1283,7 +1283,7 @@ void SCharacterRangeEditor::Construct(const FArguments& InArgs)
 			SNew(SNumericEntryBox<int32>)
 			.Value(this, &SCharacterRangeEditor::GetRangeComponentAsOptional, 0)
 			.OnValueCommitted(this, &SCharacterRangeEditor::OnRangeComponentCommittedAsNumeric, 0)
-			.ToolTipText(LOCTEXT("MinCharacterRangeEditDecTooltip", "Specifies the lower inclusive boundary of this character range as a decimal value"))
+			.ToolTipText(LOCTEXT("MinCharacterRangeEditDecTooltip", "Specifies the lower inclusive boundary of this character range as the decimal value of a unicode character.\nExample: If you wanted to use the range '65-90' (A-Z), this would be set to '65'."))
 		]
 
 		// Separator
@@ -1303,7 +1303,7 @@ void SCharacterRangeEditor::Construct(const FArguments& InArgs)
 			SNew(SEditableTextBox)
 			.Text(this, &SCharacterRangeEditor::GetRangeComponentAsTCHAR, 1)
 			.OnTextCommitted(this, &SCharacterRangeEditor::OnRangeComponentCommittedAsTCHAR, 1)
-			.ToolTipText(LOCTEXT("MaxCharacterRangeEditCharTooltip", "Specifies the upper inclusive boundary of this character range as a character value"))
+			.ToolTipText(LOCTEXT("MaxCharacterRangeEditCharTooltip", "Specifies the upper inclusive boundary of this character range as a literal unicode character.\nExample: If you wanted to use the range 'A-Z', this would be set to 'Z'."))
 		]
 
 		+SGridPanel::Slot(2, 1)
@@ -1312,7 +1312,7 @@ void SCharacterRangeEditor::Construct(const FArguments& InArgs)
 			SNew(SEditableTextBox)
 			.Text(this, &SCharacterRangeEditor::GetRangeComponentAsHexString, 1)
 			.OnTextCommitted(this, &SCharacterRangeEditor::OnRangeComponentCommittedAsHexString, 1)
-			.ToolTipText(LOCTEXT("MaxCharacterRangeEditHexTooltip", "Specifies the upper inclusive boundary of this character range as a hexadecimal value"))
+			.ToolTipText(LOCTEXT("MaxCharacterRangeEditHexTooltip", "Specifies the upper inclusive boundary of this character range as the hexadecimal value of a unicode character.\nExample: If you wanted to use the range '0x41-0x5A' (A-Z), this would be set to '0x5A'."))
 		]
 
 		+SGridPanel::Slot(2, 2)
@@ -1321,7 +1321,7 @@ void SCharacterRangeEditor::Construct(const FArguments& InArgs)
 			SNew(SNumericEntryBox<int32>)
 			.Value(this, &SCharacterRangeEditor::GetRangeComponentAsOptional, 1)
 			.OnValueCommitted(this, &SCharacterRangeEditor::OnRangeComponentCommittedAsNumeric, 1)
-			.ToolTipText(LOCTEXT("MaxCharacterRangeEditDecTooltip", "Specifies the upper inclusive boundary of this character range as a decimal value"))
+			.ToolTipText(LOCTEXT("MaxCharacterRangeEditDecTooltip", "Specifies the upper inclusive boundary of this character range as the decimal value of a unicode character.\nExample: If you wanted to use the range '65-90' (A-Z), this would be set to '90'."))
 		]
 	];
 }
diff --git a/Engine/Source/Editor/GameProjectGeneration/Private/GameProjectUtils.cpp b/Engine/Source/Editor/GameProjectGeneration/Private/GameProjectUtils.cpp
--- a/Engine/Source/Editor/GameProjectGeneration/Private/GameProjectUtils.cpp
+++ b/Engine/Source/Editor/GameProjectGeneration/Private/GameProjectUtils.cpp
@@ -16,7 +16,7 @@
 #include "DesktopPlatformModule.h"
 #include "TargetPlatform.h"
 
-#include "ClassIconFinder.h"
+#include "SlateIconFinder.h"
 #include "Editor/UnrealEd/Public/SourceCodeNavigation.h"
 
 #include "UProjectInfo.h"
@@ -116,7 +116,7 @@ FText FNewClassInfo::GetClassDescription(const bool bFullDescription/* = true*/)
 const FSlateBrush* FNewClassInfo::GetClassIcon() const
 {
 	// Safe to do even if BaseClass is null, since FindIconForClass will return the default icon
-	return FClassIconFinder::FindIconForClass(BaseClass);
+	return FSlateIconFinder::FindIconBrushForClass(BaseClass);
 }
 
 FString FNewClassInfo::GetClassPrefixCPP() const
