Copying //UE4/Dev-Framework to //UE4/Dev-Main (Source: //UE4/Dev-Framework @ 3136612)

diff --git a/Engine/Source/Editor/ContentBrowser/Private/SAssetView.cpp b/Engine/Source/Editor/ContentBrowser/Private/SAssetView.cpp
--- a/Engine/Source/Editor/ContentBrowser/Private/SAssetView.cpp
+++ b/Engine/Source/Editor/ContentBrowser/Private/SAssetView.cpp
@@ -224,6 +224,8 @@ void SAssetView::Construct( const FArguments& InArgs )
 	bPendingFocusOnSync = false;
 	bWereItemsRecursivelyFiltered = false;
 
+	NumVisibleColumns = 0;
+
 	FEditorWidgetsModule& EditorWidgetsModule = FModuleManager::LoadModuleChecked<FEditorWidgetsModule>("EditorWidgets");
 	TSharedRef<SWidget> AssetDiscoveryIndicator = EditorWidgetsModule.CreateAssetDiscoveryIndicator(EAssetDiscoveryIndicatorScaleMode::Scale_Vertical);
 
@@ -712,6 +714,8 @@ void SAssetView::SaveSettings(const FString& IniFilename, const FString& IniSect
 {
 	GConfig->SetFloat(*IniSection, *GetThumbnailScaleSettingPath(SettingsString), ThumbnailScaleSliderValue.Get(), IniFilename);
 	GConfig->SetInt(*IniSection, *GetCurrentViewTypeSettingPath(SettingsString), CurrentViewType, IniFilename);
+	
+	GConfig->SetArray(*IniSection, *(SettingsString + TEXT(".HiddenColumns")), HiddenColumnNames, IniFilename);
 }
 
 void SAssetView::LoadSettings(const FString& IniFilename, const FString& IniSection, const FString& SettingsString)
@@ -734,6 +738,8 @@ void SAssetView::LoadSettings(const FString& IniFilename, const FString& IniSect
 		}
 		SetCurrentViewType( (EAssetViewType::Type)ViewType );
 	}
+	
+	GConfig->GetArray(*IniSection, *(SettingsString + TEXT(".HiddenColumns")), HiddenColumnNames, IniFilename);	
 }
 
 // Adjusts the selected asset by the selection delta, which should be +1 or -1)
@@ -1491,8 +1497,15 @@ TSharedRef<SAssetColumnView> SAssetView::CreateColumnView()
 			.SortPriority(TAttribute< EColumnSortPriority::Type >::Create(TAttribute< EColumnSortPriority::Type >::FGetter::CreateSP(this, &SAssetView::GetColumnSortPriority, SortManager.NameColumnId)))
 			.OnSort( FOnSortModeChanged::CreateSP( this, &SAssetView::OnSortColumnHeader ) )
 			.DefaultLabel( LOCTEXT("Column_Name", "Name") )
+			.ShouldGenerateWidget(TAttribute<bool>::Create(TAttribute<bool>::FGetter::CreateSP(this, &SAssetView::ShouldColumnGenerateWidget, SortManager.NameColumnId.ToString())))
+			.MenuContent()
+			[
+				CreateRowHeaderMenuContent(SortManager.NameColumnId.ToString())
+			]
 		);
 
+	NumVisibleColumns = HiddenColumnNames.Contains(SortManager.NameColumnId.ToString()) ? 0 : 1;
+
 	if(bShowTypeInColumnView)
 	{
 		NewColumnView->GetHeaderRow()->AddColumn(
@@ -1502,7 +1515,14 @@ TSharedRef<SAssetColumnView> SAssetView::CreateColumnView()
 				.SortPriority(TAttribute< EColumnSortPriority::Type >::Create(TAttribute< EColumnSortPriority::Type >::FGetter::CreateSP(this, &SAssetView::GetColumnSortPriority, SortManager.ClassColumnId)))
 				.OnSort(FOnSortModeChanged::CreateSP(this, &SAssetView::OnSortColumnHeader))
 				.DefaultLabel(LOCTEXT("Column_Class", "Type"))
+				.ShouldGenerateWidget(TAttribute<bool>::Create(TAttribute<bool>::FGetter::CreateSP(this, &SAssetView::ShouldColumnGenerateWidget, SortManager.ClassColumnId.ToString())))
+				.MenuContent()
+				[
+					CreateRowHeaderMenuContent(SortManager.ClassColumnId.ToString())
+				]
 			);
+
+		NumVisibleColumns += HiddenColumnNames.Contains(SortManager.ClassColumnId.ToString()) ? 0 : 1;
 	}
 
 
@@ -1515,7 +1535,15 @@ TSharedRef<SAssetColumnView> SAssetView::CreateColumnView()
 				.SortPriority(TAttribute< EColumnSortPriority::Type >::Create(TAttribute< EColumnSortPriority::Type >::FGetter::CreateSP(this, &SAssetView::GetColumnSortPriority, SortManager.PathColumnId)))
 				.OnSort(FOnSortModeChanged::CreateSP(this, &SAssetView::OnSortColumnHeader))
 				.DefaultLabel(LOCTEXT("Column_Path", "Path"))
+				.ShouldGenerateWidget(TAttribute<bool>::Create(TAttribute<bool>::FGetter::CreateSP(this, &SAssetView::ShouldColumnGenerateWidget, SortManager.PathColumnId.ToString())))
+				.MenuContent()
+				[
+					CreateRowHeaderMenuContent(SortManager.PathColumnId.ToString())
+				]
 			);
+
+
+		NumVisibleColumns += HiddenColumnNames.Contains(SortManager.PathColumnId.ToString()) ? 0 : 1;
 	}
 
 	return NewColumnView.ToSharedRef();
@@ -2076,15 +2104,21 @@ void SAssetView::SetMajorityAssetType(FName NewMajorityAssetType)
 								}
 
 								ColumnView->GetHeaderRow()->AddColumn(
-										SHeaderRow::Column(TagName)
-										.SortMode( TAttribute< EColumnSortMode::Type >::Create( TAttribute< EColumnSortMode::Type >::FGetter::CreateSP( this, &SAssetView::GetColumnSortMode, TagName ) ) )
-										.SortPriority(TAttribute< EColumnSortPriority::Type >::Create(TAttribute< EColumnSortPriority::Type >::FGetter::CreateSP(this, &SAssetView::GetColumnSortPriority, TagName)))
-										.OnSort( FOnSortModeChanged::CreateSP( this, &SAssetView::OnSortColumnHeader ) )
-										.DefaultLabel( DisplayName )
-										.DefaultTooltip( TooltipText )
-										.HAlignCell( (TagIt->Type == UObject::FAssetRegistryTag::TT_Numerical) ? HAlign_Right : HAlign_Left )
-										.FillWidth(180)
-									);
+									SHeaderRow::Column(TagName)
+									.SortMode(TAttribute< EColumnSortMode::Type >::Create(TAttribute< EColumnSortMode::Type >::FGetter::CreateSP(this, &SAssetView::GetColumnSortMode, TagName)))
+									.SortPriority(TAttribute< EColumnSortPriority::Type >::Create(TAttribute< EColumnSortPriority::Type >::FGetter::CreateSP(this, &SAssetView::GetColumnSortPriority, TagName)))
+									.OnSort(FOnSortModeChanged::CreateSP(this, &SAssetView::OnSortColumnHeader))
+									.DefaultLabel(DisplayName)
+									.DefaultTooltip(TooltipText)
+									.HAlignCell((TagIt->Type == UObject::FAssetRegistryTag::TT_Numerical) ? HAlign_Right : HAlign_Left)
+									.FillWidth(180)
+									.ShouldGenerateWidget(TAttribute<bool>::Create(TAttribute<bool>::FGetter::CreateSP(this, &SAssetView::ShouldColumnGenerateWidget, TagName.ToString())))
+									.MenuContent()
+									[
+										CreateRowHeaderMenuContent(TagName.ToString())
+									]);								
+								
+								NumVisibleColumns += HiddenColumnNames.Contains(TagName.ToString()) ? 0 : 1;
 
 								// If we found a tag the matches the column we are currently sorting on, there will be no need to change the column
 								for (int32 SortIdx = 0; SortIdx < CurrentSortOrder.Num(); SortIdx++)
@@ -2715,6 +2749,28 @@ TSharedRef<SWidget> SAssetView::GetViewButtonContent()
 	}
 	MenuBuilder.EndSection();
 
+	if (GetColumnViewVisibility() == EVisibility::Visible)
+	{
+		MenuBuilder.BeginSection("AssetColumns", LOCTEXT("ToggleColumnsHeading", "Columns"));
+		{
+			MenuBuilder.AddSubMenu(
+				LOCTEXT("ToggleColumnsMenu", "Toggle columns"),
+				LOCTEXT("ToggleColumnsMenuTooltip", "Show or hide specific columns."),
+				FNewMenuDelegate::CreateSP(this, &SAssetView::FillToggleColumnsMenu)
+				);
+
+			MenuBuilder.AddMenuEntry(
+				LOCTEXT("ResetColumns", "Reset Columns"),
+				LOCTEXT("ResetColumnsToolTip", "Reset all columns to be visible again."),
+				FSlateIcon(),
+				FUIAction(FExecuteAction::CreateSP(this, &SAssetView::ResetColumns)),
+				NAME_None,
+				EUserInterfaceActionType::Button
+				);
+		}
+		MenuBuilder.EndSection();
+	}
+
 	return MenuBuilder.MakeWidget();
 }
 
@@ -4279,4 +4335,89 @@ bool SAssetView::PerformQuickJump(const bool bWasJumping)
 	return ValidMatch;
 }
 
+void SAssetView::FillToggleColumnsMenu(FMenuBuilder& MenuBuilder)
+{
+	const TIndirectArray<SHeaderRow::FColumn> Columns = ColumnView->GetHeaderRow()->GetColumns();
+
+	for (int32 ColumnIndex = 0; ColumnIndex < Columns.Num(); ++ColumnIndex)
+	{
+		const FString ColumnName = Columns[ColumnIndex].ColumnId.ToString();
+
+		MenuBuilder.AddMenuEntry(
+			Columns[ColumnIndex].DefaultText,
+			LOCTEXT("ShowHideColumnTooltip", "Show or hide column"),
+			FSlateIcon(),
+			FUIAction(
+				FExecuteAction::CreateSP(this, &SAssetView::ToggleColumn, ColumnName),
+				FCanExecuteAction::CreateSP(this, &SAssetView::CanToggleColumn, ColumnName),
+				FIsActionChecked::CreateSP(this, &SAssetView::IsColumnVisible, ColumnName),
+				EUIActionRepeatMode::RepeatEnabled
+				),
+			NAME_None,
+			EUserInterfaceActionType::Check
+			);
+	}
+}
+
+void SAssetView::ResetColumns()
+{
+	HiddenColumnNames.Empty();
+	NumVisibleColumns = ColumnView->GetHeaderRow()->GetColumns().Num();
+	ColumnView->GetHeaderRow()->RefreshColumns();
+	ColumnView->RebuildList();
+}
+
+void SAssetView::ToggleColumn(const FString ColumnName)
+{
+	SetColumnVisibility(ColumnName, HiddenColumnNames.Contains(ColumnName));
+}
+
+void SAssetView::SetColumnVisibility(const FString ColumnName, const bool bShow)
+{
+	if (!bShow)
+	{
+		--NumVisibleColumns;
+		HiddenColumnNames.Add(ColumnName);
+	}
+	else
+	{
+		++NumVisibleColumns;
+		check(HiddenColumnNames.Contains(ColumnName));
+		HiddenColumnNames.Remove(ColumnName);
+	}
+
+	ColumnView->GetHeaderRow()->RefreshColumns();
+	ColumnView->RebuildList();
+}
+
+bool SAssetView::CanToggleColumn(const FString ColumnName) const
+{
+	return (HiddenColumnNames.Contains(ColumnName) || NumVisibleColumns > 1);
+}
+
+bool SAssetView::IsColumnVisible(const FString ColumnName) const
+{
+	return !HiddenColumnNames.Contains(ColumnName);
+}
+
+bool SAssetView::ShouldColumnGenerateWidget(const FString ColumnName) const
+{
+	return !HiddenColumnNames.Contains(ColumnName);
+}
+
+TSharedRef<SWidget> SAssetView::CreateRowHeaderMenuContent(const FString ColumnName)
+{
+	FMenuBuilder MenuBuilder(true, NULL);
+
+	MenuBuilder.AddMenuEntry(
+		LOCTEXT("HideColumn", "Hide Column"),
+		LOCTEXT("HideColumnToolTip", "Hides this column."),
+		FSlateIcon(),
+		FUIAction(FExecuteAction::CreateSP(this, &SAssetView::SetColumnVisibility, ColumnName, false), FCanExecuteAction::CreateSP(this, &SAssetView::CanToggleColumn, ColumnName)),
+		NAME_None,
+		EUserInterfaceActionType::Button);
+
+	return MenuBuilder.MakeWidget();
+}
+
 #undef LOCTEXT_NAMESPACE
diff --git a/Engine/Source/Editor/DestructibleMeshEditor/Private/DestructibleMeshEditorModule.cpp b/Engine/Source/Editor/DestructibleMeshEditor/Private/DestructibleMeshEditorModule.cpp
--- a/Engine/Source/Editor/DestructibleMeshEditor/Private/DestructibleMeshEditorModule.cpp
+++ b/Engine/Source/Editor/DestructibleMeshEditor/Private/DestructibleMeshEditorModule.cpp
@@ -10,6 +10,7 @@
 #include "ApexDestructibleAssetImport.h"
 #include "Engine/DestructibleMesh.h"
 #include "Engine/StaticMesh.h"
+#include "Materials/MaterialExpressionSpeedTree.h"
 
 IMPLEMENT_MODULE( FDestructibleMeshEditorModule, DestructibleMeshEditor );
 
@@ -47,6 +48,40 @@ UDestructibleMesh* FDestructibleMeshEditorModule::CreateDestructibleMeshFromStat
 		return NULL;
 	}
 
+	// We can't use any speedtree materials for destructibles due to UV requirements, so check the graph here
+	TArray<UMaterial*> SpeedTreeMaterials;
+	for(FStaticMaterial& StaticMaterial : StaticMesh->StaticMaterials)
+	{
+		UMaterialInterface* Mat = StaticMaterial.MaterialInterface;
+
+		UMaterial* BaseMat = Mat->GetBaseMaterial();
+	
+		for(UMaterialExpression* Expression : BaseMat->Expressions)
+		{
+			if(Cast<UMaterialExpressionSpeedTree>(Expression))
+			{
+				SpeedTreeMaterials.Add(BaseMat);
+				break;
+			}
+		}
+	}
+	
+	if(SpeedTreeMaterials.Num() > 0)
+	{
+		// Invalid, can't create a mesh from this
+		FTextBuilder TextBuilder;
+		TextBuilder.AppendLine(FText::Format(LOCTEXT("StaticMeshInvalid_SpeedTree", "The static mesh '{0}' uses SpeedTree materials which are not compatible with destructible meshes. Cannot create destructible.\n\nList of Materials:\n"), FText::FromString(StaticMesh->GetName())));
+	
+		for(UMaterial* SpeedTreeMat : SpeedTreeMaterials)
+		{
+			TextBuilder.AppendLine(SpeedTreeMat->GetName());
+		}
+	
+		FMessageDialog::Open(EAppMsgType::Ok, TextBuilder.ToText());
+		return nullptr;
+	}
+
+
 	FString DestructibleName;
 	
 	if (Name == NAME_None)
diff --git a/Engine/Source/Editor/DetailCustomizations/DetailCustomizations.Build.cs b/Engine/Source/Editor/DetailCustomizations/DetailCustomizations.Build.cs
--- a/Engine/Source/Editor/DetailCustomizations/DetailCustomizations.Build.cs
+++ b/Engine/Source/Editor/DetailCustomizations/DetailCustomizations.Build.cs
@@ -47,8 +47,9 @@ public class DetailCustomizations : ModuleRules
 				"ConfigEditor",
                 "Persona",
                 "CinematicCamera",
-                "ComponentVisualizers"
-			}
+				"ComponentVisualizers",
+                "SkeletonEditor",
+            }
 		);
 
         PrivateIncludePathModuleNames.AddRange(
diff --git a/Engine/Source/Editor/DetailCustomizations/Private/AnimTransitionNodeDetails.cpp b/Engine/Source/Editor/DetailCustomizations/Private/AnimTransitionNodeDetails.cpp
--- a/Engine/Source/Editor/DetailCustomizations/Private/AnimTransitionNodeDetails.cpp
+++ b/Engine/Source/Editor/DetailCustomizations/Private/AnimTransitionNodeDetails.cpp
@@ -14,9 +14,8 @@
 #include "STextEntryPopup.h"
 #include "SExpandableArea.h"
 #include "BlueprintEditorUtils.h"
-#include "SBlendProfilePicker.h"
 #include "Animation/BlendProfile.h"
-
+#include "ISkeletonEditorModule.h"
 
 #define LOCTEXT_NAMESPACE "FAnimStateNodeDetails"
 
@@ -167,6 +166,13 @@ void FAnimTransitionNodeDetails::CustomizeDetails( IDetailLayoutBuilder& DetailB
 			BlendProfileHandle->GetValue(BlendProfilePropertyValue);
 			UBlendProfile* CurrentProfile = Cast<UBlendProfile>(BlendProfilePropertyValue);
 
+			ISkeletonEditorModule& SkeletonEditorModule = FModuleManager::LoadModuleChecked<ISkeletonEditorModule>("SkeletonEditor");
+
+			FBlendProfilePickerArgs Args;
+			Args.InitialProfile = CurrentProfile;
+			Args.OnBlendProfileSelected = FOnBlendProfileSelected::CreateSP(this, &FAnimTransitionNodeDetails::OnBlendProfileChanged, BlendProfileHandle);
+			Args.bAllowNew = false;
+
 			CrossfadeCategory.AddProperty(BlendProfileHandle).CustomWidget(true)
 				.NameContent()
 				[
@@ -174,11 +180,7 @@ void FAnimTransitionNodeDetails::CustomizeDetails( IDetailLayoutBuilder& DetailB
 				]
 				.ValueContent()
 				[
-					SNew(SBlendProfilePicker)
-					.TargetSkeleton(TargetSkeleton)
-					.AllowNew(false)
-					.OnBlendProfileSelected(this, &FAnimTransitionNodeDetails::OnBlendProfileChanged, BlendProfileHandle)
-					.InitialProfile(CurrentProfile)
+					SkeletonEditorModule.CreateBlendProfilePicker(TargetSkeleton, Args)
 				];
 		}
 
diff --git a/Engine/Source/Editor/DetailCustomizations/Private/AssetViewerSettingsCustomization.cpp b/Engine/Source/Editor/DetailCustomizations/Private/AssetViewerSettingsCustomization.cpp
--- a/Engine/Source/Editor/DetailCustomizations/Private/AssetViewerSettingsCustomization.cpp
+++ b/Engine/Source/Editor/DetailCustomizations/Private/AssetViewerSettingsCustomization.cpp
@@ -4,6 +4,9 @@
 #include "AssetViewerSettingsCustomization.h"
 #include "AssetViewerSettings.h"
 
+#include "Widgets/Input/SEditableTextBox.h"
+#include "ScopedTransaction.h"
+
 #define LOCTEXT_NAMESPACE "AssetViewerSettingsCustomizations"
 
 TSharedRef<IDetailCustomization> FAssetViewerSettingsCustomization::MakeInstance()
@@ -13,7 +16,7 @@ TSharedRef<IDetailCustomization> FAssetViewerSettingsCustomization::MakeInstance
 
 void FAssetViewerSettingsCustomization::CustomizeDetails(IDetailLayoutBuilder& DetailBuilder)
 {
-	UAssetViewerSettings* ViewerSettings = UAssetViewerSettings::Get();
+	ViewerSettings = UAssetViewerSettings::Get();
 
 	// Find profiles array property handle and hide it from settings view
 	TSharedPtr<IPropertyHandle> ProfileHandle = DetailBuilder.GetProperty(GET_MEMBER_NAME_CHECKED(UAssetViewerSettings, Profiles));
@@ -24,7 +27,7 @@ void FAssetViewerSettingsCustomization::CustomizeDetails(IDetailLayoutBuilder& D
 	IDetailCategoryBuilder& CategoryBuilder = DetailBuilder.EditCategory("Settings", LOCTEXT("AssetViewerSettingsCategory", "Settings"));
 
 	const UEditorPerProjectUserSettings* PerProjectUserSettings = GetDefault<UEditorPerProjectUserSettings>();
-	const int32 ProfileIndex = ViewerSettings->Profiles.IsValidIndex(PerProjectUserSettings->AssetViewerProfileIndex) ? PerProjectUserSettings->AssetViewerProfileIndex : 0;
+	ProfileIndex = ViewerSettings->Profiles.IsValidIndex(PerProjectUserSettings->AssetViewerProfileIndex) ? PerProjectUserSettings->AssetViewerProfileIndex : 0;
 
 	ensureMsgf(ViewerSettings && ViewerSettings->Profiles.IsValidIndex(PerProjectUserSettings->AssetViewerProfileIndex), TEXT("Invalid default settings pointer or current profile index"));
 
@@ -33,10 +36,84 @@ void FAssetViewerSettingsCustomization::CustomizeDetails(IDetailLayoutBuilder& D
 	checkf(ProfileHandle->IsValidHandle(), TEXT("Invalid indexing into profiles child properties"));	
 	uint32 PropertyCount = 0;
 	ProfilePropertyHandle->GetNumChildren(PropertyCount);
+
+	FName NamePropertyName = GET_MEMBER_NAME_CHECKED(FPreviewSceneProfile, ProfileName);
 	for (uint32 PropertyIndex = 0; PropertyIndex < PropertyCount; ++PropertyIndex)
 	{
-		CategoryBuilder.AddProperty(ProfilePropertyHandle->GetChildHandle(PropertyIndex));		
+		TSharedPtr<IPropertyHandle> ProfileProperty = ProfilePropertyHandle->GetChildHandle(PropertyIndex);
+				
+		if (ProfileProperty->GetProperty()->GetFName() == NamePropertyName)
+		{
+			NameProperty = ProfileProperty;
+			CategoryBuilder.AddCustomRow(LOCTEXT("PreviewSceneProfileDetails_ProfileNameLabel", "Profile Name"))
+					.NameContent()
+					[
+						ProfileProperty->CreatePropertyNameWidget()
+					]
+					.ValueContent()
+					.MaxDesiredWidth(250.0f)
+					[
+						// PropertyHandle->CreatePropertyValueWidget()
+						SAssignNew(NameEditTextBox, SEditableTextBox)
+						.Text(this, &FAssetViewerSettingsCustomization::OnGetProfileName)
+						//.ToolTip(VarNameTooltip)
+						.OnTextChanged(this, &FAssetViewerSettingsCustomization::OnProfileNameChanged)
+						.OnTextCommitted(this, &FAssetViewerSettingsCustomization::OnProfileNameCommitted)
+						.Font(IDetailLayoutBuilder::GetDetailFont())
+					];
+		}
+		else
+		{
+			IDetailPropertyRow& Row = CategoryBuilder.AddProperty(ProfileProperty);
+		}
+	}
+}
+
+FText FAssetViewerSettingsCustomization::OnGetProfileName() const
+{
+	return FText::FromString(ViewerSettings->Profiles[ProfileIndex].ProfileName);
+}
+
+void FAssetViewerSettingsCustomization::OnProfileNameChanged(const FText& InNewText)
+{	
+	bValidProfileName = IsProfileNameValid(InNewText.ToString());
+	if (!bValidProfileName)
+	{
+		NameEditTextBox->SetError(LOCTEXT("PreviewSceneProfileName_NotValid", "This name is already in use"));
+	}
+	else
+	{
+		NameEditTextBox->SetError(FText::GetEmpty());
+	}
+}
+
+void FAssetViewerSettingsCustomization::OnProfileNameCommitted(const FText& InNewText, ETextCommit::Type InTextCommit)
+{
+	if (bValidProfileName && InTextCommit != ETextCommit::OnCleared)
+	{
+		const FScopedTransaction Transaction(LOCTEXT("RenameProfile", "Rename Profile"));
+		ViewerSettings->Modify();
+
+		ViewerSettings->Profiles[ProfileIndex].ProfileName = InNewText.ToString();		
+		FPropertyChangedEvent PropertyEvent(NameProperty->GetProperty());
+		ViewerSettings->PostEditChangeProperty(PropertyEvent);
+	}
+
+	bValidProfileName = false;
+	NameEditTextBox->SetError(FText::GetEmpty());
+}
+
+const bool FAssetViewerSettingsCustomization::IsProfileNameValid(const FString& NewName)
+{
+	for (int32 CheckProfileIndex = 0; CheckProfileIndex < ViewerSettings->Profiles.Num(); ++CheckProfileIndex)
+	{
+		if (ProfileIndex != CheckProfileIndex && ViewerSettings->Profiles[CheckProfileIndex].ProfileName == NewName)
+		{
+			return false;
+		}
 	}
+
+	return true;
 }
 
 #undef LOCTEXT_NAMESPACE 
diff --git a/Engine/Source/Editor/DetailCustomizations/Private/AssetViewerSettingsCustomization.h b/Engine/Source/Editor/DetailCustomizations/Private/AssetViewerSettingsCustomization.h
--- a/Engine/Source/Editor/DetailCustomizations/Private/AssetViewerSettingsCustomization.h
+++ b/Engine/Source/Editor/DetailCustomizations/Private/AssetViewerSettingsCustomization.h
@@ -2,6 +2,8 @@
 
 #pragma once
 
+class UAssetViewerSettings;
+
 class FAssetViewerSettingsCustomization : public IDetailCustomization
 {
 public:
@@ -11,4 +13,20 @@ public:
 	// IDetailCustomization interface
 	virtual void CustomizeDetails(IDetailLayoutBuilder& DetailBuilder) override;
 	// End of IDetailCustomization interface
+protected:
+	// Callbacks for customized SEditableTextBox
+	FText OnGetProfileName() const;
+	void OnProfileNameChanged(const FText& InNewText);
+	void OnProfileNameCommitted(const FText& InNewText, ETextCommit::Type InTextCommit);
+	
+	// Check whether or not the given profile name is valid by cross-referencing it with existing names
+	const bool IsProfileNameValid(const FString& NewName);
+private:
+	// Customized name edit text box used for the profile name
+	TSharedPtr<SEditableTextBox> NameEditTextBox;
+	// Cached data
+	TSharedPtr<IPropertyHandle> NameProperty;
+	int32 ProfileIndex;
+	UAssetViewerSettings* ViewerSettings;
+	bool bValidProfileName;
 };
\ No newline at end of file
diff --git a/Engine/Source/Editor/DetailCustomizations/Private/DestructibleMeshDetails.cpp b/Engine/Source/Editor/DetailCustomizations/Private/DestructibleMeshDetails.cpp
--- a/Engine/Source/Editor/DetailCustomizations/Private/DestructibleMeshDetails.cpp
+++ b/Engine/Source/Editor/DetailCustomizations/Private/DestructibleMeshDetails.cpp
@@ -42,13 +42,6 @@ void AddStructToDetails(FName CategoryName, FName PropertyName, IDetailLayoutBui
 
 void FDestructibleMeshDetails::CustomizeDetails(IDetailLayoutBuilder& DetailBuilder)
 {
-		//we always hide bodysetup as it's not useful in this editor
-		TSharedPtr<IPropertyHandle> BodySetupHandler = DetailBuilder.GetProperty(GET_MEMBER_NAME_CHECKED(UDestructibleMesh, BodySetup));
-		if (BodySetupHandler.IsValid())
-		{
-			DetailBuilder.HideProperty(BodySetupHandler);
-		}
-		
 		//rest of customization is just moving stuff out of DefaultDestructibleParameters so it's nicer to view
 		TSharedPtr<IPropertyHandle> DefaultParams = DetailBuilder.GetProperty(GET_MEMBER_NAME_CHECKED(UDestructibleMesh, DefaultDestructibleParameters));
 		if (DefaultParams.IsValid() == false)
@@ -63,11 +56,36 @@ void FDestructibleMeshDetails::CustomizeDetails(IDetailLayoutBuilder& DetailBuil
 		AddStructToDetails("HierarchyDepth", "DefaultDestructibleParameters.SpecialHierarchyDepths", DetailBuilder);
 		AddStructToDetails("HierarchyDepth", "DefaultDestructibleParameters.DepthParameters", DetailBuilder, false, true);
 
-		
+		// There are some properties we inherit that aren't supported on destructibles, hide them from the details panel.
+		HideUnsupportedProperties(DetailBuilder);
 
 		//hide the default params as we've taken everything out of it
 		DetailBuilder.HideProperty(DefaultParams);
 }
 
+void FDestructibleMeshDetails::HideUnsupportedProperties(IDetailLayoutBuilder &DetailBuilder)
+{
+	// Body setups are not available on destructible meshes as we set up the bodies through APEX
+	TSharedPtr<IPropertyHandle> BodySetupHandle = DetailBuilder.GetProperty(GET_MEMBER_NAME_CHECKED(UDestructibleMesh, BodySetup));
+	if(BodySetupHandle.IsValid())
+	{
+		BodySetupHandle->MarkHiddenByCustomization();
+	}
+
+	// Capsule shadows only supported on skeletal meshes
+	TSharedPtr<IPropertyHandle> ShadowPhysicsAssetHandle = DetailBuilder.GetProperty(GET_MEMBER_NAME_CHECKED(UDestructibleMesh, ShadowPhysicsAsset), USkeletalMesh::StaticClass());
+	if(ShadowPhysicsAssetHandle.IsValid())
+	{
+		ShadowPhysicsAssetHandle->MarkHiddenByCustomization();
+	}
+
+	// Post processing graphs only supported on skeletal meshes
+	TSharedPtr<IPropertyHandle> PostProcessBlueprintHandle = DetailBuilder.GetProperty(GET_MEMBER_NAME_CHECKED(UDestructibleMesh, PostProcessAnimBlueprint), USkeletalMesh::StaticClass());
+	if(PostProcessBlueprintHandle.IsValid())
+	{
+		PostProcessBlueprintHandle->MarkHiddenByCustomization();
+	}
+}
+
 #undef LOCTEXT_NAMESPACE
 
diff --git a/Engine/Source/Editor/DetailCustomizations/Private/DestructibleMeshDetails.h b/Engine/Source/Editor/DetailCustomizations/Private/DestructibleMeshDetails.h
--- a/Engine/Source/Editor/DetailCustomizations/Private/DestructibleMeshDetails.h
+++ b/Engine/Source/Editor/DetailCustomizations/Private/DestructibleMeshDetails.h
@@ -12,6 +12,10 @@ public:
 	virtual void CustomizeDetails( IDetailLayoutBuilder& DetailBuilder ) override;
 
 private:
+
+	/** Hide any properties that aren't supported on destructible meshes */
+	void HideUnsupportedProperties(IDetailLayoutBuilder &DetailBuilder);
+
 	TArray< TWeakObjectPtr<UObject> > ObjectsCustomized;
 };
 
diff --git a/Engine/Source/Editor/EditorStyle/Private/SlateEditorStyle.cpp b/Engine/Source/Editor/EditorStyle/Private/SlateEditorStyle.cpp
--- a/Engine/Source/Editor/EditorStyle/Private/SlateEditorStyle.cpp
+++ b/Engine/Source/Editor/EditorStyle/Private/SlateEditorStyle.cpp
@@ -4441,6 +4441,15 @@ void FSlateEditorStyle::FStyle::SetupGraphEditorStyles()
 	{
 		Set( "Persona.PipelineSeparator", new BOX_BRUSH( "Persona/Modes/PipelineSeparator", FMargin(15.0f/16.0f, 22.0f/24.0f, 1.0f/16.0f, 1.0f/24.0f), FLinearColor(1,1,1,0.5f) ) );
 	}
+
+	// montage editor
+	{
+		Set("Persona.MontageEditor.ChildMontageInstruction", FTextBlockStyle(NormalText)
+			.SetFont(TTF_FONT("Fonts/Roboto-BoldCondensed", 14))
+			.SetColorAndOpacity(FLinearColor(1.f, 1.f, 1.f, 1.0f))
+			.SetShadowOffset(FVector2D::ZeroVector)
+		);
+	}
 #endif // WITH_EDITOR || IS_PROGRAM
 	}
 
@@ -5131,6 +5140,7 @@ void FSlateEditorStyle::FStyle::SetupPersonaStyle()
 
 		Set("AnimViewportMenu.PlayBackSpeed", new IMAGE_BRUSH("Persona/Viewport/icon_Playback_speed_16x", Icon16x16));
 		Set("AnimViewportMenu.TurnTableSpeed", new IMAGE_BRUSH("Persona/Viewport/icon_turn_table_16x", Icon16x16));
+		Set("AnimViewportMenu.SceneSetup", new IMAGE_BRUSH("Icons/icon_tab_SceneOutliner_16x", Icon16x16));
 
 		Set( "AnimViewport.MessageFont", TTF_CORE_FONT("Fonts/Roboto-Bold", 9) );
 		
@@ -5184,9 +5194,8 @@ void FSlateEditorStyle::FStyle::SetupPersonaStyle()
 		// persona extras
 		Set("Persona.ConvertAnimationGraph", new IMAGE_BRUSH("Old/Graph/ConvertIcon", Icon40x40));
 		Set("Persona.ReimportAsset", new IMAGE_BRUSH("Icons/Reimport_12x", Icon12x12));
-		Set("SkeletonTree.SkeletonSocket", new IMAGE_BRUSH("Persona/SkeletonTree/icon_SocketG_16px", Icon16x16));
-		Set("SkeletonTree.MeshSocket", new IMAGE_BRUSH("Persona/SkeletonTree/icon_SocketC_16px", Icon16x16));
-		Set("SkeletonTree.LODBone", new IMAGE_BRUSH(TEXT("Persona/SkeletonTree/icon_LODBone_16x"), Icon16x16));
+		Set("Persona.ConvertToStaticMesh", new IMAGE_BRUSH("Icons/icon_ShowStaticMeshes_40x", Icon40x40));
+		Set("Persona.ConvertToStaticMesh.Small", new IMAGE_BRUSH("Icons/icon_ShowStaticMeshes_40x", Icon20x20));
 
 		// Anim Slot Manager
 		Set("AnimSlotManager.SaveSkeleton", new IMAGE_BRUSH("Persona/AnimSlotManager/icon_SaveSkeleton_40x", Icon40x40));
@@ -5198,7 +5207,7 @@ void FSlateEditorStyle::FStyle::SetupPersonaStyle()
 		Set("AnimNotifyEditor.BranchingPoint", new IMAGE_BRUSH("Persona/NotifyEditor/BranchingPoints_24x", Icon24x24));
 
 		// AnimBlueprint Preview Warning Background
-		FSlateColor PreviewPropertiesWarningColour(FLinearColor(0.1f, 0.2f, 0.5f));
+		FSlateColor PreviewPropertiesWarningColour(FLinearColor::Gray);
 		Set("Persona.PreviewPropertiesWarning", new BOX_BRUSH("Common/RoundedSelection_16x", 4.0f / 16.0f, PreviewPropertiesWarningColour));
 
 		// Persona-specific tabs
@@ -5212,6 +5221,62 @@ void FSlateEditorStyle::FStyle::SetupPersonaStyle()
 		Set("Persona.Tabs.AnimAssetDetails", new IMAGE_BRUSH(TEXT("Icons/icon_Persona_Anim_Asset_Details_16x"), Icon16x16));
 	}
 
+	// Skeleton editor
+	{
+		Set("SkeletonEditor.AnimNotifyWindow", new IMAGE_BRUSH(TEXT("Icons/icon_Persona_AnimNotift_40x"), Icon40x40));
+		Set("SkeletonEditor.AnimNotifyWindow.Small", new IMAGE_BRUSH(TEXT("Icons/icon_Persona_AnimNotift_40x"), Icon20x20));
+		Set("SkeletonEditor.RetargetManager", new IMAGE_BRUSH(TEXT("Icons/icon_Persona_Retarget_40x"), Icon40x40));
+		Set("SkeletonEditor.RetargetManager.Small", new IMAGE_BRUSH(TEXT("Icons/icon_Persona_Retarget_40x"), Icon20x20));
+		Set("SkeletonEditor.ImportMesh", new IMAGE_BRUSH(TEXT("Icons/icon_Persona_ImportMesh_40x"), Icon40x40));
+		Set("SkeletonEditor.ImportMesh.Small", new IMAGE_BRUSH(TEXT("Icons/icon_Persona_ImportMesh_40x"), Icon20x20));
+
+		// Skeleton Tree
+		Set("SkeletonTree.SkeletonSocket", new IMAGE_BRUSH("Persona/SkeletonTree/icon_SocketG_16px", Icon16x16));
+		Set("SkeletonTree.MeshSocket", new IMAGE_BRUSH("Persona/SkeletonTree/icon_SocketC_16px", Icon16x16));
+		Set("SkeletonTree.LODBone", new IMAGE_BRUSH(TEXT("Persona/SkeletonTree/icon_LODBone_16x"), Icon16x16));
+		Set("SkeletonTree.NormalFont", FTextBlockStyle(NormalText)
+			.SetColorAndOpacity(FLinearColor(1.0f, 1.0f, 1.0f)));
+		Set("SkeletonTree.BoldFont", FTextBlockStyle(NormalText)
+			.SetColorAndOpacity(FLinearColor(1.0f, 1.0f, 1.0f))
+			.SetFont(TTF_CORE_FONT("Fonts/Roboto-Bold", 10)));
+
+		Set("SkeletonTree.HyperlinkSpinBox", FSpinBoxStyle(GetWidgetStyle<FSpinBoxStyle>("SpinBox"))
+			.SetTextPadding(FMargin(0))
+			.SetBackgroundBrush(BORDER_BRUSH("Old/HyperlinkDotted", FMargin(0, 0, 0, 3 / 16.0f), FSlateColor::UseSubduedForeground()))
+			.SetHoveredBackgroundBrush(FSlateNoResource())
+			.SetInactiveFillBrush(FSlateNoResource())
+			.SetActiveFillBrush(FSlateNoResource())
+			.SetForegroundColor(FSlateColor::UseSubduedForeground())
+			.SetArrowsImage(FSlateNoResource())
+			);
+
+		Set("SkeletonTree.BlendProfile", new IMAGE_BRUSH(TEXT("Icons/icon_Persona_NewBlendSpace_16x"), Icon16x16));
+		Set("SkeletonTree.InlineEditorShadowTop", new IMAGE_BRUSH(TEXT("Common/ScrollBoxShadowTop"), FVector2D(64, 8)));
+		Set("SkeletonTree.InlineEditorShadowBottom", new IMAGE_BRUSH(TEXT("Common/ScrollBoxShadowBottom"), FVector2D(64, 8)));
+	}
+
+	// Animation editor
+	{
+		Set("AnimationEditor.ApplyCompression", new IMAGE_BRUSH(TEXT("Icons/icon_Persona_Compression_40x"), Icon40x40));
+		Set("AnimationEditor.ApplyCompression.Small", new IMAGE_BRUSH(TEXT("Icons/icon_Persona_Compression_40x"), Icon20x20));
+		Set("AnimationEditor.ExportToFBX", new IMAGE_BRUSH(TEXT("Icons/icon_Persona_ExportFBX_40x"), Icon40x40));
+		Set("AnimationEditor.ExportToFBX.Small", new IMAGE_BRUSH(TEXT("Icons/icon_Persona_ExportFBX_40x"), Icon20x20));
+		Set("AnimationEditor.CreateAsset", new IMAGE_BRUSH(TEXT("Icons/icon_Persona_CreateAsset_40x"), Icon40x40));
+		Set("AnimationEditor.CreateAsset.Small", new IMAGE_BRUSH(TEXT("Icons/icon_Persona_CreateAsset_40x"), Icon20x20));
+		Set("AnimationEditor.SetKey", new IMAGE_BRUSH(TEXT("Icons/icon_Persona_SetKey_40x"), Icon40x40));
+		Set("AnimationEditor.SetKey.Small", new IMAGE_BRUSH(TEXT("Icons/icon_Persona_SetKey_40x"), Icon20x20));
+		Set("AnimationEditor.ApplyAnimation", new IMAGE_BRUSH(TEXT("Icons/icon_Persona_BakeAnim_40x"), Icon40x40));
+		Set("AnimationEditor.ApplyAnimation.Small", new IMAGE_BRUSH(TEXT("Icons/icon_Persona_BakeAnim_40x"), Icon20x20));
+	}
+
+	// Skeletal mesh editor
+	{
+		Set("SkeletalMeshEditor.ReimportMesh", new IMAGE_BRUSH(TEXT("Icons/icon_Persona_ReimportMesh_40x"), Icon40x40));
+		Set("SkeletalMeshEditor.ReimportMesh.Small", new IMAGE_BRUSH(TEXT("Icons/icon_Persona_ReimportMesh_40x"), Icon20x20));
+		Set("SkeletalMeshEditor.ImportLODs", new IMAGE_BRUSH(TEXT("Icons/icon_Persona_ImportLODs_40x"), Icon40x40));
+		Set("SkeletalMeshEditor.ImportLODs.Small", new IMAGE_BRUSH(TEXT("Icons/icon_Persona_ImportLODs_40x"), Icon20x20));
+	}
+
 	// Kismet 2
 	{
 		Set( "FullBlueprintEditor.SwitchToScriptingMode", new IMAGE_BRUSH( "Icons/icon_BlueprintEditor_EventGraph_40x", Icon40x40 ) );
@@ -5349,6 +5414,14 @@ void FSlateEditorStyle::FStyle::SetupPersonaStyle()
 		Set( "SCS.NativeComponent", new IMAGE_BRUSH( "Icons/NativeSCSComponent", Icon20x20 ));
 		Set( "SCS.Component", new IMAGE_BRUSH( "Icons/SCSComponent", Icon20x20 ));
 
+		// curve viewer
+		Set("AnimCurveViewer.MorphTargetOn", new IMAGE_BRUSH(TEXT("Persona/AnimCurveViewer/MorphTarget_On"), Icon16x16));
+		Set("AnimCurveViewer.MaterialOn", new IMAGE_BRUSH(TEXT("Persona/AnimCurveViewer/Material_On"), Icon16x16));
+		Set("AnimCurveViewer.MorphTargetOff", new IMAGE_BRUSH(TEXT("Persona/AnimCurveViewer/MorphTarget_Off"), Icon16x16));
+		Set("AnimCurveViewer.MaterialOff", new IMAGE_BRUSH(TEXT("Persona/AnimCurveViewer/Material_Off"), Icon16x16));
+		Set("AnimCurveViewer.MorphTargetHover", new IMAGE_BRUSH(TEXT("Persona/AnimCurveViewer/MorphTarget_On"), Icon16x16));
+		Set("AnimCurveViewer.MaterialHover", new IMAGE_BRUSH(TEXT("Persona/AnimCurveViewer/Material_On"), Icon16x16));
+
 		const FButtonStyle BlueprintContextTargetsButtonStyle = FButtonStyle()
 			.SetNormal(IMAGE_BRUSH("Common/TreeArrow_Collapsed_Hovered", Icon10x10, FLinearColor(0.2f, 0.2f, 0.2f, 1.f)))
 			.SetHovered(IMAGE_BRUSH("Common/TreeArrow_Collapsed_Hovered", Icon10x10, FLinearColor(0.4f, 0.4f, 0.4f, 1.f)))
@@ -5673,8 +5746,8 @@ void FSlateEditorStyle::FStyle::SetupClassIconsAndThumbnails()
 			TEXT("FileMediaSource"),
 			TEXT("Font"),
 			TEXT("ForceFeedbackEffect"),
-			TEXT("GameMode"),
-			TEXT("GameState"),
+			TEXT("GameModeBase"),
+			TEXT("GameStateBase"),
 			TEXT("HUD"),
 			TEXT("Interface"),
 			TEXT("InterpData"),
diff --git a/Engine/Source/Editor/EditorWidgets/Private/STransportControl.cpp b/Engine/Source/Editor/EditorWidgets/Private/STransportControl.cpp
--- a/Engine/Source/Editor/EditorWidgets/Private/STransportControl.cpp
+++ b/Engine/Source/Editor/EditorWidgets/Private/STransportControl.cpp
@@ -36,12 +36,16 @@ TSharedPtr<SWidget> STransportControl::MakeTransportControlWidget(ETransportCont
 				. Image( this, &STransportControl::GetBackwardStatusIcon )
 			];
 	case ETransportControlWidgetType::Record:
-		return SNew(SButton)
-			. ButtonStyle(&FEditorStyle::Get().GetWidgetStyle<FButtonStyle>("Animation.Record"))
+		return SAssignNew(RecordButton, SButton)
+			. ButtonStyle(FEditorStyle::Get(), "NoBorder")
 			. OnClicked(TransportControlArgs.OnRecord)
 			. Visibility(TransportControlArgs.OnRecord.IsBound() ? EVisibility::Visible : EVisibility::Collapsed)
 			. ToolTipText( this, &STransportControl::GetRecordStatusTooltip )
-			. ContentPadding(2.0f);
+			. ContentPadding(2.0f)
+			[
+				SNew(SImage)
+				.Image(this, &STransportControl::GetRecordStatusIcon )
+			];
 	case ETransportControlWidgetType::ForwardPlay:
 		return SAssignNew(ForwardPlayButton, SButton)
 			. OnClicked(TransportControlArgs.OnForwardPlay)
@@ -145,7 +149,7 @@ bool STransportControl::IsTickable() const
 void STransportControl::Tick( float DeltaTime )
 {
 	const auto PlaybackMode = TransportControlArgs.OnGetPlaybackMode.Execute();
-	const bool bIsPlaying = PlaybackMode == EPlaybackMode::PlayingForward || PlaybackMode == EPlaybackMode::PlayingReverse || PlaybackMode == EPlaybackMode::Recording;
+	const bool bIsPlaying = PlaybackMode == EPlaybackMode::PlayingForward || PlaybackMode == EPlaybackMode::PlayingReverse;
 
 	if ( bIsPlaying && !ActiveTimerHandle.IsValid() )
 	{
@@ -159,9 +163,19 @@ void STransportControl::Tick( float DeltaTime )
 
 const FSlateBrush* STransportControl::GetForwardStatusIcon() const
 {
-	const auto PlaybackMode = TransportControlArgs.OnGetPlaybackMode.Execute();
-	if (TransportControlArgs.OnGetPlaybackMode.IsBound() &&
-		 ( PlaybackMode == EPlaybackMode::PlayingForward || PlaybackMode == EPlaybackMode::Recording ) )
+	EPlaybackMode::Type PlaybackMode = EPlaybackMode::Stopped;
+	if (TransportControlArgs.OnGetPlaybackMode.IsBound())
+	{
+		PlaybackMode = TransportControlArgs.OnGetPlaybackMode.Execute();
+	}
+
+	bool bIsRecording = false;
+	if (TransportControlArgs.OnGetRecording.IsBound())
+	{
+		bIsRecording = TransportControlArgs.OnGetRecording.Execute();
+	}
+
+	if ( PlaybackMode == EPlaybackMode::PlayingForward || bIsRecording )
 	{
 		return ForwardPlayButton.IsValid() && ForwardPlayButton->IsPressed() ? 
 			&FEditorStyle::Get().GetWidgetStyle<FButtonStyle>("Animation.Pause").Pressed : 
@@ -184,10 +198,30 @@ FText STransportControl::GetForwardStatusTooltip() const
 	return LOCTEXT("Play", "Play");
 }
 
+const FSlateBrush* STransportControl::GetRecordStatusIcon() const
+{
+	bool bIsRecording = false;
+	if (TransportControlArgs.OnGetRecording.IsBound())
+	{
+		bIsRecording = TransportControlArgs.OnGetRecording.Execute();
+	}
+
+	if (bIsRecording)
+	{
+		return RecordButton.IsValid() && RecordButton->IsPressed() ?
+			&FEditorStyle::Get().GetWidgetStyle<FButtonStyle>("Animation.Recording").Pressed :
+			&FEditorStyle::Get().GetWidgetStyle<FButtonStyle>("Animation.Recording").Normal;
+	}
+
+	return RecordButton.IsValid() && RecordButton->IsPressed() ?
+		&FEditorStyle::Get().GetWidgetStyle<FButtonStyle>("Animation.Record").Pressed :
+		&FEditorStyle::Get().GetWidgetStyle<FButtonStyle>("Animation.Record").Normal;
+}
+
 FText STransportControl::GetRecordStatusTooltip() const
 {
-	if (TransportControlArgs.OnGetPlaybackMode.IsBound() &&
-		TransportControlArgs.OnGetPlaybackMode.Execute() == EPlaybackMode::Recording)
+	if (TransportControlArgs.OnGetRecording.IsBound() &&
+		TransportControlArgs.OnGetRecording.Execute())
 	{
 		return LOCTEXT("StopRecording", "Stop Recording");
 	}
diff --git a/Engine/Source/Editor/EditorWidgets/Private/STransportControl.h b/Engine/Source/Editor/EditorWidgets/Private/STransportControl.h
--- a/Engine/Source/Editor/EditorWidgets/Private/STransportControl.h
+++ b/Engine/Source/Editor/EditorWidgets/Private/STransportControl.h
@@ -31,6 +31,7 @@ public:
 private:
 	const FSlateBrush* GetForwardStatusIcon() const;
 	FText GetForwardStatusTooltip() const;
+	const FSlateBrush* GetRecordStatusIcon() const;
 	FText GetRecordStatusTooltip() const;
 	const FSlateBrush* GetBackwardStatusIcon() const;
 	const FSlateBrush* GetLoopStatusIcon() const;
@@ -54,4 +55,5 @@ private:
 
 	TSharedPtr<SButton> ForwardPlayButton;
 	TSharedPtr<SButton> BackwardPlayButton;
+	TSharedPtr<SButton> RecordButton;
 };
diff --git a/Engine/Source/Editor/EditorWidgets/Public/ITransportControl.h b/Engine/Source/Editor/EditorWidgets/Public/ITransportControl.h
--- a/Engine/Source/Editor/EditorWidgets/Public/ITransportControl.h
+++ b/Engine/Source/Editor/EditorWidgets/Public/ITransportControl.h
@@ -9,11 +9,11 @@ namespace EPlaybackMode
 		Stopped,
 		PlayingForward,
 		PlayingReverse,
-		Recording,
 	};
 }
 
 DECLARE_DELEGATE_RetVal(bool, FOnGetLooping)
+DECLARE_DELEGATE_RetVal(bool, FOnGetRecording)
 DECLARE_DELEGATE_RetVal(EPlaybackMode::Type, FOnGetPlaybackMode)
 DECLARE_DELEGATE_TwoParams(FOnTickPlayback, double /*InCurrentTime*/, float /*InDeltaTime*/)
 DECLARE_DELEGATE_RetVal(TSharedRef<SWidget>, FOnMakeTransportWidget)
@@ -66,6 +66,7 @@ struct FTransportControlArgs
 		, OnGetLooping()
 		, OnGetPlaybackMode()
 		, OnTickPlayback()
+		, OnGetRecording()
 	{}
 
 	FOnClicked OnForwardPlay;
@@ -79,6 +80,7 @@ struct FTransportControlArgs
 	FOnGetLooping OnGetLooping;
 	FOnGetPlaybackMode OnGetPlaybackMode;
 	FOnTickPlayback OnTickPlayback;
+	FOnGetRecording OnGetRecording;
 
 	/** 
 	 * Array of custom widgets to create - if this array is used the default widget ordering will be
diff --git a/Engine/Source/Editor/GameProjectGeneration/Private/GameProjectUtils.cpp b/Engine/Source/Editor/GameProjectGeneration/Private/GameProjectUtils.cpp
--- a/Engine/Source/Editor/GameProjectGeneration/Private/GameProjectUtils.cpp
+++ b/Engine/Source/Editor/GameProjectGeneration/Private/GameProjectUtils.cpp
@@ -23,7 +23,7 @@
 #include "DesktopPlatformModule.h"
 #include "SNotificationList.h"
 #include "NotificationManager.h"
-#include "GameFramework/GameMode.h"
+#include "GameFramework/GameModeBase.h"
 #include "HotReloadInterface.h"
 #include "SVerbChoiceDialog.h"
 #include "SourceCodeNavigation.h"
@@ -31,6 +31,9 @@
 
 #include "SOutputLogDialog.h"
 
+#include "Sound/SoundEffectSubmix.h"
+#include "Sound/SoundEffectSource.h"
+
 #include "PlatformInfo.h"
 
 #define LOCTEXT_NAMESPACE "GameProjectUtils"
@@ -329,6 +332,19 @@ FString FNewClassInfo::GetHeaderTemplateFilename() const
 				{
 					return TEXT("CharacterClass.h.template");
 				}
+
+				// Only check audio-mixer module specific classes if audio mixer is loaded
+				if (FModuleManager::Get().IsModuleLoaded("AudioMixer"))
+				{
+					if (BaseClass == USoundEffectSource::StaticClass())
+					{
+						return TEXT("SoundEffectSourceClass.h.template");
+					}
+					else if (BaseClass == USoundEffectSubmix::StaticClass())
+					{
+						return TEXT("SoundEffectSubmixClass.h.template");
+					}
+				}
 			}
 			// Some other non-actor, non-component UObject class
 			return TEXT( "UObjectClass.h.template" );
@@ -375,6 +391,14 @@ FString FNewClassInfo::GetSourceTemplateFilename() const
 				{
 					return TEXT("CharacterClass.cpp.template");
 				}
+				else if (BaseClass == USoundEffectSubmix::StaticClass())
+				{
+					return TEXT("SoundEffectSubmixClass.cpp.template");
+				}
+				else if (BaseClass == USoundEffectSource::StaticClass())
+				{
+					return TEXT("SoundEffectSourceClass.cpp.template");
+				}
 			}
 			// Some other non-actor, non-component UObject class
 			return TEXT( "UObjectClass.cpp.template" );
@@ -1955,7 +1979,7 @@ bool GameProjectUtils::GenerateGameFrameworkSourceCode(const FString& NewProject
 
 	// MyGameGameMode.h
 	{
-		const UClass* BaseClass = AGameMode::StaticClass();
+		const UClass* BaseClass = AGameModeBase::StaticClass();
 		const FString NewClassName = NewProjectName + BaseClass->GetName();
 		const FString NewHeaderFilename = GameModulePath / NewClassName + TEXT(".h");
 		FString UnusedSyncLocation;
@@ -1971,7 +1995,7 @@ bool GameProjectUtils::GenerateGameFrameworkSourceCode(const FString& NewProject
 
 	// MyGameGameMode.cpp
 	{
-		const UClass* BaseClass = AGameMode::StaticClass();
+		const UClass* BaseClass = AGameModeBase::StaticClass();
 		const FString NewClassName = NewProjectName + BaseClass->GetName();
 		const FString NewCPPFilename = GameModulePath / NewClassName + TEXT(".cpp");
 		
diff --git a/Engine/Source/Editor/GameProjectGeneration/Private/SNewClassDialog.cpp b/Engine/Source/Editor/GameProjectGeneration/Private/SNewClassDialog.cpp
--- a/Engine/Source/Editor/GameProjectGeneration/Private/SNewClassDialog.cpp
+++ b/Engine/Source/Editor/GameProjectGeneration/Private/SNewClassDialog.cpp
@@ -13,8 +13,8 @@
 #include "SWizard.h"
 #include "SHyperlink.h"
 #include "TutorialMetaData.h"
-#include "GameFramework/GameMode.h"
-#include "GameFramework/GameState.h"
+#include "GameFramework/GameModeBase.h"
+#include "GameFramework/GameStateBase.h"
 #include "GameFramework/Pawn.h"
 #include "GameFramework/WorldSettings.h"
 #include "GameFramework/PlayerController.h"
diff --git a/Engine/Source/Editor/GameProjectGeneration/Public/FeaturedClasses.inl b/Engine/Source/Editor/GameProjectGeneration/Public/FeaturedClasses.inl
--- a/Engine/Source/Editor/GameProjectGeneration/Public/FeaturedClasses.inl
+++ b/Engine/Source/Editor/GameProjectGeneration/Public/FeaturedClasses.inl
@@ -9,11 +9,11 @@
 #include "GameFramework/Pawn.h"
 #include "GameFramework/Actor.h"
 #include "GameFramework/PlayerController.h"
-#include "GameFramework/GameMode.h"
+#include "GameFramework/GameModeBase.h"
 #include "GameFramework/WorldSettings.h"
 #include "GameFramework/HUD.h"
 #include "GameFramework/PlayerState.h"
-#include "GameFramework/GameState.h"
+#include "GameFramework/GameStateBase.h"
 #include "Components/ActorComponent.h"
 #include "Components/SceneComponent.h"
 #include "Components/StaticMeshComponent.h"
@@ -74,11 +74,11 @@ FORCEINLINE void FFeaturedClasses::AddExtraActorClasses(TArray<FNewClassInfo>& A
 {
 	Array.Add(FNewClassInfo(APlayerCameraManager::StaticClass()));
 	Array.Add(FNewClassInfo(APlayerController::StaticClass()));
-	Array.Add(FNewClassInfo(AGameMode::StaticClass()));
+	Array.Add(FNewClassInfo(AGameModeBase::StaticClass()));
 	Array.Add(FNewClassInfo(AWorldSettings::StaticClass()));
 	Array.Add(FNewClassInfo(AHUD::StaticClass()));
 	Array.Add(FNewClassInfo(APlayerState::StaticClass()));
-	Array.Add(FNewClassInfo(AGameState::StaticClass()));
+	Array.Add(FNewClassInfo(AGameStateBase::StaticClass()));
 }
 
 /** Append the featured Component class types that are commonly used */
diff --git a/Engine/Source/Editor/Kismet/Private/BlueprintDetailsCustomization.cpp b/Engine/Source/Editor/Kismet/Private/BlueprintDetailsCustomization.cpp
--- a/Engine/Source/Editor/Kismet/Private/BlueprintDetailsCustomization.cpp
+++ b/Engine/Source/Editor/Kismet/Private/BlueprintDetailsCustomization.cpp
@@ -5535,6 +5535,43 @@ void FChildActorComponentDetails::CustomizeDetails(IDetailLayoutBuilder& DetailB
 				}
 			}
 		}
+
+
+		auto TemplatesVisible = [](const TArray<TWeakObjectPtr<UObject>>& WeakObjects)
+		{
+			bool bResult = true;
+			
+			return bResult;
+		};
+
+		TArray<TWeakObjectPtr<UObject>> ObjectsBeingCustomized;
+		DetailBuilder.GetObjectsBeingCustomized(ObjectsBeingCustomized);
+
+		IDetailCategoryBuilder& CategoryBuilder = DetailBuilder.EditCategory(TEXT("ChildActorComponent"));
+				
+		// Ensure ordering is what we want by adding class in first
+		CategoryBuilder.AddProperty(GET_MEMBER_NAME_CHECKED(UChildActorComponent, ChildActorClass));
+
+		IDetailPropertyRow& CATRow = CategoryBuilder.AddProperty(GET_MEMBER_NAME_CHECKED(UChildActorComponent, ChildActorTemplate));
+		CATRow.Visibility(TAttribute<EVisibility>::Create([ObjectsBeingCustomized]()
+		{
+			for (const TWeakObjectPtr<UObject>& ObjectBeingCustomized : ObjectsBeingCustomized)
+			{
+				if (UChildActorComponent* CAC = Cast<UChildActorComponent>(ObjectBeingCustomized.Get()))
+				{
+					if (CAC->ChildActorTemplate == nullptr)
+					{
+						return EVisibility::Hidden;
+					}
+				}
+				else
+				{
+					return EVisibility::Hidden;
+				}
+			}
+
+			return EVisibility::Visible;
+		}));
 	}
 }
 
diff --git a/Engine/Source/Editor/Kismet/Private/BlueprintEditor.cpp b/Engine/Source/Editor/Kismet/Private/BlueprintEditor.cpp
--- a/Engine/Source/Editor/Kismet/Private/BlueprintEditor.cpp
+++ b/Engine/Source/Editor/Kismet/Private/BlueprintEditor.cpp
@@ -3088,7 +3088,7 @@ void FBlueprintEditor::OnGraphActionMenuClosed(bool bActionExecuted, bool bConte
 	HasOpenActionMenu = nullptr;
 }
 
-void FBlueprintEditor::OnSelectedNodesChanged(const FGraphPanelSelectionSet& NewSelection)
+void FBlueprintEditor::OnSelectedNodesChangedImpl(const FGraphPanelSelectionSet& NewSelection)
 {
 	if ( NewSelection.Num() > 0 )
 	{
diff --git a/Engine/Source/Editor/Kismet/Public/BlueprintEditor.h b/Engine/Source/Editor/Kismet/Public/BlueprintEditor.h
--- a/Engine/Source/Editor/Kismet/Public/BlueprintEditor.h
+++ b/Engine/Source/Editor/Kismet/Public/BlueprintEditor.h
@@ -609,9 +609,12 @@ protected:
 
 	/** Whether or not we can redo an undone action */
 	bool CanRedoGraphAction() const;
+	
+	/** Called when the selection changes in the GraphEditor */
+	virtual void OnSelectedNodesChangedImpl(const TSet<class UObject*>& NewSelection);
 
 	/** Called when the selection changes in the GraphEditor */
-	void OnSelectedNodesChanged(const TSet<class UObject*>& NewSelection);
+	void OnSelectedNodesChanged(const TSet<class UObject*>& NewSelection) { OnSelectedNodesChangedImpl(NewSelection); }
 
 	/** Called when an actor is dropped onto the graph editor */
 	void OnGraphEditorDropActor(const TArray< TWeakObjectPtr<AActor> >& Actors, UEdGraph* Graph, const FVector2D& DropLocation);
diff --git a/Engine/Source/Editor/KismetCompiler/Private/KismetCompiler.cpp b/Engine/Source/Editor/KismetCompiler/Private/KismetCompiler.cpp
--- a/Engine/Source/Editor/KismetCompiler/Private/KismetCompiler.cpp
+++ b/Engine/Source/Editor/KismetCompiler/Private/KismetCompiler.cpp
@@ -3607,7 +3607,6 @@ void FKismetCompilerContext::Compile()
 
 		CopyTermDefaultsToDefaultObject(NewCDO);
 		SetCanEverTick();
-		SetWantsBeginPlay();
 
 		// Note: The old->new CDO copy is deferred when regenerating, so we skip this step in that case.
 		if (!Blueprint->HasAnyFlags(RF_BeingRegenerated))
@@ -4023,38 +4022,6 @@ void FKismetCompilerContext::SetCanEverTick() const
 	}
 }
 
-void FKismetCompilerContext::SetWantsBeginPlay() const
-{
-	UActorComponent* CDComponent = Cast<UActorComponent>(NewClass->GetDefaultObject());
-
-	if (CDComponent == nullptr)
-	{
-		return;
-	}
-
-	const bool bOldFlag = CDComponent->bWantsBeginPlay;
-
-	CDComponent->bWantsBeginPlay = NewClass->GetSuperClass()->GetDefaultObject<UActorComponent>()->bWantsBeginPlay;
-
-	if (!CDComponent->bWantsBeginPlay)
-	{
-		static FName ReceiveBeginPlayName(GET_FUNCTION_NAME_CHECKED(UActorComponent, ReceiveBeginPlay));
-		static FName ReceiveEndPlayName(GET_FUNCTION_NAME_CHECKED(UActorComponent, ReceiveEndPlay));
-		const UFunction* ReciveBeginPlayEvent = FKismetCompilerUtilities::FindOverriddenImplementableEvent(ReceiveBeginPlayName, NewClass);
-		const UFunction* ReciveEndPlayEvent = FKismetCompilerUtilities::FindOverriddenImplementableEvent(ReceiveEndPlayName, NewClass);
-		if (ReciveBeginPlayEvent || ReciveEndPlayEvent)
-		{
-			CDComponent->bWantsBeginPlay = true;
-		}
-	}
-
-	if(CDComponent->bWantsBeginPlay != bOldFlag)
-	{
-		UE_LOG(LogK2Compiler, Verbose, TEXT("Overridden flag for class '%s': bWantsBeginPlay %s "), *NewClass->GetName(),
-			CDComponent->bWantsBeginPlay ? *(GTrue.ToString()) : *(GFalse.ToString()) );
-	}
-}
-
 bool FKismetCompilerContext::UsePersistentUberGraphFrame() const
 {
 	return UBlueprintGeneratedClass::UsePersistentUberGraphFrame() && !CompileOptions.DoesRequireCppCodeGeneration();
diff --git a/Engine/Source/Editor/KismetCompiler/Public/KismetCompiler.h b/Engine/Source/Editor/KismetCompiler/Public/KismetCompiler.h
--- a/Engine/Source/Editor/KismetCompiler/Public/KismetCompiler.h
+++ b/Engine/Source/Editor/KismetCompiler/Public/KismetCompiler.h
@@ -291,12 +291,6 @@ protected:
 	 */
 	void SetCanEverTick() const;
 
-	/** 
-	 * Function works only if subclass UActorComponent.
-	 * If BeginPlay or EndPlay event is defined on the Component, force bWantsBeginPlay.
-	 */
-	void SetWantsBeginPlay() const;
-
 	/** Scan FunctionList and return Entry point, for matching one  */
 	const UK2Node_FunctionEntry* FindLocalEntryPoint(const UFunction* Function) const;
 
diff --git a/Engine/Source/Editor/KismetWidgets/Private/SScrubControlPanel.cpp b/Engine/Source/Editor/KismetWidgets/Private/SScrubControlPanel.cpp
--- a/Engine/Source/Editor/KismetWidgets/Private/SScrubControlPanel.cpp
+++ b/Engine/Source/Editor/KismetWidgets/Private/SScrubControlPanel.cpp
@@ -26,12 +26,15 @@ void SScrubControlPanel::Construct( const SScrubControlPanel::FArguments& InArgs
 	TransportControlArgs.OnToggleLooping = InArgs._OnClickedToggleLoop;
 	TransportControlArgs.OnGetLooping = InArgs._OnGetLooping;
 	TransportControlArgs.OnGetPlaybackMode = InArgs._OnGetPlaybackMode;
+	TransportControlArgs.OnGetRecording = InArgs._OnGetRecording;
 	TransportControlArgs.OnTickPlayback = InArgs._OnTickPlayback;
 	
 	FTransportControlArgs TransportControlArgsForRealtimeStreamingMode;
 	TransportControlArgsForRealtimeStreamingMode.OnForwardPlay = TransportControlArgs.OnForwardPlay;
 	TransportControlArgsForRealtimeStreamingMode.OnForwardStep = TransportControlArgs.OnForwardStep;
 	TransportControlArgsForRealtimeStreamingMode.OnGetPlaybackMode = TransportControlArgs.OnGetPlaybackMode;
+	TransportControlArgsForRealtimeStreamingMode.OnGetRecording = TransportControlArgs.OnGetRecording;
+	TransportControlArgsForRealtimeStreamingMode.OnRecord = TransportControlArgs.OnRecord;
 
 	this->ChildSlot
 	.Padding( FMargin( 0.0f, 1.0f) )
@@ -48,6 +51,7 @@ void SScrubControlPanel::Construct( const SScrubControlPanel::FArguments& InArgs
 				SAssignNew(ScrubWidget, SScrubWidget)
 				.Value(InArgs._Value)
 				.NumOfKeys(InArgs._NumOfKeys)
+				.DisplayDrag(InArgs._DisplayDrag)
 				.SequenceLength(InArgs._SequenceLength)
 				.OnValueChanged(InArgs._OnValueChanged)
 				.OnBeginSliderMovement(InArgs._OnBeginSliderMovement)
diff --git a/Engine/Source/Editor/KismetWidgets/Private/SScrubWidget.cpp b/Engine/Source/Editor/KismetWidgets/Private/SScrubWidget.cpp
--- a/Engine/Source/Editor/KismetWidgets/Private/SScrubWidget.cpp
+++ b/Engine/Source/Editor/KismetWidgets/Private/SScrubWidget.cpp
@@ -12,12 +12,13 @@ static const float MinStepLengh=15.f;
 int32 SScrubWidget::GetDivider(float InputMinX, float InputMaxX, FVector2D WidgetSize, float SequenceLength, int32 NumFrames)
 {
 	FTrackScaleInfo TimeScaleInfo(InputMinX, InputMaxX, 0.f, 0.f, WidgetSize);
+	const FAnimKeyHelper Helper(SequenceLength, NumFrames);
 
-	float TimePerKey = (NumFrames > 0) ? SequenceLength / (float)(NumFrames) : 0.0f;
+	float TimePerKey = Helper.TimePerKey();
 
 	float TotalWidgetWidth = TimeScaleInfo.WidgetSize.X;
-	float NumKeys = TimeScaleInfo.ViewInputRange / TimePerKey;
-	float KeyWidgetWidth = TotalWidgetWidth / NumKeys;
+	int32 NumKeys = TimeScaleInfo.ViewInputRange / TimePerKey;
+	float KeyWidgetWidth = TotalWidgetWidth / (float)Helper.GetNumKeys();
 	int32 Divider = 1; 
 	if (KeyWidgetWidth > 0)
 	{
@@ -46,7 +47,7 @@ void SScrubWidget::Construct( const SScrubWidget::FArguments& InArgs )
 
 	DraggableBars = InArgs._DraggableBars;
 	OnBarDrag = InArgs._OnBarDrag;
-
+	bDisplayDrag = InArgs._DisplayDrag;
 	bMouseMovedDuringPanning = false;
 	bDragging = false;
 	bPanning = false;
@@ -79,17 +80,17 @@ int32 SScrubWidget::OnPaint( const FPaintArgs& Args, const FGeometry& AllottedGe
 	const FSlateBrush* StyleInfo = FEditorStyle::GetBrush( TEXT( "ProgressBar.Background" ) );
 	const float GeomHeight = AllottedGeometry.Size.Y;
 
+	const FTrackScaleInfo TimeScaleInfo(ViewInputMin.Get(), ViewInputMax.Get(), 0.f, 0.f, AllottedGeometry.Size);
+
 	if ( NumOfKeys.Get() > 0 && SequenceLength.Get() > 0)
 	{
-		const FTrackScaleInfo TimeScaleInfo(ViewInputMin.Get(), ViewInputMax.Get(), 0.f, 0.f, AllottedGeometry.Size);
 		const int32 Divider = SScrubWidget::GetDivider( ViewInputMin.Get(), ViewInputMax.Get(), AllottedGeometry.Size, SequenceLength.Get(), NumOfKeys.Get() );
+		const FAnimKeyHelper Helper(SequenceLength.Get(), NumOfKeys.Get());
 		const float HalfDivider = Divider/2.f;
 		
-		const int32 TotalNumKeys = NumOfKeys.Get();
-
-		const float TimePerKey = (TotalNumKeys > 0) ? SequenceLength.Get() / (float)(TotalNumKeys) : 0.0f;
+		const float TimePerKey = Helper.TimePerKey();
 
-		for (float KeyVal = 0; KeyVal < TotalNumKeys; KeyVal += HalfDivider)
+		for (float KeyVal = 0; KeyVal < Helper.GetNumKeys(); KeyVal += HalfDivider)
 		{
 			const float CurValue = KeyVal*TimePerKey;
 			const float XPos = TimeScaleInfo.InputToLocalX(CurValue);
@@ -143,7 +144,10 @@ int32 SScrubWidget::OnPaint( const FPaintArgs& Args, const FGeometry& AllottedGe
  
  			}
 		}
+	}
 
+	if (SequenceLength.Get() > 0 && bDisplayDrag.Get())
+	{
 		const int32 ArrowLayer = TextLayer + 1;
 		{
 			const float XPos = TimeScaleInfo.InputToLocalX(ValueAttribute.Get());
diff --git a/Engine/Source/Editor/KismetWidgets/Public/SScrubControlPanel.h b/Engine/Source/Editor/KismetWidgets/Public/SScrubControlPanel.h
--- a/Engine/Source/Editor/KismetWidgets/Public/SScrubControlPanel.h
+++ b/Engine/Source/Editor/KismetWidgets/Public/SScrubControlPanel.h
@@ -13,6 +13,7 @@ public:
 		: _Value( 0 )
 		, _NumOfKeys()
 		, _SequenceLength()
+		, _DisplayDrag(true)
 		, _OnValueChanged()
 		, _OnBeginSliderMovement()
 		, _OnEndSliderMovement()
@@ -26,6 +27,7 @@ public:
 		, _OnClickedToggleLoop()
 		, _OnGetLooping()
 		, _OnGetPlaybackMode()
+		, _OnGetRecording()
 		, _ViewInputMin()
 		, _ViewInputMax()
 		, _OnSetInputViewRange()
@@ -36,6 +38,7 @@ public:
 		SLATE_ATTRIBUTE( uint32, NumOfKeys )
 		SLATE_ATTRIBUTE( float, SequenceLength )
 		SLATE_ARGUMENT( bool, bAllowZoom )
+		SLATE_ATTRIBUTE(bool, DisplayDrag)
 		/** Called when the value is changed by slider or typing */
 		SLATE_EVENT( FOnFloatValueChanged, OnValueChanged )
 		/** Called right before the slider begins to move */
@@ -55,6 +58,7 @@ public:
 		SLATE_EVENT( FOnClicked, OnClickedToggleLoop )
 		SLATE_EVENT( FOnGetLooping, OnGetLooping )
 		SLATE_EVENT( FOnGetPlaybackMode, OnGetPlaybackMode )
+		SLATE_EVENT( FOnGetRecording, OnGetRecording )
 		/** View Input range **/
 		SLATE_ATTRIBUTE( float, ViewInputMin )
 		SLATE_ATTRIBUTE( float, ViewInputMax )
diff --git a/Engine/Source/Editor/KismetWidgets/Public/SScrubWidget.h b/Engine/Source/Editor/KismetWidgets/Public/SScrubWidget.h
--- a/Engine/Source/Editor/KismetWidgets/Public/SScrubWidget.h
+++ b/Engine/Source/Editor/KismetWidgets/Public/SScrubWidget.h
@@ -25,6 +25,7 @@ public:
 		, _NumOfKeys(10)
 		, _SequenceLength()
 		, _bAllowZoom(false)
+		, _DisplayDrag(true)
 		, _OnValueChanged()
 		, _OnBeginSliderMovement()
 		, _OnEndSliderMovement()
@@ -39,6 +40,7 @@ public:
 		SLATE_ATTRIBUTE( uint32, NumOfKeys )
 		SLATE_ATTRIBUTE( float, SequenceLength )
 		SLATE_ARGUMENT( bool, bAllowZoom )
+		SLATE_ATTRIBUTE( bool, DisplayDrag )
 		/** Called when the value is changed by slider or typing */
 		SLATE_EVENT( FOnFloatValueChanged, OnValueChanged )
 		/** Called right before the slider begins to move */
@@ -136,6 +138,7 @@ private:
 	/** The number of decimal places to display */
 	bool bDragging;
 	bool bAllowZoom;	
+	TAttribute<bool> bDisplayDrag;
 	/** If we are currently panning the panel*/
 	bool bPanning;
 	/** Has the mouse moved during panning - used to determine if we should open the context menu or not */
diff --git a/Engine/Source/Editor/LevelEditor/Private/LevelEditorToolBar.cpp b/Engine/Source/Editor/LevelEditor/Private/LevelEditorToolBar.cpp
--- a/Engine/Source/Editor/LevelEditor/Private/LevelEditorToolBar.cpp
+++ b/Engine/Source/Editor/LevelEditor/Private/LevelEditorToolBar.cpp
@@ -28,8 +28,8 @@
 #include "ISourceControlModule.h"
 #include "SVolumeControl.h"
 #include "ModuleManager.h"
-#include "GameFramework/GameMode.h"
-#include "GameFramework/GameState.h"
+#include "GameFramework/GameModeBase.h"
+#include "GameFramework/GameStateBase.h"
 #include "GameFramework/HUD.h"
 #include "GameFramework/Pawn.h"
 #include "GameFramework/PlayerController.h"
@@ -446,7 +446,7 @@ namespace LevelEditorActionHelpers
 		FText ResultText;
 
 		// Game modes can always be created and selected (providing the config is checked out, handled separately)
-		if(InRootClass != AGameMode::StaticClass() && InGameModeClass->HasAllClassFlags(CLASS_Native))
+		if(InRootClass != AGameModeBase::StaticClass() && InGameModeClass->HasAllClassFlags(CLASS_Native))
 		{
 			ResultText = LOCTEXT("CannotCreateClasses", "Cannot create classes when the game mode is a native class!");
 		}
@@ -473,7 +473,7 @@ namespace LevelEditorActionHelpers
 		FText ResultText;
 
 		// Game modes can always be created and selected (providing the config is checked out, handled separately)
-		if(InRootClass != AGameMode::StaticClass() && InGameModeClass->HasAllClassFlags(CLASS_Native))
+		if(InRootClass != AGameModeBase::StaticClass() && InGameModeClass->HasAllClassFlags(CLASS_Native))
 		{
 			ResultText = LOCTEXT("CannotSelectClasses", "Cannot select classes when the game mode is a native class!");
 		}
@@ -503,7 +503,7 @@ namespace LevelEditorActionHelpers
 		GameModeMenuSettings.OnCreateClassPicked = FOnClassPicked::CreateStatic( &LevelEditorActionHelpers::OnCreateGameModeClassPicked, InLevelEditor, bInProjectSettings );
 		GameModeMenuSettings.OnSelectClassPicked = FOnClassPicked::CreateStatic( &LevelEditorActionHelpers::OnSelectGameModeClassPicked, InLevelEditor, bInProjectSettings );
 		GameModeMenuSettings.CurrentClass = LevelEditorActionHelpers::GetGameModeClass(InLevelEditor, bInProjectSettings);
-		GameModeMenuSettings.RootClass = AGameMode::StaticClass();
+		GameModeMenuSettings.RootClass = AGameModeBase::StaticClass();
 		GameModeMenuSettings.LevelEditor = InLevelEditor;
 		GameModeMenuSettings.bIsProjectSettings = bInProjectSettings;
 
@@ -544,7 +544,7 @@ namespace LevelEditorActionHelpers
 		GameStateMenuSettings.OnCreateClassPicked = FOnClassPicked::CreateStatic( &LevelEditorActionHelpers::OnCreateGameStateClassPicked, InSettingsData.LevelEditor, InSettingsData.bIsProjectSettings );
 		GameStateMenuSettings.OnSelectClassPicked = FOnClassPicked::CreateStatic( &LevelEditorActionHelpers::OnSelectGameStateClassPicked, InSettingsData.LevelEditor, InSettingsData.bIsProjectSettings );
 		GameStateMenuSettings.CurrentClass = LevelEditorActionHelpers::GetGameStateClass(InSettingsData.LevelEditor, InSettingsData.bIsProjectSettings);
-		GameStateMenuSettings.RootClass = AGameState::StaticClass();
+		GameStateMenuSettings.RootClass = AGameStateBase::StaticClass();
 		GameStateMenuSettings.LevelEditor = InSettingsData.LevelEditor;
 		GameStateMenuSettings.bIsProjectSettings = InSettingsData.bIsProjectSettings;
 
@@ -656,7 +656,7 @@ void LevelEditorActionHelpers::GetBlueprintSettingsSubMenu(FMenuBuilder& InMenuB
 		}
 	}
 
-	if(InSettingsData.bIsProjectSettings && InSettingsData.CurrentClass && InSettingsData.CurrentClass->IsChildOf(AGameMode::StaticClass()) && !FLevelEditorActionCallbacks::CanSelectGameModeBlueprint())
+	if(InSettingsData.bIsProjectSettings && InSettingsData.CurrentClass && InSettingsData.CurrentClass->IsChildOf(AGameModeBase::StaticClass()) && !FLevelEditorActionCallbacks::CanSelectGameModeBlueprint())
 	{
 		InMenuBuilder.BeginSection("CheckoutSection", LOCTEXT("CheckoutSection","Check Out Project Settings") );
 		TAttribute<FText> CheckOutLabel;
@@ -684,7 +684,7 @@ void LevelEditorActionHelpers::GetBlueprintSettingsSubMenu(FMenuBuilder& InMenuB
 		FNewMenuDelegate::CreateStatic( &LevelEditorActionHelpers::GetCreateSettingsClassSubMenu, InSettingsData.RootClass, InSettingsData.OnCreateClassPicked ),
 		FUIAction(
 			FExecuteAction(), 
-			InSettingsData.RootClass == AGameMode::StaticClass()? 
+			InSettingsData.RootClass == AGameModeBase::StaticClass()?
 				FCanExecuteAction::CreateStatic(CannotCreateSelectNativeProjectGameMode, InSettingsData.bIsProjectSettings) 
 				: FCanExecuteAction::CreateStatic( &CanCreateSelectSubClass, GetGameModeClass(InSettingsData.LevelEditor, InSettingsData.bIsProjectSettings), InSettingsData.bIsProjectSettings )
 		),
@@ -697,7 +697,7 @@ void LevelEditorActionHelpers::GetBlueprintSettingsSubMenu(FMenuBuilder& InMenuB
 		FNewMenuDelegate::CreateStatic( &LevelEditorActionHelpers::GetSelectSettingsClassSubMenu, InSettingsData.RootClass, InSettingsData.OnSelectClassPicked ),
 		FUIAction(
 			FExecuteAction(), 
-			InSettingsData.RootClass == AGameMode::StaticClass()?
+			InSettingsData.RootClass == AGameModeBase::StaticClass()?
 				FCanExecuteAction::CreateStatic(CannotCreateSelectNativeProjectGameMode, InSettingsData.bIsProjectSettings) 
 				: FCanExecuteAction::CreateStatic( &CanCreateSelectSubClass, GetGameModeClass(InSettingsData.LevelEditor, InSettingsData.bIsProjectSettings), InSettingsData.bIsProjectSettings )
 		),
@@ -705,7 +705,7 @@ void LevelEditorActionHelpers::GetBlueprintSettingsSubMenu(FMenuBuilder& InMenuB
 	);
 
 	// For GameMode classes only, there are some sub-classes we need to add to the menu
-	if(InSettingsData.RootClass == AGameMode::StaticClass())
+	if(InSettingsData.RootClass == AGameModeBase::StaticClass())
 	{
 		InMenuBuilder.BeginSection(NAME_None, LOCTEXT("GameModeClasses", "Game Mode Classes"));
 
@@ -830,7 +830,7 @@ UClass* LevelEditorActionHelpers::GetGameStateClass(TWeakPtr< SLevelEditor > InL
 {
 	if(UClass* GameModeClass = GetGameModeClass(InLevelEditor, bInIsProjectSettings))
 	{
-		AGameMode* ActiveGameMode = Cast<AGameMode>(GameModeClass->GetDefaultObject());
+		AGameModeBase* ActiveGameMode = Cast<AGameModeBase>(GameModeClass->GetDefaultObject());
 		if(ActiveGameMode)
 		{
 			return ActiveGameMode->GameStateClass;
@@ -883,7 +883,7 @@ void LevelEditorActionHelpers::OnSelectGameStateClassPicked(UClass* InChosenClas
 	if(UClass* GameModeClass = GetGameModeClass(InLevelEditor, bInIsProjectSettings))
 	{
 		const FScopedTransaction Transaction( NSLOCTEXT("LevelEditorCommands", "SelectGameStateClassAction", "Set Game State Class") );
-		AGameMode* ActiveGameMode = Cast<AGameMode>(GameModeClass->GetDefaultObject());
+		AGameModeBase* ActiveGameMode = Cast<AGameModeBase>(GameModeClass->GetDefaultObject());
 		ActiveGameMode->Modify();
 		ActiveGameMode->GameStateClass = InChosenClass;
 	}
@@ -895,7 +895,7 @@ UClass* LevelEditorActionHelpers::GetPawnClass(TWeakPtr< SLevelEditor > InLevelE
 {
 	if(UClass* GameModeClass = GetGameModeClass(InLevelEditor, bInIsProjectSettings))
 	{
-		AGameMode* ActiveGameMode = Cast<AGameMode>(GameModeClass->GetDefaultObject());
+		AGameModeBase* ActiveGameMode = Cast<AGameModeBase>(GameModeClass->GetDefaultObject());
 
 		if(ActiveGameMode)
 		{
@@ -950,7 +950,7 @@ void LevelEditorActionHelpers::OnSelectPawnClassPicked(UClass* InChosenClass, TW
 	{
 		const FScopedTransaction Transaction( NSLOCTEXT("LevelEditorCommands", "SelectPawnClassAction", "Set Pawn Class") );
 
-		AGameMode* ActiveGameMode = Cast<AGameMode>(GameModeClass->GetDefaultObject());
+		AGameModeBase* ActiveGameMode = Cast<AGameModeBase>(GameModeClass->GetDefaultObject());
 		ActiveGameMode->Modify();
 		ActiveGameMode->DefaultPawnClass = InChosenClass;
 	}
@@ -962,7 +962,7 @@ UClass* LevelEditorActionHelpers::GetHUDClass(TWeakPtr< SLevelEditor > InLevelEd
 {
 	if(UClass* GameModeClass = GetGameModeClass(InLevelEditor, bInIsProjectSettings))
 	{
-		AGameMode* ActiveGameMode = Cast<AGameMode>(GameModeClass->GetDefaultObject());
+		AGameModeBase* ActiveGameMode = Cast<AGameModeBase>(GameModeClass->GetDefaultObject());
 		if(ActiveGameMode)
 		{
 			return ActiveGameMode->HUDClass;
@@ -1016,7 +1016,7 @@ void LevelEditorActionHelpers::OnSelectHUDClassPicked(UClass* InChosenClass, TWe
 	{
 		const FScopedTransaction Transaction( NSLOCTEXT("LevelEditorCommands", "SelectHUDClassAction", "Set HUD Class") );
 
-		AGameMode* ActiveGameMode = Cast<AGameMode>(GameModeClass->GetDefaultObject());
+		AGameModeBase* ActiveGameMode = Cast<AGameModeBase>(GameModeClass->GetDefaultObject());
 		ActiveGameMode->Modify();
 		ActiveGameMode->HUDClass = InChosenClass;
 	}
@@ -1028,7 +1028,7 @@ UClass* LevelEditorActionHelpers::GetPlayerControllerClass(TWeakPtr< SLevelEdito
 {
 	if(UClass* GameModeClass = GetGameModeClass(InLevelEditor, bInIsProjectSettings))
 	{
-		AGameMode* ActiveGameMode = Cast<AGameMode>(GameModeClass->GetDefaultObject());
+		AGameModeBase* ActiveGameMode = Cast<AGameModeBase>(GameModeClass->GetDefaultObject());
 		if(ActiveGameMode)
 		{
 			return ActiveGameMode->PlayerControllerClass;
@@ -1082,7 +1082,7 @@ void LevelEditorActionHelpers::OnSelectPlayerControllerClassPicked(UClass* InCho
 	{
 		const FScopedTransaction Transaction( NSLOCTEXT("LevelEditorCommands", "SelectPlayerControllerClassAction", "Set Player Controller Class") );
 
-		AGameMode* ActiveGameMode = Cast<AGameMode>(GameModeClass->GetDefaultObject());
+		AGameModeBase* ActiveGameMode = Cast<AGameModeBase>(GameModeClass->GetDefaultObject());
 		ActiveGameMode->Modify();
 		ActiveGameMode->PlayerControllerClass = InChosenClass;
 	}
diff --git a/Engine/Source/Editor/MergeActors/Private/MeshMergingTool/SMeshMergingDialog.cpp b/Engine/Source/Editor/MergeActors/Private/MeshMergingTool/SMeshMergingDialog.cpp
--- a/Engine/Source/Editor/MergeActors/Private/MeshMergingTool/SMeshMergingDialog.cpp
+++ b/Engine/Source/Editor/MergeActors/Private/MeshMergingTool/SMeshMergingDialog.cpp
@@ -204,15 +204,42 @@ void SMeshMergingDialog::UpdateSelectedStaticMeshComponents()
 
 	// Retrieve static mesh components from selected actors
 	SelectedMeshComponents.Empty();
-	for (AActor* Actor : Actors)
+	for (int32 ActorIndex = 0; ActorIndex < Actors.Num(); ++ActorIndex )
 	{
+		AActor* Actor = Actors[ActorIndex];
 		check(Actor != nullptr);
 
+		TArray<UChildActorComponent*> ChildActorComponents;
+		Actor->GetComponents<UChildActorComponent>(ChildActorComponents);
+		for (UChildActorComponent* ChildComponent : ChildActorComponents)
+		{
+			// Push actor at the back of array so we will process it
+			AActor* ChildActor = ChildComponent->GetChildActor();
+			if (ChildActor)
+			{
+				Actors.Add(ChildActor);
+			}
+		}
+		
 		TArray<UStaticMeshComponent*> StaticMeshComponents;
 		Actor->GetComponents<UStaticMeshComponent>(StaticMeshComponents);
 		for (UStaticMeshComponent* MeshComponent : StaticMeshComponents)
 		{
 			SelectedMeshComponents.Add(TSharedPtr<FMeshComponentData>(new FMeshComponentData(MeshComponent)));
+			TSharedPtr<FMeshComponentData>& MeshComponentData = SelectedMeshComponents.Last();
+			MeshComponentData->bShouldIncorporate = (MeshComponentData->MeshComponent->StaticMesh != nullptr);			
+			// See if we stored a checkbox state for this mesh component, and set accordingly
+			ECheckBoxState State = (MeshComponentData->MeshComponent->StaticMesh != nullptr) ? ECheckBoxState::Checked : ECheckBoxState::Unchecked;
+			auto StoredState = StoredCheckBoxStates.Find(MeshComponentData->MeshComponent.Get());
+			if (StoredState)
+			{
+				MeshComponentData->bShouldIncorporate = (*StoredState == ECheckBoxState::Checked);
+				NumSelectedMeshComponents += (MeshComponentData->bShouldIncorporate) ? 1 : 0;
+			}
+			else if (MeshComponentData->bShouldIncorporate == true)
+			{
+				NumSelectedMeshComponents++;
+			}
 		}		
 	}
 }
@@ -224,8 +251,6 @@ TSharedRef<ITableRow> SMeshMergingDialog::MakeComponentListItemWidget(TSharedPtr
 	// Retrieve information about the mesh component
 	const FString OwningActorName = MeshComponentData->MeshComponent->GetOwner()->GetName();
 	const FString StaticMeshName = (MeshComponentData->MeshComponent->StaticMesh != nullptr) ? MeshComponentData->MeshComponent->StaticMesh->GetName() : TEXT("No Static Mesh Available");
-	// Only static mesh components with valid static mesh data should be incorporated	
-	MeshComponentData->bShouldIncorporate = (MeshComponentData->MeshComponent->StaticMesh != nullptr);
 	const FString ComponentName = MeshComponentData->MeshComponent->GetName();
 
 	// See if we stored a checkbox state for this mesh component, and set accordingly
@@ -234,15 +259,8 @@ TSharedRef<ITableRow> SMeshMergingDialog::MakeComponentListItemWidget(TSharedPtr
 	if (StoredState)
 	{
 		State = *StoredState;
-		MeshComponentData->bShouldIncorporate = (State == ECheckBoxState::Checked);
-		
-		NumSelectedMeshComponents += (MeshComponentData->bShouldIncorporate) ? 1 : 0;
 	}
-	else if (MeshComponentData->bShouldIncorporate == true)
-	{
-		NumSelectedMeshComponents++;
-	}
-
+	
 	return SNew(STableRow<TSharedPtr<FMeshComponentData>>, OwnerTable)
 		[
 			SNew(SBox)
diff --git a/Engine/Source/Editor/MovieSceneCaptureDialog/Private/MovieSceneCaptureDialogModule.cpp b/Engine/Source/Editor/MovieSceneCaptureDialog/Private/MovieSceneCaptureDialogModule.cpp
--- a/Engine/Source/Editor/MovieSceneCaptureDialog/Private/MovieSceneCaptureDialogModule.cpp
+++ b/Engine/Source/Editor/MovieSceneCaptureDialog/Private/MovieSceneCaptureDialogModule.cpp
@@ -4,7 +4,7 @@
 #include "MovieSceneCapture.h"
 
 #include "UnrealEdMisc.h"
-#include "GameFramework/GameMode.h"
+#include "GameFramework/GameModeBase.h"
 
 #include "SlateBasics.h"
 #include "SlateExtras.h"
@@ -571,7 +571,7 @@ private:
 	TArray<uint8> BackedUpPlaySettings;
 	UMovieSceneCapture* CaptureObject;
 
-	TSubclassOf<AGameMode> CachedGameMode;
+	TSubclassOf<AGameModeBase> CachedGameMode;
 };
 
 class FMovieSceneCaptureDialogModule : public IMovieSceneCaptureDialogModule
diff --git a/Engine/Source/Editor/Persona/Persona.Build.cs b/Engine/Source/Editor/Persona/Persona.Build.cs
--- a/Engine/Source/Editor/Persona/Persona.Build.cs
+++ b/Engine/Source/Editor/Persona/Persona.Build.cs
@@ -12,8 +12,12 @@ public class Persona : ModuleRules
 			new string[] {
 				"AssetRegistry", 
 				"MainFrame",
-				"DesktopPlatform"
-			}
+				"DesktopPlatform",
+                "SkeletonEditor",
+                "ContentBrowser",
+                "AssetTools",
+                "AnimationEditor",
+            }
 		);
 
 		PrivateDependencyModuleNames.AddRange(
@@ -44,7 +48,10 @@ public class Persona : ModuleRules
 				"Documentation",
 				"MainFrame",
 				"DesktopPlatform",
-			}
+                "SkeletonEditor",
+                "AssetTools",
+                "AnimationEditor",
+            }
 		);
 	}
 }
diff --git a/Engine/Source/Editor/Persona/Private/AnimGraphNodeDetails.cpp b/Engine/Source/Editor/Persona/Private/AnimGraphNodeDetails.cpp
--- a/Engine/Source/Editor/Persona/Private/AnimGraphNodeDetails.cpp
+++ b/Engine/Source/Editor/Persona/Private/AnimGraphNodeDetails.cpp
@@ -18,13 +18,13 @@
 #include "PropertyCustomizationHelpers.h"
 #include "AnimGraphNode_Base.h"
 #include "GraphEditor.h"
-#include "Persona.h"
 #include "BoneSelectionWidget.h"
 #include "SExpandableArea.h"
 #include "Animation/BlendProfile.h"
-#include "SBlendProfilePicker.h"
 #include "AnimGraphNode_AssetPlayerBase.h"
 #include "Animation/AnimInstance.h"
+#include "ISkeletonEditorModule.h"
+#include "BlueprintEditor.h"
 
 #define LOCTEXT_NAMESPACE "KismetNodeWithOptionalPinsDetails"
 
@@ -226,11 +226,14 @@ TSharedRef<SWidget> FAnimGraphNodeDetails::CreatePropertyWidget(UProperty* Targe
 
 			UBlendProfile* CurrentProfile = Cast<UBlendProfile>(PropertyValue);
 
-			return SNew(SBlendProfilePicker)
-				.TargetSkeleton(this->TargetSkeleton)
-				.AllowNew(false)
-				.OnBlendProfileSelected(this, &FAnimGraphNodeDetails::OnBlendProfileChanged, PropertyPtr)
-				.InitialProfile(CurrentProfile);
+			FBlendProfilePickerArgs Args;
+			Args.bAllowNew = false;
+			Args.OnBlendProfileSelected = FOnBlendProfileSelected::CreateSP(this, &FAnimGraphNodeDetails::OnBlendProfileChanged, PropertyPtr);
+			Args.InitialProfile = CurrentProfile;
+
+			ISkeletonEditorModule& SkeletonEditorModule = FModuleManager::Get().LoadModuleChecked<ISkeletonEditorModule>("SkeletonEditor");
+			return SkeletonEditorModule.CreateBlendProfilePicker(this->TargetSkeleton, Args);
+
 		}
 	}
 
@@ -510,6 +513,9 @@ void FBoneReferenceCustomization::CustomizeHeader( TSharedRef<IPropertyHandle> S
 
 	if (TargetSkeleton)
 	{
+		ISkeletonEditorModule& SkeletonEditorModule = FModuleManager::LoadModuleChecked<ISkeletonEditorModule>("SkeletonEditor");
+		TSharedRef<IEditableSkeleton> EditableSkeleton = SkeletonEditorModule.CreateEditableSkeleton(TargetSkeleton);
+
 		HeaderRow.NameContent()
 		[
 			StructPropertyHandle->CreatePropertyNameWidget()
@@ -517,7 +523,7 @@ void FBoneReferenceCustomization::CustomizeHeader( TSharedRef<IPropertyHandle> S
 
 		HeaderRow.ValueContent()
 		[
-			SNew(SBoneSelectionWidget, TargetSkeleton)
+			SNew(SBoneSelectionWidget, EditableSkeleton)
 			.Tooltip(StructPropertyHandle->GetToolTipText())
 			.OnBoneSelectionChanged(this, &FBoneReferenceCustomization::OnBoneSelectionChanged)
 			.OnGetSelectedBone(this, &FBoneReferenceCustomization::GetSelectedBone)
@@ -545,9 +551,9 @@ FName FBoneReferenceCustomization::GetSelectedBone() const
 	return FName(*OutText);
 }
 
-TSharedRef<IDetailCustomization> FAnimGraphParentPlayerDetails::MakeInstance(TWeakPtr<FPersona> InPersona)
+TSharedRef<IDetailCustomization> FAnimGraphParentPlayerDetails::MakeInstance(TSharedRef<FBlueprintEditor> InBlueprintEditor)
 {
-	return MakeShareable(new FAnimGraphParentPlayerDetails(InPersona));
+	return MakeShareable(new FAnimGraphParentPlayerDetails(InBlueprintEditor));
 }
 
 void FAnimGraphParentPlayerDetails::CustomizeDetails(class IDetailLayoutBuilder& DetailBuilder)
@@ -628,7 +634,7 @@ void FAnimGraphParentPlayerDetails::CustomizeDetails(class IDetailLayoutBuilder&
 
 TSharedRef<ITableRow> FAnimGraphParentPlayerDetails::OnGenerateRow(TSharedPtr<FPlayerTreeViewEntry> EntryPtr, const TSharedRef< STableViewBase >& OwnerTable)
 {
-	return SNew(SParentPlayerTreeRow, OwnerTable).Item(EntryPtr).OverrideObject(EditorObject).Persona(PersonaPtr);
+	return SNew(SParentPlayerTreeRow, OwnerTable).Item(EntryPtr).OverrideObject(EditorObject).BlueprintEditor(BlueprintEditorPtr);
 }
 
 void FAnimGraphParentPlayerDetails::OnGetChildren(TSharedPtr<FPlayerTreeViewEntry> InParent, TArray< TSharedPtr<FPlayerTreeViewEntry> >& OutChildren)
@@ -640,7 +646,7 @@ void SParentPlayerTreeRow::Construct(const FArguments& InArgs, const TSharedRef<
 {
 	Item = InArgs._Item;
 	EditorObject = InArgs._OverrideObject;
-	Persona = InArgs._Persona;
+	BlueprintEditor = InArgs._BlueprintEditor;
 
 	if(Item->Override)
 	{
@@ -746,13 +752,13 @@ void SParentPlayerTreeRow::OnAssetSelected(const FAssetData& AssetData)
 
 FReply SParentPlayerTreeRow::OnFocusNodeButtonClicked()
 {
-	TSharedPtr<FPersona> SharedPersona = Persona.Pin();
-	if(SharedPersona.IsValid())
+	TSharedPtr<FBlueprintEditor> SharedBlueprintEditor = BlueprintEditor.Pin();
+	if(SharedBlueprintEditor.IsValid())
 	{
 		if(GraphNode)
 		{
 			UEdGraph* EdGraph = GraphNode->GetGraph();
-			TSharedPtr<SGraphEditor> GraphEditor = SharedPersona->OpenGraphAndBringToFront(EdGraph);
+			TSharedPtr<SGraphEditor> GraphEditor = SharedBlueprintEditor->OpenGraphAndBringToFront(EdGraph);
 			if (GraphEditor.IsValid())
 			{
 				GraphEditor->JumpToNode(GraphNode, false);
diff --git a/Engine/Source/Editor/Persona/Private/AnimGraphNodeDetails.h b/Engine/Source/Editor/Persona/Private/AnimGraphNodeDetails.h
--- a/Engine/Source/Editor/Persona/Private/AnimGraphNodeDetails.h
+++ b/Engine/Source/Editor/Persona/Private/AnimGraphNodeDetails.h
@@ -123,13 +123,13 @@ struct FPlayerTreeViewEntry : public TSharedFromThis<FPlayerTreeViewEntry>
 class FAnimGraphParentPlayerDetails : public IDetailCustomization
 {
 public:
-	static TSharedRef<IDetailCustomization> MakeInstance(TWeakPtr<FPersona> InPersona);
+	static TSharedRef<IDetailCustomization> MakeInstance(TSharedRef<class FBlueprintEditor> InBlueprintEditor);
 
 	virtual void CustomizeDetails(class IDetailLayoutBuilder& DetailBuilder);
 
 private:
 
-	FAnimGraphParentPlayerDetails(TWeakPtr<FPersona> InPersona) : PersonaPtr(InPersona)
+	FAnimGraphParentPlayerDetails(const TSharedRef<class FBlueprintEditor>& InBlueprintEditor) : BlueprintEditorPtr(InBlueprintEditor)
 	{}
 
 	TSharedRef<ITableRow> OnGenerateRow(TSharedPtr<FPlayerTreeViewEntry> EventPtr, const TSharedRef< STableViewBase >& OwnerTable);
@@ -138,8 +138,8 @@ private:
 	// Entries in the tree view
 	TArray<TSharedPtr<FPlayerTreeViewEntry>> ListEntries;
 	
-	// Hosting Persona instance
-	TWeakPtr<FPersona> PersonaPtr;
+	// Hosting Blueprint Editor instance
+	TWeakPtr<class FBlueprintEditor> BlueprintEditorPtr;
 
 	// Editor meta-object containing override information
 	UEditorParentPlayerListObj* EditorObject;
@@ -151,7 +151,7 @@ public:
 	SLATE_BEGIN_ARGS(SParentPlayerTreeRow){}
 		SLATE_ARGUMENT(TSharedPtr<FPlayerTreeViewEntry>, Item);
 		SLATE_ARGUMENT(UEditorParentPlayerListObj*, OverrideObject);
-		SLATE_ARGUMENT(TWeakPtr<FPersona>, Persona);
+		SLATE_ARGUMENT(TWeakPtr<class FBlueprintEditor>, BlueprintEditor);
 	SLATE_END_ARGS();
 
 	void Construct(const FArguments& InArgs, const TSharedRef<STableViewBase>& InOwnerTableView);
@@ -193,6 +193,6 @@ private:
 	// Graphnode this row represents, if any
 	UAnimGraphNode_Base* GraphNode;
 
-	// Persona editor pointer
-	TWeakPtr<FPersona> Persona;
+	// Blueprint editor pointer
+	TWeakPtr<FBlueprintEditor> BlueprintEditor;
 };
diff --git a/Engine/Source/Editor/Persona/Private/AnimViewportMenuCommands.cpp b/Engine/Source/Editor/Persona/Private/AnimViewportMenuCommands.cpp
--- a/Engine/Source/Editor/Persona/Private/AnimViewportMenuCommands.cpp
+++ b/Engine/Source/Editor/Persona/Private/AnimViewportMenuCommands.cpp
@@ -7,13 +7,10 @@
 
 void FAnimViewportMenuCommands::RegisterCommands()
 {
-	UI_COMMAND( Auto, "Auto", "Preview what is active on the current window", EUserInterfaceActionType::RadioButton, FInputChord() );
-	UI_COMMAND( Lock, "Lock", "Lock to current preview", EUserInterfaceActionType::RadioButton, FInputChord() );
+	UI_COMMAND( PreviewSceneSettings, "Preview Scene Settings...", "The Advanced Preview Settings tab will let you alter the preview scene's settings.", EUserInterfaceActionType::Button, FInputChord());
 
 	UI_COMMAND( CameraFollow, "Camera Follow", "Follow the bound of the mesh ", EUserInterfaceActionType::ToggleButton, FInputChord() );
 
-	UI_COMMAND( UseInGameBound, "In-game Bound", "Use in-game bound on preview mesh", EUserInterfaceActionType::ToggleButton, FInputChord() );
-
 	UI_COMMAND( SetCPUSkinning, "CPU Skinning", "Toggles display of CPU skinning in the Preview Pane.", EUserInterfaceActionType::ToggleButton, FInputChord());
 	UI_COMMAND( SetShowNormals, "Normals", "Toggles display of vertex normals in the Preview Pane.", EUserInterfaceActionType::ToggleButton, FInputChord() );
 	UI_COMMAND( SetShowTangents, "Tangents", "Toggles display of vertex tangents in the Preview Pane.", EUserInterfaceActionType::ToggleButton, FInputChord() );
diff --git a/Engine/Source/Editor/Persona/Private/AnimViewportMenuCommands.h b/Engine/Source/Editor/Persona/Private/AnimViewportMenuCommands.h
--- a/Engine/Source/Editor/Persona/Private/AnimViewportMenuCommands.h
+++ b/Engine/Source/Editor/Persona/Private/AnimViewportMenuCommands.h
@@ -21,18 +21,12 @@ public:
 	{
 	}
 
-	/** Preview what is active on the current window */
-	TSharedPtr< FUICommandInfo > Auto;
-
-	/** Lock to current preview */
-	TSharedPtr< FUICommandInfo > Lock;
-
+	/** Open settings for the preview scene */
+	TSharedPtr< FUICommandInfo > PreviewSceneSettings;
+	
 	/** Select camera follow */
 	TSharedPtr< FUICommandInfo > CameraFollow;
 
-	/** Use in-game Bound of preview mesh */
-	TSharedPtr< FUICommandInfo > UseInGameBound;
-
 	/** Show vertex normals */
 	TSharedPtr< FUICommandInfo > SetCPUSkinning;
 
diff --git a/Engine/Source/Editor/Persona/Private/AnimViewportShowCommands.cpp b/Engine/Source/Editor/Persona/Private/AnimViewportShowCommands.cpp
--- a/Engine/Source/Editor/Persona/Private/AnimViewportShowCommands.cpp
+++ b/Engine/Source/Editor/Persona/Private/AnimViewportShowCommands.cpp
@@ -14,9 +14,9 @@ void FAnimViewportShowCommands::RegisterCommands()
 
 	UI_COMMAND(ProcessRootMotion, "Process Root Motion", "Moves preview based on animation root motion", EUserInterfaceActionType::ToggleButton, FInputChord());
 
-	UI_COMMAND( ShowReferencePose, "Reference Pose", "Show reference pose on preview mesh", EUserInterfaceActionType::ToggleButton, FInputChord() );
 	UI_COMMAND( ShowRetargetBasePose, "Retarget Base Pose", "Show retarget Base pose on preview mesh", EUserInterfaceActionType::ToggleButton, FInputChord() );
 	UI_COMMAND( ShowBound, "Bound", "Show bound on preview mesh", EUserInterfaceActionType::ToggleButton, FInputChord() );
+	UI_COMMAND( UseInGameBound, "In-game Bound", "Use in-game bound on preview mesh", EUserInterfaceActionType::ToggleButton, FInputChord() );
 	UI_COMMAND( ShowPreviewMesh, "Mesh", "Show the preview mesh", EUserInterfaceActionType::ToggleButton, FInputChord() );
 	UI_COMMAND( ShowMorphTargets, "MorphTargets", "Display Applied Morph Targets of the mesh", EUserInterfaceActionType::ToggleButton, FInputChord() );
 
diff --git a/Engine/Source/Editor/Persona/Private/AnimViewportShowCommands.h b/Engine/Source/Editor/Persona/Private/AnimViewportShowCommands.h
--- a/Engine/Source/Editor/Persona/Private/AnimViewportShowCommands.h
+++ b/Engine/Source/Editor/Persona/Private/AnimViewportShowCommands.h
@@ -33,15 +33,15 @@ public:
 	/** Option to show root motion in viewport */
 	TSharedPtr< FUICommandInfo > ProcessRootMotion;
 
-	/** Show reference pose on preview mesh */
-	TSharedPtr< FUICommandInfo > ShowReferencePose;
-
 	/** Show reference pose on preview mesh */
 	TSharedPtr< FUICommandInfo > ShowRetargetBasePose;
 	
 	/** Show Bound of preview mesh */
 	TSharedPtr< FUICommandInfo > ShowBound;
 
+	/** Use in-game Bound of preview mesh */
+	TSharedPtr< FUICommandInfo > UseInGameBound;
+
 	/** Show/hide the preview mesh */
 	TSharedPtr< FUICommandInfo > ShowPreviewMesh;
 
