Fix compile warning: declaration shadows a local variable

diff --git a/Engine/Source/Editor/ContentBrowser/Private/SAssetView.cpp b/Engine/Source/Editor/ContentBrowser/Private/SAssetView.cpp
--- a/Engine/Source/Editor/ContentBrowser/Private/SAssetView.cpp
+++ b/Engine/Source/Editor/ContentBrowser/Private/SAssetView.cpp
@@ -2630,7 +2630,7 @@ void SAssetView::OnAssetRenamed(const FAssetData& AssetData, const FString& OldO
 	// Remove the old asset, if it exists
 	FName OldObjectPackageName = *OldObjectPath;
 	RemoveAssetByPath( OldObjectPackageName );
-	RecentlyAddedAssets.RemoveAllSwap( [&](const FAssetData& AssetData) { return AssetData.ObjectPath == OldObjectPackageName; } );
+	RecentlyAddedAssets.RemoveAllSwap( [&](const FAssetData& Other) { return Other.ObjectPath == OldObjectPackageName; } );
 
 	// Add the new asset, if it should be in the cached list
 	OnAssetAdded( AssetData );
