Merging //UE4/Dev-Main to Dev-Geometry (//UE4/Dev-Geometry)

diff --git a/Engine/Source/Editor/ContentBrowser/Private/SAssetView.cpp b/Engine/Source/Editor/ContentBrowser/Private/SAssetView.cpp
--- a/Engine/Source/Editor/ContentBrowser/Private/SAssetView.cpp
+++ b/Engine/Source/Editor/ContentBrowser/Private/SAssetView.cpp
@@ -2630,7 +2630,7 @@ void SAssetView::OnAssetRenamed(const FAssetData& AssetData, const FString& OldO
 	// Remove the old asset, if it exists
 	FName OldObjectPackageName = *OldObjectPath;
 	RemoveAssetByPath( OldObjectPackageName );
-	RecentlyAddedAssets.RemoveAllSwap( [&](const FAssetData& AssetData) { return AssetData.ObjectPath == OldObjectPackageName; } );
+	RecentlyAddedAssets.RemoveAllSwap( [&](const FAssetData& Other) { return Other.ObjectPath == OldObjectPackageName; } );
 
 	// Add the new asset, if it should be in the cached list
 	OnAssetAdded( AssetData );
diff --git a/Engine/Source/Editor/UnrealEd/Classes/Editor/Transactor.h b/Engine/Source/Editor/UnrealEd/Classes/Editor/Transactor.h
--- a/Engine/Source/Editor/UnrealEd/Classes/Editor/Transactor.h
+++ b/Engine/Source/Editor/UnrealEd/Classes/Editor/Transactor.h
@@ -170,6 +170,11 @@ protected:
 
 		struct FSerializedObject
 		{
+			FSerializedObject()
+				: bIsPendingKill(false)
+			{
+			}
+
 			void SetObject(const UObject* InObject)
 			{
 				ObjectName = InObject->GetFName();
diff --git a/Engine/Source/Programs/UnrealBuildTool/Configuration/UEBuildPlatform.cs b/Engine/Source/Programs/UnrealBuildTool/Configuration/UEBuildPlatform.cs
--- a/Engine/Source/Programs/UnrealBuildTool/Configuration/UEBuildPlatform.cs
+++ b/Engine/Source/Programs/UnrealBuildTool/Configuration/UEBuildPlatform.cs
@@ -43,7 +43,8 @@ namespace UnrealBuildTool
 			new FileSystemName("NoRedist"),
 			new FileSystemName("PS4"),
 			new FileSystemName("XboxOne"),
-			new FileSystemName("Switch")
+			new FileSystemName("Switch"),
+			new FileSystemName("Quail")
 		};
 
 		/// <summary>
