Fix crash when importing assets in a folder containing an invalid character for a package name

diff --git a/Engine/Source/Developer/AssetTools/Private/AssetTools.cpp b/Engine/Source/Developer/AssetTools/Private/AssetTools.cpp
--- a/Engine/Source/Developer/AssetTools/Private/AssetTools.cpp
+++ b/Engine/Source/Developer/AssetTools/Private/AssetTools.cpp
@@ -1476,7 +1476,7 @@ TArray<UObject*> UAssetToolsImpl::ImportAssetsInternal(const TArray<FString>& Fi
 			}
 			Name = ObjectTools::SanitizeObjectName(Name);
 
-			FString PackageName = FPaths::Combine(*DestinationPath, *Name);
+			FString PackageName = ObjectTools::SanitizeInvalidChars(FPaths::Combine(*DestinationPath, *Name), INVALID_LONGPACKAGE_CHARACTERS);
 
 			// We can not create assets that share the name of a map file in the same location
 			if(FEditorFileUtils::IsMapPackageAsset(PackageName))
diff --git a/Engine/Source/Editor/ContentBrowser/Private/SAssetView.cpp b/Engine/Source/Editor/ContentBrowser/Private/SAssetView.cpp
--- a/Engine/Source/Editor/ContentBrowser/Private/SAssetView.cpp
+++ b/Engine/Source/Editor/ContentBrowser/Private/SAssetView.cpp
@@ -1362,7 +1362,8 @@ FReply SAssetView::OnDrop( const FGeometry& MyGeometry, const FDragDropEvent& Dr
 						const FString& Filename = FilesAndDestinations[FileIdx].Key;
 						const FString& DestinationPath = FilesAndDestinations[FileIdx].Value;
 						FString Name = ObjectTools::SanitizeObjectName(FPaths::GetBaseFilename(Filename));
-						FString PackageName = DestinationPath + TEXT("/") + Name;
+						FString PackageName = ObjectTools::SanitizeInvalidChars(DestinationPath + TEXT("/") + Name, INVALID_LONGPACKAGE_CHARACTERS);
+
 
 						// We can not create assets that share the name of a map file in the same location
 						if (FEditorFileUtils::IsMapPackageAsset(PackageName))
