Fixed some places that were calling LoadModuleChecked during shutdown

diff --git a/Engine/Source/Editor/ContentBrowser/Private/EmptyFolderVisibilityManager.cpp b/Engine/Source/Editor/ContentBrowser/Private/EmptyFolderVisibilityManager.cpp
--- a/Engine/Source/Editor/ContentBrowser/Private/EmptyFolderVisibilityManager.cpp
+++ b/Engine/Source/Editor/ContentBrowser/Private/EmptyFolderVisibilityManager.cpp
@@ -30,8 +30,8 @@ FEmptyFolderVisibilityManager::~FEmptyFolderVisibilityManager()
 {
 	if (FModuleManager::Get().IsModuleLoaded(TEXT("AssetRegistry")))
 	{
-		// Load the asset registry module to stop listening for updates
-		FAssetRegistryModule& AssetRegistryModule = FModuleManager::LoadModuleChecked<FAssetRegistryModule>(TEXT("AssetRegistry"));
+		// Get the asset registry module to stop listening for updates
+		FAssetRegistryModule& AssetRegistryModule = FModuleManager::GetModuleChecked<FAssetRegistryModule>(TEXT("AssetRegistry"));
 		AssetRegistryModule.Get().OnPathRemoved().RemoveAll(this);
 		AssetRegistryModule.Get().OnAssetAdded().RemoveAll(this);
 	}
diff --git a/Engine/Source/Runtime/Engine/Private/Internationalization/EnginePackageLocalizationCache.cpp b/Engine/Source/Runtime/Engine/Private/Internationalization/EnginePackageLocalizationCache.cpp
--- a/Engine/Source/Runtime/Engine/Private/Internationalization/EnginePackageLocalizationCache.cpp
+++ b/Engine/Source/Runtime/Engine/Private/Internationalization/EnginePackageLocalizationCache.cpp
@@ -23,7 +23,7 @@ FEnginePackageLocalizationCache::~FEnginePackageLocalizationCache()
 {
 	if (FModuleManager::Get().IsModuleLoaded(TEXT("AssetRegistry")))
 	{
-		FAssetRegistryModule& AssetRegistryModule = FModuleManager::LoadModuleChecked<FAssetRegistryModule>(TEXT("AssetRegistry"));
+		FAssetRegistryModule& AssetRegistryModule = FModuleManager::GetModuleChecked<FAssetRegistryModule>(TEXT("AssetRegistry"));
 		IAssetRegistry& AssetRegistry = AssetRegistryModule.Get();
 
 		AssetRegistry.OnAssetAdded().RemoveAll(this);
