Fixed content browser filter failing to match certain names if asset paths were omitted from the search

diff --git a/Engine/Source/Editor/ContentBrowser/Private/FrontendFilters.cpp b/Engine/Source/Editor/ContentBrowser/Private/FrontendFilters.cpp
--- a/Engine/Source/Editor/ContentBrowser/Private/FrontendFilters.cpp
+++ b/Engine/Source/Editor/ContentBrowser/Private/FrontendFilters.cpp
@@ -273,12 +273,12 @@ public:
 	{
 		AssetPtr = InAsset;
 		
+		// Get the full asset path, and also split it so we can compare each part in the filter
 		AssetPtr->PackageName.AppendString(AssetFullPath);
 		AssetFullPath.ParseIntoArray(AssetSplitPath, TEXT("/"));
 
+		// Get the full export text path as people sometimes search by copying this (requires class and asset path search to be enabled in order to match)
 		AssetPtr->GetExportTextName(AssetExportTextName);
-		const TCHAR* Delimiters[] = { TEXT("/"), TEXT("'") };
-		AssetExportTextName.ParseIntoArray(AssetSplitExportText, &Delimiters[0], 2, true);
 
 		if (FCollectionManagerModule::IsModuleAvailable())
 		{
@@ -305,7 +305,6 @@ public:
 		AssetFullPath.Reset();
 		AssetExportTextName.Reset();
 		AssetSplitPath.Reset();
-		AssetSplitExportText.Reset();
 		AssetCollectionNames.Reset();
 	}
 
@@ -341,9 +340,18 @@ public:
 
 	virtual bool TestBasicStringExpression(const FTextFilterString& InValue, const ETextFilterTextComparisonMode InTextComparisonMode) const override
 	{
-		int32 NumPathParts = AssetSplitPath.Num();
+		if (TextFilterUtils::TestBasicStringExpression(AssetPtr->AssetName, InValue, InTextComparisonMode))
+		{
+			return true;
+		}
+
 		if (bIncludeAssetPath)
 		{
+			if (TextFilterUtils::TestBasicStringExpression(AssetFullPath, InValue, InTextComparisonMode))
+			{
+				return true;
+			}
+
 			for (const FString& AssetPathPart : AssetSplitPath)
 			{
 				if (TextFilterUtils::TestBasicStringExpression(AssetPathPart, InValue, InTextComparisonMode))
@@ -352,13 +360,6 @@ public:
 				}
 			}
 		}
-		else if (NumPathParts > 0)
-		{
-			if (TextFilterUtils::TestBasicStringExpression(AssetSplitPath[NumPathParts - 1], InValue, InTextComparisonMode))
-			{
-				return true;
-			}
-		}
 
 		if (bIncludeClassName)
 		{
@@ -366,14 +367,14 @@ public:
 			{
 				return true;
 			}
+		}
 
-			if (AssetSplitExportText.Num() > 0)
+		if (bIncludeClassName && bIncludeAssetPath)
+		{
+			// Only test this if we're searching the class name and asset path too, as the exported text contains the type and path in the string
+			if (TextFilterUtils::TestBasicStringExpression(AssetExportTextName, InValue, InTextComparisonMode))
 			{
-				// Only test this if we're searching the class name too, as the exported text contains the type in the string
-				if (TextFilterUtils::TestBasicStringExpression(AssetSplitExportText[0], InValue, InTextComparisonMode))
-				{
-					return true;
-				}
+				return true;
 			}
 		}
 
@@ -511,9 +512,6 @@ private:
 	/** Split path of the current asset */
 	TArray<FString> AssetSplitPath;
 
-	/** Split path of the Export Text */
-	TArray<FString> AssetSplitExportText;
-
 	/** Names of the collections that the current asset is in */
 	TArray<FName> AssetCollectionNames;
 
