Safe guard against shutdown crash when asset views tries to stream out textures during final shutdown gc

diff --git a/Engine/Source/Editor/ContentBrowser/Private/AssetViewWidgets.cpp b/Engine/Source/Editor/ContentBrowser/Private/AssetViewWidgets.cpp
--- a/Engine/Source/Editor/ContentBrowser/Private/AssetViewWidgets.cpp
+++ b/Engine/Source/Editor/ContentBrowser/Private/AssetViewWidgets.cpp
@@ -363,14 +363,18 @@ void SAssetColumnView::Tick(const FGeometry& AllottedGeometry, const double InCu
 
 SAssetViewItem::~SAssetViewItem()
 {
-	if ( AssetItem.IsValid() )
+	if (AssetItem.IsValid())
 	{
-		AssetItem->OnAssetDataChanged.RemoveAll( this );
+		AssetItem->OnAssetDataChanged.RemoveAll(this);
 	}
 
-	OnItemDestroyed.ExecuteIfBound( AssetItem );
+	OnItemDestroyed.ExecuteIfBound(AssetItem);
 
-	SetForceMipLevelsToBeResident(false);
+	if (!GExitPurge)
+	{
+		// This hack is here to make abnormal shutdowns less frequent.  Crashes here are the result UI's being shut down as a result of GC at edtior shutdown.  This code attemps to call FindObject which is not allowed during GC.
+		SetForceMipLevelsToBeResident(false);
+	}
 }
 
 void SAssetViewItem::Construct( const FArguments& InArgs )
