Fix one crash in SyncPathsFromSourceControl, add logging in hopes of clarifying another

diff --git a/Engine/Source/Editor/ContentBrowser/Private/ContentBrowserUtils.cpp b/Engine/Source/Editor/ContentBrowser/Private/ContentBrowserUtils.cpp
--- a/Engine/Source/Editor/ContentBrowser/Private/ContentBrowserUtils.cpp
+++ b/Engine/Source/Editor/ContentBrowser/Private/ContentBrowserUtils.cpp
@@ -1995,7 +1995,10 @@ void ShowSyncDependenciesDialog(const TArray<FString>& InDependencies, TArray<FS
 
 			for (UPackage* SelectedPackage : SelectedPackages)
 			{
-				OutExtraPackagesToSync.Emplace(SelectedPackage->GetName());
+				if (SelectedPackage)
+				{
+					OutExtraPackagesToSync.Emplace(SelectedPackage->GetName());
+				}
 			}
 		}
 	}
@@ -2171,6 +2174,12 @@ void ContentBrowserUtils::SyncPathsFromSourceControl(const TArray<FString>& Cont
 			return false; // keep package
 		});
 
+		UE_LOG(LogContentBrowser, Log, TEXT("Syncing %d path(s):"), ContentPaths.Num());
+		for (const UPackage* Package : LoadedPackages)
+		{
+			UE_LOG(LogContentBrowser, Log, TEXT("\t - %s"), *Package->GetName());
+		}
+
 		// Hot-reload the new packages...
 		PackageTools::ReloadPackages(LoadedPackages);
 
