PR #4924: UE-62113 Fix Other filters toggling all assets to show up (Contributed by [REDACTED_PERSON])

diff --git a/Engine/Source/Editor/ContentBrowser/Private/SFilterList.cpp b/Engine/Source/Editor/ContentBrowser/Private/SFilterList.cpp
--- a/Engine/Source/Editor/ContentBrowser/Private/SFilterList.cpp
+++ b/Engine/Source/Editor/ContentBrowser/Private/SFilterList.cpp
@@ -910,6 +910,7 @@ void SFilterList::RemoveFilter(const TSharedRef<SFilter>& FilterToRemove)
 		// Update the frontend filters collection
 		const TSharedRef<FFrontendFilter>& FrontendFilter = FilterToRemove->GetFrontendFilter().ToSharedRef();
 		SetFrontendFilterActive(FrontendFilter, false);
+		OnFilterChanged.ExecuteIfBound();
 	}
 }
 
