PR #4989: Add icons to folder context menu favorites (Contributed by [REDACTED_PERSON])

diff --git a/Engine/Source/Editor/ContentBrowser/Private/PathContextMenu.cpp b/Engine/Source/Editor/ContentBrowser/Private/PathContextMenu.cpp
--- a/Engine/Source/Editor/ContentBrowser/Private/PathContextMenu.cpp
+++ b/Engine/Source/Editor/ContentBrowser/Private/PathContextMenu.cpp
@@ -235,7 +235,7 @@ void FPathContextMenu::MakePathViewContextMenu(FMenuBuilder& MenuBuilder)
 				MenuBuilder.AddMenuEntry(
 					LOCTEXT("RemoveFromFavorites", "Remove From Favorites"),
 					LOCTEXT("RemoveFromFavoritesTooltip", "Removes this folder from the favorites section."),
-					FSlateIcon(),
+					FSlateIcon(FEditorStyle::GetStyleSetName(), "PropertyWindow.Favorites_Disabled"),
 					FUIAction(FExecuteAction::CreateSP(this, &FPathContextMenu::ExecuteFavorite))
 				);
 			}
@@ -245,7 +245,7 @@ void FPathContextMenu::MakePathViewContextMenu(FMenuBuilder& MenuBuilder)
 				MenuBuilder.AddMenuEntry(
 					LOCTEXT("AddToFavorites", "Add To Favorites"),
 					LOCTEXT("AddToFavoritesTooltip", "Adds this folder to the favorites section for easy access."),
-					FSlateIcon(),
+					FSlateIcon(FEditorStyle::GetStyleSetName(), "PropertyWindow.Favorites_Enabled"),
 					FUIAction(FExecuteAction::CreateSP(this, &FPathContextMenu::ExecuteFavorite))
 				);
 			}
