Fix content browser search filter by collection

diff --git a/Engine/Source/Editor/ContentBrowser/Private/FrontendFilters.cpp b/Engine/Source/Editor/ContentBrowser/Private/FrontendFilters.cpp
--- a/Engine/Source/Editor/ContentBrowser/Private/FrontendFilters.cpp
+++ b/Engine/Source/Editor/ContentBrowser/Private/FrontendFilters.cpp
@@ -267,7 +267,7 @@ public:
 		, TypeKeyName("Type")
 		, CollectionKeyName("Collection")
 		, TagKeyName("Tag")
-		, CollectionManager(FCollectionManagerModule::IsModuleAvailable() ? &FCollectionManagerModule::GetModule().Get() : nullptr)
+		, CollectionManager(nullptr)
 	{
 	}
 
@@ -290,7 +290,12 @@ public:
 			}
 		}
 
-		if (bIncludeCollectionNames && CollectionManager)
+		if (CollectionManager == nullptr)
+		{
+			CollectionManager = &FCollectionManagerModule::GetModule().Get();
+		}
+
+		if (CollectionManager)
 		{
 			CollectionManager->GetCollectionsContainingObject(AssetPtr->ObjectPath, ECollectionShareType::CST_All, AssetCollectionNames, ECollectionRecursionFlags::SelfAndChildren);
 
