Add Support for drag/drop of Collection on the Foliage Palette (sharing of FT Presets)
(Support was added to anything that uses AssetUtil::ExtractAssetDataFromDrag)

diff --git a/Engine/Source/Editor/ContentBrowser/Private/SAssetView.cpp b/Engine/Source/Editor/ContentBrowser/Private/SAssetView.cpp
--- a/Engine/Source/Editor/ContentBrowser/Private/SAssetView.cpp
+++ b/Engine/Source/Editor/ContentBrowser/Private/SAssetView.cpp
@@ -1601,20 +1601,24 @@ FReply SAssetView::OnDrop( const FGeometry& MyGeometry, const FDragDropEvent& Dr
 
 		if (SelectedAssetDatas.Num() > 0)
 		{
-			TArray<FName> ObjectPaths;
-			for (const auto& AssetData : SelectedAssetDatas)
+			TSharedPtr<FAssetDragDropOp> AssetDragDropOp = DragDropEvent.GetOperationAs< FAssetDragDropOp >();
+			if (AssetDragDropOp.IsValid())
 			{
-				if (!AssetData.GetClass()->IsChildOf(UClass::StaticClass()))
+				TArray<FName> ObjectPaths;
+				for (const auto& AssetData : SelectedAssetDatas)
 				{
-					ObjectPaths.Add(AssetData.ObjectPath);
+					if (!AssetData.GetClass()->IsChildOf(UClass::StaticClass()))
+					{
+						ObjectPaths.Add(AssetData.ObjectPath);
+					}
 				}
-			}
 
-			if (ObjectPaths.Num() > 0)
-			{
-				FCollectionManagerModule& CollectionManagerModule = FCollectionManagerModule::GetModule();
-				const FCollectionNameType& Collection = SourcesData.Collections[0];
-				CollectionManagerModule.Get().AddToCollection(Collection.Name, Collection.Type, ObjectPaths);
+				if (ObjectPaths.Num() > 0)
+				{
+					FCollectionManagerModule& CollectionManagerModule = FCollectionManagerModule::GetModule();
+					const FCollectionNameType& Collection = SourcesData.Collections[0];
+					CollectionManagerModule.Get().AddToCollection(Collection.Name, Collection.Type, ObjectPaths);
+				}	
 			}
 
 			return FReply::Handled();
diff --git a/Engine/Source/Editor/UnrealEd/Private/AssetSelection.cpp b/Engine/Source/Editor/UnrealEd/Private/AssetSelection.cpp
--- a/Engine/Source/Editor/UnrealEd/Private/AssetSelection.cpp
+++ b/Engine/Source/Editor/UnrealEd/Private/AssetSelection.cpp
@@ -28,6 +28,7 @@
 #include "ComponentAssetBroker.h"
 
 #include "DragAndDrop/AssetDragDropOp.h"
+#include "DragAndDrop/CollectionDragDropOp.h"
 
 #include "AssetRegistryModule.h"
 #include "IContentBrowserSingleton.h"
@@ -611,6 +612,11 @@ namespace AssetUtil
 				}
 			}
 		}
+		else if (Operation->IsOfType<FCollectionDragDropOp>())
+		{
+			TSharedPtr<FCollectionDragDropOp> DragDropOp = StaticCastSharedPtr<FCollectionDragDropOp>( Operation );
+			DroppedAssetData.Append( DragDropOp->GetAssets() );
+		}
 		else if (Operation->IsOfType<FAssetDragDropOp>())
 		{
 			TSharedPtr<FAssetDragDropOp> DragDropOp = StaticCastSharedPtr<FAssetDragDropOp>( Operation );
diff --git a/Engine/Source/Editor/UnrealEd/Private/DragAndDrop/CollectionDragDropOp.cpp b/Engine/Source/Editor/UnrealEd/Private/DragAndDrop/CollectionDragDropOp.cpp
new file mode 100644
--- /dev/null
+++ b/Engine/Source/Editor/UnrealEd/Private/DragAndDrop/CollectionDragDropOp.cpp
@@ -0,0 +1,33 @@
+#include "DragAndDrop/CollectionDragDropOp.h"
+#include "ICollectionManager.h"
+#include "CollectionManagerModule.h"
+#include "AssetRegistryModule.h"
+#include "IAssetRegistry.h"
+
+TArray<FAssetData>  FCollectionDragDropOp::GetAssets() const
+{
+	ICollectionManager& CollectionManager = FModuleManager::LoadModuleChecked<FCollectionManagerModule>("CollectionManager").Get();
+	FAssetRegistryModule& AssetRegistryModule = FModuleManager::LoadModuleChecked<FAssetRegistryModule>(TEXT("AssetRegistry"));
+	IAssetRegistry& AssetRegistry = AssetRegistryModule.Get();
+
+	TArray<FName> AssetPaths;
+	for (const FCollectionNameType& CollectionNameType : Collections)
+	{
+		TArray<FName> CollectionAssetPaths;
+		CollectionManager.GetAssetsInCollection(CollectionNameType.Name, CollectionNameType.Type, CollectionAssetPaths);
+		AssetPaths.Append(CollectionAssetPaths);
+	}
+
+	TArray<FAssetData> AssetDatas;
+	AssetDatas.Reserve(AssetPaths.Num());
+	for (const FName& AssetPath : AssetPaths)
+	{
+		FAssetData AssetData = AssetRegistry.GetAssetByObjectPath(AssetPath);
+		if (AssetData.IsValid())
+		{
+			AssetDatas.AddUnique(AssetData);
+		}
+	}
+	
+	return AssetDatas;
+}
\ No newline at end of file
diff --git a/Engine/Source/Editor/UnrealEd/Public/DragAndDrop/CollectionDragDropOp.h b/Engine/Source/Editor/UnrealEd/Public/DragAndDrop/CollectionDragDropOp.h
--- a/Engine/Source/Editor/UnrealEd/Public/DragAndDrop/CollectionDragDropOp.h
+++ b/Engine/Source/Editor/UnrealEd/Public/DragAndDrop/CollectionDragDropOp.h
@@ -12,6 +12,7 @@
 #include "Widgets/Images/SImage.h"
 #include "EditorStyleSet.h"
 #include "DragAndDrop/DecoratedDragDropOp.h"
+#include "AssetData.h"
 
 class FCollectionDragDropOp : public FDecoratedDragDropOp
 {
@@ -33,6 +34,9 @@ public:
 	}
 	
 public:
+	/** @return The assets from this drag operation */
+	TArray<FAssetData> GetAssets() const;
+
 	FText GetDecoratorText() const
 	{
 		if (CurrentHoverText.IsEmpty() && Collections.Num() > 0)
