PR #5758: Store the folder color to the new path location v2 (Contributed by [REDACTED_PERSON])

diff --git a/Engine/Source/Editor/ContentBrowser/Private/ContentBrowserUtils.cpp b/Engine/Source/Editor/ContentBrowser/Private/ContentBrowserUtils.cpp
--- a/Engine/Source/Editor/ContentBrowser/Private/ContentBrowserUtils.cpp
+++ b/Engine/Source/Editor/ContentBrowser/Private/ContentBrowserUtils.cpp
@@ -759,6 +759,42 @@ void ContentBrowserUtils::DisplayConfirmationPopup(const FText& Message, const F
 	Popup->OpenPopup(ParentContent);
 }
 
+bool ContentBrowserUtils::RenameFolder(const FString& DestPath, const FString& SourcePath)
+{
+	if (DestPath == SourcePath)
+	{
+		return false;
+	}
+
+	FAssetRegistryModule& AssetRegistryModule = FModuleManager::LoadModuleChecked<FAssetRegistryModule>("AssetRegistry");
+	// move any assets in our folder
+	TArray<FAssetData> AssetsInFolder;
+	AssetRegistryModule.Get().GetAssetsByPath(*SourcePath, AssetsInFolder, true);
+	TArray<UObject*> ObjectsInFolder;
+	GetObjectsInAssetData(AssetsInFolder, ObjectsInFolder);
+	MoveAssets(ObjectsInFolder, DestPath, SourcePath);
+
+	// Now check to see if the original folder is empty, if so we can delete it
+	TArray<FAssetData> AssetsInOriginalFolder;
+	AssetRegistryModule.Get().GetAssetsByPath(*SourcePath, AssetsInOriginalFolder, true);
+	if (AssetsInOriginalFolder.Num() == 0)
+	{
+		TArray<FString> FoldersToDelete;
+		FoldersToDelete.Add(SourcePath);
+		DeleteFolders(FoldersToDelete);
+	}
+
+	// set color of folder to new path
+	const TSharedPtr<FLinearColor> FolderColor = LoadColor(SourcePath);
+	if (FolderColor.IsValid())
+	{
+		SaveColor(SourcePath, nullptr);
+		SaveColor(DestPath, FolderColor);
+	}
+
+	return true;
+}
+
 bool ContentBrowserUtils::CopyFolders(const TArray<FString>& InSourcePathNames, const FString& DestPath)
 {
 	TMap<FString, TArray<UObject*> > SourcePathToLoadedAssets;
@@ -798,15 +834,10 @@ bool ContentBrowserUtils::CopyFolders(const TArray<FString>& InSourcePathNames,
 			ObjectTools::DuplicateObjects( PathIt.Value(), SourcePath, Destination, /*bOpenDialog=*/false );
 		}
 
-		// Attempt to copy the folder color to the new path location
-		if (FPaths::FileExists(GEditorPerProjectIni))
+		const TSharedPtr<FLinearColor> FolderColor = LoadColor(SourcePath);
+		if (FolderColor.IsValid())
 		{
-			FString ColorStr;
-			if (GConfig->GetString(TEXT("PathColor"), *SourcePath, ColorStr, GEditorPerProjectIni))
-			{
-				// Add the new path
-				GConfig->SetString(TEXT("PathColor"), *Destination, *ColorStr, GEditorPerProjectIni);
-			}
+			SaveColor(Destination, FolderColor);
 		}
 	}
 
@@ -861,7 +892,7 @@ bool ContentBrowserUtils::MoveFolders(const TArray<FString>& InSourcePathNames,
 		if ( PathIt.Value().Num() > 0 )
 		{
 			// Move assets and supply a source path to indicate it is relative
-			ContentBrowserUtils::MoveAssets( PathIt.Value(), Destination, PathIt.Key() );
+			MoveAssets( PathIt.Value(), Destination, PathIt.Key() );
 		}
 
 		// Attempt to remove the old path
@@ -870,18 +901,11 @@ bool ContentBrowserUtils::MoveFolders(const TArray<FString>& InSourcePathNames,
 			AssetRegistryModule.Get().RemovePath(SourcePath);
 		}
 
-		// Attempt to move the folder color to the new path location
-		if (FPaths::FileExists(GEditorPerProjectIni))
+		const TSharedPtr<FLinearColor> FolderColor = LoadColor(SourcePath);
+		if (FolderColor.IsValid())
 		{
-			FString ColorStr;
-			if (GConfig->GetString(TEXT("PathColor"), *SourcePath, ColorStr, GEditorPerProjectIni))
-			{
-				// Remove the old path
-				GConfig->RemoveKey(TEXT("PathColor"), *SourcePath, GEditorPerProjectIni);
-
-				// Add the new path
-				GConfig->SetString(TEXT("PathColor"), *Destination, *ColorStr, GEditorPerProjectIni);
-			}
+			SaveColor(SourcePath, nullptr);
+			SaveColor(Destination, FolderColor);
 		}
 	}
 
diff --git a/Engine/Source/Editor/ContentBrowser/Private/ContentBrowserUtils.h b/Engine/Source/Editor/ContentBrowser/Private/ContentBrowserUtils.h
--- a/Engine/Source/Editor/ContentBrowser/Private/ContentBrowserUtils.h
+++ b/Engine/Source/Editor/ContentBrowser/Private/ContentBrowserUtils.h
@@ -111,6 +111,9 @@ namespace ContentBrowserUtils
 	/** Displays a modeless message asking yes or no type question */
 	void DisplayConfirmationPopup(const FText& Message, const FText& YesString, const FText& NoString, const TSharedRef<SWidget>& ParentContent, const FOnClicked& OnYesClicked, const FOnClicked& OnNoClicked = FOnClicked());
 
+	/** Moves all assets from the source path to the destination path, preserving path structure, deletes source path afterwards if possible */
+	bool RenameFolder(const FString& DestPath, const FString& SourcePath);
+
 	/** Copies all assets in all source paths to the destination path, preserving path structure */
 	bool CopyFolders(const TArray<FString>& InSourcePathNames, const FString& DestPath);
 
diff --git a/Engine/Source/Editor/ContentBrowser/Private/SAssetView.cpp b/Engine/Source/Editor/ContentBrowser/Private/SAssetView.cpp
--- a/Engine/Source/Editor/ContentBrowser/Private/SAssetView.cpp
+++ b/Engine/Source/Editor/ContentBrowser/Private/SAssetView.cpp
@@ -4447,7 +4447,6 @@ void SAssetView::AssetRenameCommit(const TSharedPtr<FAssetViewItem>& Item, const
 	}
 	else if( ItemType == EAssetItemType::Folder )
 	{
-		TArray<FMovedContentFolder> MovedFolders;
 		const TSharedPtr<FAssetViewFolder>& ItemAsFolder = StaticCastSharedPtr<FAssetViewFolder>(Item);
 		if(ItemAsFolder->bNewFolder)
 		{
@@ -4510,27 +4509,12 @@ void SAssetView::AssetRenameCommit(const TSharedPtr<FAssetViewItem>& Item, const
 				}
 			}
 
-			if(bSuccess)
+			if(bSuccess && ContentBrowserUtils::RenameFolder(NewPath, ItemAsFolder->FolderPath))
 			{
+				TArray<FMovedContentFolder> MovedFolders;
 				MovedFolders.Add(FMovedContentFolder(ItemAsFolder->FolderPath, NewPath));
-				// move any assets in our folder
-				TArray<FAssetData> AssetsInFolder;
-				AssetRegistryModule.Get().GetAssetsByPath(*ItemAsFolder->FolderPath, AssetsInFolder, true);
-				TArray<UObject*> ObjectsInFolder;
-				ContentBrowserUtils::GetObjectsInAssetData(AssetsInFolder, ObjectsInFolder);
-				ContentBrowserUtils::MoveAssets(ObjectsInFolder, NewPath, ItemAsFolder->FolderPath);
-
-				// Now check to see if the original folder is empty, if so we can delete it
-				TArray<FAssetData> AssetsInOriginalFolder;
-				AssetRegistryModule.Get().GetAssetsByPath(*ItemAsFolder->FolderPath, AssetsInOriginalFolder, true);
-				if(AssetsInOriginalFolder.Num() == 0)
-				{
-					TArray<FString> FoldersToDelete;
-					FoldersToDelete.Add(ItemAsFolder->FolderPath);
-					ContentBrowserUtils::DeleteFolders(FoldersToDelete);
-				}
+				OnFolderPathChanged.ExecuteIfBound(MovedFolders);
 			}
-			OnFolderPathChanged.ExecuteIfBound(MovedFolders);
 			RequestQuickFrontendListRefresh();
 		}		
 	}
diff --git a/Engine/Source/Editor/ContentBrowser/Private/SPathView.cpp b/Engine/Source/Editor/ContentBrowser/Private/SPathView.cpp
--- a/Engine/Source/Editor/ContentBrowser/Private/SPathView.cpp
+++ b/Engine/Source/Editor/ContentBrowser/Private/SPathView.cpp
@@ -1213,7 +1213,6 @@ bool SPathView::VerifyFolderNameChanged(const FString& InName, FText& OutErrorMe
 
 void SPathView::FolderNameChanged( const TSharedPtr< FTreeItem >& TreeItem, const FString& OldPath, const FVector2D& MessageLocation, const ETextCommit::Type CommitType )
 {
-	TArray<FMovedContentFolder> MovedFolders;
 	// Verify the name of the folder
 	FText Reason;
 	if ( ContentBrowserUtils::IsValidFolderName(TreeItem->FolderName, Reason) )
@@ -1281,25 +1280,10 @@ void SPathView::FolderNameChanged( const TSharedPtr< FTreeItem >& TreeItem, cons
 				EmptyFolderVisibilityManager->SetAlwaysShowPath(TreeItem->FolderPath);
 
 				FAssetRegistryModule& AssetRegistryModule = FModuleManager::LoadModuleChecked<FAssetRegistryModule>("AssetRegistry");
-				if (AssetRegistryModule.Get().AddPath(TreeItem->FolderPath) && TreeItem->FolderPath != OldPath)
+				if (AssetRegistryModule.Get().AddPath(TreeItem->FolderPath) && ContentBrowserUtils::RenameFolder(TreeItem->FolderPath, OldPath))
 				{
+					TArray<FMovedContentFolder> MovedFolders;
 					MovedFolders.Add(FMovedContentFolder(OldPath, TreeItem->FolderPath));
-					// move any assets in our folder
-					TArray<FAssetData> AssetsInFolder;
-					AssetRegistryModule.Get().GetAssetsByPath(*OldPath, AssetsInFolder, true);
-					TArray<UObject*> ObjectsInFolder;
-					ContentBrowserUtils::GetObjectsInAssetData(AssetsInFolder, ObjectsInFolder);
-					ContentBrowserUtils::MoveAssets(ObjectsInFolder, TreeItem->FolderPath, OldPath);
-
-					// Now check to see if the original folder is empty, if so we can delete it
-					TArray<FAssetData> AssetsInOriginalFolder;
-					AssetRegistryModule.Get().GetAssetsByPath(*OldPath, AssetsInOriginalFolder, true);
-					if (AssetsInOriginalFolder.Num() == 0)
-					{
-						TArray<FString> FoldersToDelete;
-						FoldersToDelete.Add(OldPath);
-						ContentBrowserUtils::DeleteFolders(FoldersToDelete);
-					}
 					OnFolderPathChanged.ExecuteIfBound(MovedFolders);
 				}
 			}
