Disallowed using Square Brackets in a Folder Name to prevent Meshes from being unusable inside the Folder.

diff --git a/Engine/Source/Editor/ContentBrowser/Private/ContentBrowserUtils.cpp b/Engine/Source/Editor/ContentBrowser/Private/ContentBrowserUtils.cpp
--- a/Engine/Source/Editor/ContentBrowser/Private/ContentBrowserUtils.cpp
+++ b/Engine/Source/Editor/ContentBrowser/Private/ContentBrowserUtils.cpp
@@ -1303,7 +1303,7 @@ bool ContentBrowserUtils::IsValidFolderName(const FString& FolderName, FText& Re
 		return false;
 	}
 
-	const FString InvalidChars = INVALID_LONGPACKAGE_CHARACTERS TEXT("/"); // Slash is an invalid character for a folder name
+	const FString InvalidChars = INVALID_LONGPACKAGE_CHARACTERS TEXT("/[]"); // Slash is an invalid character for a folder name
 
 	// See if the name contains invalid characters.
 	FString Char;
