Add an experimental setting to enable support for long paths in the content browser.

diff --git a/Engine/Source/Editor/ContentBrowser/Private/AssetViewWidgets.cpp b/Engine/Source/Editor/ContentBrowser/Private/AssetViewWidgets.cpp
--- a/Engine/Source/Editor/ContentBrowser/Private/AssetViewWidgets.cpp
+++ b/Engine/Source/Editor/ContentBrowser/Private/AssetViewWidgets.cpp
@@ -692,8 +692,9 @@ TSharedRef<SWidget> SAssetViewItem::CreateToolTipWidget() const
 			{
 				int32 PackageNameLengthForCooking = ContentBrowserUtils::GetPackageLengthForCooking(AssetData.PackageName.ToString(), FEngineBuildSettings::IsInternalBuild());
 
+				int32 MaxCookPathLen = ContentBrowserUtils::GetMaxCookPathLen();
 				AddToToolTipInfoBox(InfoBox, LOCTEXT("TileViewTooltipPathLengthForCookingKey", "Cooking Filepath Length"), FText::Format(LOCTEXT("TileViewTooltipPathLengthForCookingValue", "{0} / {1}"),
-					FText::AsNumber(PackageNameLengthForCooking), FText::AsNumber(ContentBrowserUtils::MaxCookPathLen)), PackageNameLengthForCooking > ContentBrowserUtils::MaxCookPathLen ? true : false);
+					FText::AsNumber(PackageNameLengthForCooking), FText::AsNumber(MaxCookPathLen)), PackageNameLengthForCooking > MaxCookPathLen ? true : false);
 			}
 			else
 			{
diff --git a/Engine/Source/Editor/ContentBrowser/Private/ContentBrowserUtils.cpp b/Engine/Source/Editor/ContentBrowser/Private/ContentBrowserUtils.cpp
--- a/Engine/Source/Editor/ContentBrowser/Private/ContentBrowserUtils.cpp
+++ b/Engine/Source/Editor/ContentBrowser/Private/ContentBrowserUtils.cpp
@@ -41,6 +41,7 @@
 #include "AssetToolsModule.h"
 #include "NativeClassHierarchy.h"
 #include "EmptyFolderVisibilityManager.h"
+#include "Settings/EditorExperimentalSettings.h"
 
 #include "Toolkits/AssetEditorManager.h"
 #include "PackagesDialog.h"
@@ -1775,6 +1776,9 @@ bool ContentBrowserUtils::IsValidFolderPathForCreate(const FString& InFolderPath
 
 int32 ContentBrowserUtils::GetPackageLengthForCooking(const FString& PackageName, bool IsInternalBuild)
 {
+	// We assume the game name is 20 characters (the maximum allowed) to make sure that content can be ported between projects
+	static const int32 MaxGameNameLen = 20;
+
 	// Pad out the game name to the maximum allowed
 	const FString GameName = FApp::GetProjectName();
 	FString GameNamePadded = GameName;
@@ -1839,6 +1843,7 @@ bool ContentBrowserUtils::IsValidPackageForCooking(const FString& PackageName, F
 {
 	int32 AbsoluteCookPathToAssetLength = GetPackageLengthForCooking(PackageName, FEngineBuildSettings::IsInternalBuild());
 
+	int32 MaxCookPathLen = GetMaxCookPathLen();
 	if (AbsoluteCookPathToAssetLength > MaxCookPathLen)
 	{
 		// See TTP# 332328:
@@ -2283,4 +2288,18 @@ const TArray<FString>& ContentBrowserUtils::GetFavoriteFolders()
 	return FContentBrowserSingleton::Get().FavoriteFolderPaths;
 }
 
+int32 ContentBrowserUtils::GetMaxCookPathLen()
+{
+	if (GetDefault<UEditorExperimentalSettings>()->bEnableLongPathsSupport)
+	{
+		// Allow the longest path allowed by the system
+		return FPlatformMisc::GetMaxPathLength();
+	}
+	else
+	{
+		// 260 characters is the limit on Windows, which is the shortest max path of any platforms that support cooking
+		return 260;
+	}
+}
+
 #undef LOCTEXT_NAMESPACE
diff --git a/Engine/Source/Editor/ContentBrowser/Private/ContentBrowserUtils.h b/Engine/Source/Editor/ContentBrowser/Private/ContentBrowserUtils.h
--- a/Engine/Source/Editor/ContentBrowser/Private/ContentBrowserUtils.h
+++ b/Engine/Source/Editor/ContentBrowser/Private/ContentBrowserUtils.h
@@ -281,9 +281,7 @@ namespace ContentBrowserUtils
 
 	const TArray<FString>& GetFavoriteFolders();
 
-	// We assume the game name is 20 characters (the maximum allowed) to make sure that content can be ported between projects
-	// 260 characters is the limit on Windows, which is the shortest max path of any platforms that support cooking
-	static const int32 MaxGameNameLen = 20;
-	static const int32 MaxCookPathLen = 260;
+	/** Gets the maximum path length for a cooked file. Changes behavior based on whether the editor experimental setting for long paths is enabled. */
+	int32 GetMaxCookPathLen();
 }
 
diff --git a/Engine/Source/Editor/UnrealEd/Classes/Settings/EditorExperimentalSettings.h b/Engine/Source/Editor/UnrealEd/Classes/Settings/EditorExperimentalSettings.h
--- a/Engine/Source/Editor/UnrealEd/Classes/Settings/EditorExperimentalSettings.h
+++ b/Engine/Source/Editor/UnrealEd/Classes/Settings/EditorExperimentalSettings.h
@@ -125,6 +125,10 @@ public:
 	UPROPERTY(EditAnywhere, config, Category = Materials)
 	bool bExampleLayersAndBlends;
 
+	/** Allows creation of assets with paths longer than 260 characters. Note that this also requires the Windows 10 Anniversary Update (1607), and support for long paths to be enabled through the group policy editor. */
+	UPROPERTY(EditAnywhere, config, Category = "Content Browser", meta = (DisplayName = "Enable support for long paths (> 260 characters)"))
+	bool bEnableLongPathsSupport;
+
 	/**
 	 * Returns an event delegate that is executed when a setting has changed.
 	 *
