Context menu opening cancels a rename

diff --git a/Engine/Source/Editor/ContentBrowser/Private/AssetViewTypes.h b/Engine/Source/Editor/ContentBrowser/Private/AssetViewTypes.h
--- a/Engine/Source/Editor/ContentBrowser/Private/AssetViewTypes.h
+++ b/Engine/Source/Editor/ContentBrowser/Private/AssetViewTypes.h
@@ -43,6 +43,9 @@ struct FAssetViewItem
 	/** Broadcasts whenever a rename is requested */
 	FSimpleDelegate RenamedRequestEvent;
 
+	/** Broadcasts whenever a rename is canceled */
+	FSimpleDelegate RenameCanceledEvent;
+
 	/** An event to fire when the asset data for this item changes */
 	DECLARE_MULTICAST_DELEGATE( FOnAssetDataChanged );
 	FOnAssetDataChanged OnAssetDataChanged;
diff --git a/Engine/Source/Editor/ContentBrowser/Private/AssetViewWidgets.cpp b/Engine/Source/Editor/ContentBrowser/Private/AssetViewWidgets.cpp
--- a/Engine/Source/Editor/ContentBrowser/Private/AssetViewWidgets.cpp
+++ b/Engine/Source/Editor/ContentBrowser/Private/AssetViewWidgets.cpp
@@ -1556,7 +1556,8 @@ void SAssetListItem::Construct( const FArguments& InArgs )
 
 	if(AssetItem.IsValid())
 	{
-		AssetItem->RenamedRequestEvent.BindSP( InlineRenameWidget.Get(), &SInlineEditableTextBlock::EnterEditingMode );
+		AssetItem->RenamedRequestEvent.BindSP(InlineRenameWidget.Get(), &SInlineEditableTextBlock::EnterEditingMode);
+		AssetItem->RenameCanceledEvent.BindSP(InlineRenameWidget.Get(), &SInlineEditableTextBlock::ExitEditingMode);
 	}
 
 	SetForceMipLevelsToBeResident(true);
@@ -1695,7 +1696,8 @@ void SAssetTileItem::Construct( const FArguments& InArgs )
 
 	if(AssetItem.IsValid())
 	{
-		AssetItem->RenamedRequestEvent.BindSP( InlineRenameWidget.Get(), &SInlineEditableTextBlock::EnterEditingMode );
+		AssetItem->RenamedRequestEvent.BindSP(InlineRenameWidget.Get(), &SInlineEditableTextBlock::EnterEditingMode);
+		AssetItem->RenameCanceledEvent.BindSP(InlineRenameWidget.Get(), &SInlineEditableTextBlock::ExitEditingMode);
 	}
 
 	SetForceMipLevelsToBeResident(true);
@@ -1913,7 +1915,8 @@ TSharedRef<SWidget> SAssetColumnItem::GenerateWidgetForColumn( const FName& Colu
 
 		if(AssetItem.IsValid())
 		{
-			AssetItem->RenamedRequestEvent.BindSP( InlineRenameWidget.Get(), &SInlineEditableTextBlock::EnterEditingMode );
+			AssetItem->RenamedRequestEvent.BindSP(InlineRenameWidget.Get(), &SInlineEditableTextBlock::EnterEditingMode);
+			AssetItem->RenameCanceledEvent.BindSP(InlineRenameWidget.Get(), &SInlineEditableTextBlock::ExitEditingMode);
 		}
 
 		return SNew(SBorder)
diff --git a/Engine/Source/Editor/ContentBrowser/Private/SAssetView.cpp b/Engine/Source/Editor/ContentBrowser/Private/SAssetView.cpp
--- a/Engine/Source/Editor/ContentBrowser/Private/SAssetView.cpp
+++ b/Engine/Source/Editor/ContentBrowser/Private/SAssetView.cpp
@@ -3948,6 +3948,12 @@ TSharedPtr<SWidget> SAssetView::OnGetContextMenuContent()
 {
 	if ( CanOpenContextMenu() )
 	{
+		if (IsRenamingAsset())
+		{
+			RenamingAsset.Pin()->RenameCanceledEvent.ExecuteIfBound();
+			RenamingAsset.Reset();
+		}
+
 		const TArray<FString> SelectedFolders = GetSelectedFolders();
 		if(SelectedFolders.Num() > 0)
 		{
