Fix recently used filter when asset list contains classes

diff --git a/Engine/Source/Editor/ContentBrowser/Private/FrontendFilters.cpp b/Engine/Source/Editor/ContentBrowser/Private/FrontendFilters.cpp
--- a/Engine/Source/Editor/ContentBrowser/Private/FrontendFilters.cpp
+++ b/Engine/Source/Editor/ContentBrowser/Private/FrontendFilters.cpp
@@ -1207,6 +1207,14 @@ void FFrontendFilter_Recent::ActiveStateChanged(bool bActive)
 
 bool FFrontendFilter_Recent::PassesFilter(FAssetFilterType InItem) const
 {
+	static FName ClassName(TEXT("Class"));
+
+	// Exclude classes because they will not have a valid long package name and cause an assert in FindMRUItemIdx
+	if (ClassName == InItem.AssetClass)
+	{
+		return false;
+	}
+
 	FString PackagePath = InItem.PackageName.ToString();
 	FContentBrowserModule& CBModule = FModuleManager::LoadModuleChecked<FContentBrowserModule>(TEXT("ContentBrowser"));
 	FMainMRUFavoritesList* RecentlyOpenedAssets = CBModule.GetRecentlyOpenedAssets();
