PR #4927: UE-62113: Call OnFilterChanged for other scenarios (Contributed by [REDACTED_PERSON])

diff --git a/Engine/Source/Editor/ContentBrowser/Private/SFilterList.cpp b/Engine/Source/Editor/ContentBrowser/Private/SFilterList.cpp
--- a/Engine/Source/Editor/ContentBrowser/Private/SFilterList.cpp
+++ b/Engine/Source/Editor/ContentBrowser/Private/SFilterList.cpp
@@ -164,12 +164,15 @@ public:
 	}
 
 	/** Sets whether or not this filter is applied to the combined filter */
-	void SetEnabled(bool InEnabled)
+	void SetEnabled(bool InEnabled, bool InExecuteOnFilterChanged = true)
 	{
 		if ( InEnabled != bEnabled)
 		{
 			bEnabled = InEnabled;
-			OnFilterChanged.ExecuteIfBound();
+			if (InExecuteOnFilterChanged)
+			{
+				OnFilterChanged.ExecuteIfBound();
+			}
 		}
 	}
 
@@ -563,25 +566,21 @@ void SFilterList::DisableAllFilters()
 {
 	for ( auto FilterIt = Filters.CreateConstIterator(); FilterIt; ++FilterIt )
 	{
-		(*FilterIt)->SetEnabled(false);
+		(*FilterIt)->SetEnabled(false, false);
 	}
+
+	OnFilterChanged.ExecuteIfBound();
 }
 
 void SFilterList::RemoveAllFilters()
 {
 	if ( HasAnyFilters() )
 	{
-		bool bBackendFilterChanged = false;
-
 		for ( auto FilterIt = Filters.CreateConstIterator(); FilterIt; ++FilterIt )
 		{
 			const TSharedRef<SFilter>& FilterToRemove = *FilterIt;
 
-			if ( FilterToRemove->GetAssetTypeActions().IsValid() )
-			{
-				bBackendFilterChanged = true;
-			}
-			else if (FilterToRemove->GetFrontendFilter().IsValid() )
+			if (FilterToRemove->GetFrontendFilter().IsValid() )
 			{
 				// Update the frontend filters collection
 				const TSharedRef<FFrontendFilter>& FrontendFilter = FilterToRemove->GetFrontendFilter().ToSharedRef();
@@ -592,11 +591,8 @@ void SFilterList::RemoveAllFilters()
 		FilterBox->ClearChildren();
 		Filters.Empty();
 
-		if ( bBackendFilterChanged )
-		{
-			// Notify that a backend filter changed
-			OnFilterChanged.ExecuteIfBound();
-		}
+		// Notify that a filter has changed
+		OnFilterChanged.ExecuteIfBound();
 	}
 }
 
@@ -617,6 +613,7 @@ void SFilterList::DisableFiltersThatHideAssets(const TArray<FAssetData>& AssetDa
 
 		// Iterate over all enabled filters and disable any frontend filters that would hide any of the supplied assets
 		// and disable all backend filters if it was determined that the combined backend filter hides any of the assets
+		bool ExecuteOnFilteChanged = false;
 		for ( auto FilterIt = Filters.CreateConstIterator(); FilterIt; ++FilterIt )
 		{
 			const TSharedRef<SFilter>& Filter = *FilterIt;
@@ -630,7 +627,8 @@ void SFilterList::DisableFiltersThatHideAssets(const TArray<FAssetData>& AssetDa
 						if (!FrontendFilter->IsInverseFilter() && !FrontendFilter->PassesFilter(*AssetIt))
 						{
 							// This is a frontend filter and at least one asset did not pass.
-							Filter->SetEnabled(false);
+							Filter->SetEnabled(false, false);
+							ExecuteOnFilteChanged = true;
 						}
 					}
 				}
@@ -640,11 +638,17 @@ void SFilterList::DisableFiltersThatHideAssets(const TArray<FAssetData>& AssetDa
 					FARFilter BackendFilter = Filter->GetBackendFilter();
 					if ( !BackendFilter.IsEmpty() )
 					{
-						Filter->SetEnabled(false);
+						Filter->SetEnabled(false, false);
+						ExecuteOnFilteChanged = true;
 					}
 				}
 			}
 		}
+
+		if (ExecuteOnFilteChanged)
+		{
+			OnFilterChanged.ExecuteIfBound();
+		}
 	}
 }
 
@@ -743,7 +747,7 @@ void SFilterList::LoadSettings(const FString& IniFilename, const FString& IniSec
 
 					if ( EnabledTypeFilterNames.Contains(ClassName) )
 					{
-						NewFilter->SetEnabled(true);
+						NewFilter->SetEnabled(true, false);
 					}
 				}
 			}
@@ -776,7 +780,7 @@ void SFilterList::LoadSettings(const FString& IniFilename, const FString& IniSec
 
 					if ( EnabledFrontendFilterNames.Contains(FilterName) )
 					{
-						NewFilter->SetEnabled(true);
+						NewFilter->SetEnabled(true, false);
 					}
 				}
 			}
@@ -785,6 +789,8 @@ void SFilterList::LoadSettings(const FString& IniFilename, const FString& IniSec
 			FrontendFilter->LoadSettings(IniFilename, IniSection, CustomSettingsString);
 		}
 	}
+
+	OnFilterChanged.ExecuteIfBound();
 }
 
 void SFilterList::SetFrontendFilterActive(const TSharedRef<FFrontendFilter>& Filter, bool bActive)
@@ -812,7 +818,7 @@ TSharedRef<SFilter> SFilterList::AddFilter(const TWeakPtr<IAssetTypeActions>& As
 		SNew(SFilter)
 		.AssetTypeActions(AssetTypeActions)
 		.OnFilterChanged(OnFilterChanged)
-		.OnRequestRemove(this, &SFilterList::RemoveFilter)
+		.OnRequestRemove(this, &SFilterList::RemoveFilterAndUpdate)
 		.OnRequestEnableOnly(this, &SFilterList::EnableOnlyThisFilter)
 		.OnRequestDisableAll(this, &SFilterList::DisableAllFilters)
 		.OnRequestRemoveAll(this, &SFilterList::RemoveAllFilters);
@@ -828,7 +834,7 @@ TSharedRef<SFilter> SFilterList::AddFilter(const TSharedRef<FFrontendFilter>& Fr
 		SNew(SFilter)
 		.FrontendFilter(FrontendFilter)
 		.OnFilterChanged( this, &SFilterList::FrontendFilterChanged, FrontendFilter )
-		.OnRequestRemove(this, &SFilterList::RemoveFilter)
+		.OnRequestRemove(this, &SFilterList::RemoveFilterAndUpdate)
 		.OnRequestDisableAll(this, &SFilterList::DisableAllFilters)
 		.OnRequestRemoveAll(this, &SFilterList::RemoveAllFilters);
 
@@ -848,7 +854,7 @@ void SFilterList::AddFilter(const TSharedRef<SFilter>& FilterToAdd)
 	];
 }
 
-void SFilterList::RemoveFilter(const TWeakPtr<IAssetTypeActions>& AssetTypeActions)
+void SFilterList::RemoveFilter(const TWeakPtr<IAssetTypeActions>& AssetTypeActions, bool ExecuteOnFilterChanged)
 {
 	TSharedPtr<SFilter> FilterToRemove;
 	for ( auto FilterIt = Filters.CreateConstIterator(); FilterIt; ++FilterIt )
@@ -863,7 +869,14 @@ void SFilterList::RemoveFilter(const TWeakPtr<IAssetTypeActions>& AssetTypeActio
 
 	if ( FilterToRemove.IsValid() )
 	{
-		RemoveFilter(FilterToRemove.ToSharedRef());
+		if (ExecuteOnFilterChanged)
+		{
+			RemoveFilterAndUpdate(FilterToRemove.ToSharedRef());
+		}
+		else
+		{
+			RemoveFilter(FilterToRemove.ToSharedRef());
+		}
 	}
 }
 
@@ -872,11 +885,13 @@ void SFilterList::EnableOnlyThisFilter(const TSharedRef<SFilter>& FilterToEnable
 	for ( auto FilterIt = Filters.CreateConstIterator(); FilterIt; ++FilterIt )
 	{
 		bool bEnable = *FilterIt == FilterToEnable;
-		(*FilterIt)->SetEnabled(bEnable);
+		(*FilterIt)->SetEnabled(bEnable, false);
 	}
+
+	OnFilterChanged.ExecuteIfBound();
 }
 
-void SFilterList::RemoveFilter(const TSharedRef<FFrontendFilter>& FrontendFilter)
+void SFilterList::RemoveFilter(const TSharedRef<FFrontendFilter>& FrontendFilter, bool ExecuteOnFilterChanged)
 {
 	TSharedPtr<SFilter> FilterToRemove;
 	for ( auto FilterIt = Filters.CreateConstIterator(); FilterIt; ++FilterIt )
@@ -891,7 +906,14 @@ void SFilterList::RemoveFilter(const TSharedRef<FFrontendFilter>& FrontendFilter
 
 	if ( FilterToRemove.IsValid() )
 	{
-		RemoveFilter(FilterToRemove.ToSharedRef());
+		if (ExecuteOnFilterChanged)
+		{
+			RemoveFilterAndUpdate(FilterToRemove.ToSharedRef());
+		}
+		else
+		{
+			RemoveFilter(FilterToRemove.ToSharedRef());
+		}
 	}
 }
 
@@ -900,12 +922,7 @@ void SFilterList::RemoveFilter(const TSharedRef<SFilter>& FilterToRemove)
 	FilterBox->RemoveSlot(FilterToRemove);
 	Filters.Remove(FilterToRemove);
 
-	if ( FilterToRemove->GetAssetTypeActions().IsValid() )
-	{
-		// Notify that a backend filter changed
-		OnFilterChanged.ExecuteIfBound();
-	}
-	else if (FilterToRemove->GetFrontendFilter().IsValid() )
+	if (FilterToRemove->GetFrontendFilter().IsValid() )
 	{
 		// Update the frontend filters collection
 		const TSharedRef<FFrontendFilter>& FrontendFilter = FilterToRemove->GetFrontendFilter().ToSharedRef();
@@ -914,6 +931,14 @@ void SFilterList::RemoveFilter(const TSharedRef<SFilter>& FilterToRemove)
 	}
 }
 
+void SFilterList::RemoveFilterAndUpdate(const TSharedRef<SFilter>& FilterToRemove)
+{
+	RemoveFilter(FilterToRemove);
+
+	// Notify that a filter has changed
+	OnFilterChanged.ExecuteIfBound();
+}
+
 void SFilterList::FrontendFilterChanged(TSharedRef<FFrontendFilter> FrontendFilter)
 {
 	TSharedPtr<SFilter> FilterToUpdate;
@@ -930,6 +955,8 @@ void SFilterList::FrontendFilterChanged(TSharedRef<FFrontendFilter> FrontendFilt
 	if ( FilterToUpdate.IsValid() )
 	{
 		SetFrontendFilterActive(FrontendFilter, FilterToUpdate->IsEnabled());
+
+		OnFilterChanged.ExecuteIfBound();
 	}
 }
 
@@ -1243,6 +1270,7 @@ void SFilterList::FilterByTypeCategoryClicked(EAssetTypeCategories::Type Categor
 	GetTypeActionsForCategory(Category, TypeActionsList);
 
 	bool bFullCategoryInUse = IsAssetTypeCategoryInUse(Category);
+	bool ExecuteOnFilterChanged = false;
 	for ( auto TypeIt = TypeActionsList.CreateConstIterator(); TypeIt; ++TypeIt )
 	{
 		auto AssetTypeActions = (*TypeIt);
@@ -1251,14 +1279,21 @@ void SFilterList::FilterByTypeCategoryClicked(EAssetTypeCategories::Type Categor
 			if ( bFullCategoryInUse )
 			{
 				RemoveFilter(AssetTypeActions);
+				ExecuteOnFilterChanged = true;
 			}
 			else if ( !IsAssetTypeActionsInUse(AssetTypeActions) )
 			{
 				TSharedRef<SFilter> NewFilter = AddFilter(AssetTypeActions);
-				NewFilter->SetEnabled(true);
+				NewFilter->SetEnabled(true, false);
+				ExecuteOnFilterChanged = true;
 			}
 		}
 	}
+
+	if (ExecuteOnFilterChanged)
+	{
+		OnFilterChanged.ExecuteIfBound();
+	}
 }
 
 bool SFilterList::IsAssetTypeCategoryInUse(EAssetTypeCategories::Type Category) const
@@ -1331,6 +1366,7 @@ bool SFilterList::IsFrontendFilterInUse(TSharedRef<FFrontendFilter> FrontendFilt
 void SFilterList::FrontendFilterCategoryClicked(TSharedPtr<FFrontendFilterCategory> MenuCategory)
 {
 	bool bFullCategoryInUse = IsFrontendFilterCategoryInUse(MenuCategory);
+	bool ExecuteOnFilterChanged = false;
 	for ( auto FrontendFilterIt = AllFrontendFilters.CreateConstIterator(); FrontendFilterIt; ++FrontendFilterIt )
 	{
 		const TSharedRef<FFrontendFilter>& FrontendFilter = *FrontendFilterIt;
@@ -1339,15 +1375,23 @@ void SFilterList::FrontendFilterCategoryClicked(TSharedPtr<FFrontendFilterCatego
 		{
 			if ( bFullCategoryInUse )
 			{
-				RemoveFilter( FrontendFilter );
+				RemoveFilter( FrontendFilter, false );
+				ExecuteOnFilterChanged = true;
 			}
 			else if ( !IsFrontendFilterInUse( FrontendFilter ) )
 			{
 				TSharedRef<SFilter> NewFilter = AddFilter( FrontendFilter );
-				NewFilter->SetEnabled(true);
+				NewFilter->SetEnabled(true, false);
+				SetFrontendFilterActive(FrontendFilter, NewFilter->IsEnabled());
+				ExecuteOnFilterChanged = true;
 			}
 		}
 	}
+
+	if (ExecuteOnFilterChanged)
+	{
+		OnFilterChanged.ExecuteIfBound();
+	}
 }
 
 bool SFilterList::IsFrontendFilterCategoryInUse(TSharedPtr<FFrontendFilterCategory> MenuCategory) const
diff --git a/Engine/Source/Editor/ContentBrowser/Private/SFilterList.h b/Engine/Source/Editor/ContentBrowser/Private/SFilterList.h
--- a/Engine/Source/Editor/ContentBrowser/Private/SFilterList.h
+++ b/Engine/Source/Editor/ContentBrowser/Private/SFilterList.h
@@ -85,9 +85,10 @@ private:
 	void AddFilter(const TSharedRef<SFilter>& FilterToAdd);
 
 	/** Handler for when the remove filter button was clicked on a filter */
-	void RemoveFilter(const TWeakPtr<IAssetTypeActions>& AssetTypeActions);
-	void RemoveFilter(const TSharedRef<FFrontendFilter>& FrontendFilter);
+	void RemoveFilter(const TWeakPtr<IAssetTypeActions>& AssetTypeActions, bool ExecuteOnFilterChanged = true);
+	void RemoveFilter(const TSharedRef<FFrontendFilter>& FrontendFilter, bool ExecuteOnFilterChanged = true);
 	void RemoveFilter(const TSharedRef<SFilter>& FilterToRemove);
+	void RemoveFilterAndUpdate(const TSharedRef<SFilter>& FilterToRemove);
 
 	/** Handler for when the enable only this button was clicked on a single filter */
 	void EnableOnlyThisFilter(const TSharedRef<SFilter>& FilterToEnable);
