adds check of FSourceControlWindows::PromptForCheckin() return value to pop error message on failure

diff --git a/Engine/Source/Editor/ContentBrowser/Private/AssetContextMenu.cpp b/Engine/Source/Editor/ContentBrowser/Private/AssetContextMenu.cpp
--- a/Engine/Source/Editor/ContentBrowser/Private/AssetContextMenu.cpp
+++ b/Engine/Source/Editor/ContentBrowser/Private/AssetContextMenu.cpp
@@ -2418,7 +2418,12 @@ void FAssetContextMenu::ExecuteSCCCheckIn()
 		GetSelectedPackageNames(PackageNames);
 
 		const bool bUseSourceControlStateCache = true;
-		FSourceControlWindows::PromptForCheckin(bUseSourceControlStateCache, PackageNames);
+		const bool bCheckinGood = FSourceControlWindows::PromptForCheckin(bUseSourceControlStateCache, PackageNames);
+
+		if (!bCheckinGood)
+		{
+			FMessageDialog::Open(EAppMsgType::Ok, NSLOCTEXT("UnrealEd", "SCC_Checkin_Failed", "Check-in failed as a result of save failure."));
+		}
 	}
 	else
 	{
