PR #5078: Show in New Content browser context menu option V2 (Contributed by [REDACTED_PERSON])

diff --git a/Engine/Source/Editor/ContentBrowser/Private/ContentBrowserSingleton.cpp b/Engine/Source/Editor/ContentBrowser/Private/ContentBrowserSingleton.cpp
--- a/Engine/Source/Editor/ContentBrowser/Private/ContentBrowserSingleton.cpp
+++ b/Engine/Source/Editor/ContentBrowser/Private/ContentBrowserSingleton.cpp
@@ -230,7 +230,7 @@ void FContentBrowserSingleton::CreateNewAsset(const FString& DefaultAssetName, c
 	}
 }
 
-TSharedPtr<SContentBrowser> FContentBrowserSingleton::FindContentBrowserToSync(bool bAllowLockedBrowsers)
+TSharedPtr<SContentBrowser> FContentBrowserSingleton::FindContentBrowserToSync(bool bAllowLockedBrowsers, bool bNewSpawnBrowser)
 {
 	TSharedPtr<SContentBrowser> ContentBrowserToSync;
 
@@ -241,8 +241,12 @@ TSharedPtr<SContentBrowser> FContentBrowserSingleton::FindContentBrowserToSync(b
 
 	if ( PrimaryContentBrowser.IsValid() && (bAllowLockedBrowsers || !PrimaryContentBrowser.Pin()->IsLocked()) )
 	{
-		// If the primary content browser is not locked, sync it
-		ContentBrowserToSync = PrimaryContentBrowser.Pin();
+		// If wanting to spawn a new browser window, don't set the BrowserToSync in order to summon a new browser
+		if (!bNewSpawnBrowser)
+		{
+			// If the primary content browser is not locked, sync it
+			ContentBrowserToSync = PrimaryContentBrowser.Pin();
+		}
 	}
 	else
 	{
@@ -260,12 +264,12 @@ TSharedPtr<SContentBrowser> FContentBrowserSingleton::FindContentBrowserToSync(b
 	if ( !ContentBrowserToSync.IsValid() )
 	{
 		// There are no valid, unlocked browsers, attempt to summon a new one.
-		SummonNewBrowser(bAllowLockedBrowsers);
+		const FName NewBrowserName = SummonNewBrowser(bAllowLockedBrowsers);
 
 		// Now try to find a non-locked valid browser again, now that a new one may exist
 		for (int32 BrowserIdx = 0; BrowserIdx < AllContentBrowsers.Num(); ++BrowserIdx)
 		{
-			if ( AllContentBrowsers[BrowserIdx].IsValid() && (bAllowLockedBrowsers || !AllContentBrowsers[BrowserIdx].Pin()->IsLocked()) )
+			if (AllContentBrowsers[BrowserIdx].IsValid() && (NewBrowserName == NAME_None && (bAllowLockedBrowsers || !AllContentBrowsers[BrowserIdx].Pin()->IsLocked())) || (AllContentBrowsers[BrowserIdx].Pin()->GetInstanceName() == NewBrowserName) )
 			{
 				ContentBrowserToSync = AllContentBrowsers[BrowserIdx].Pin();
 				break;
@@ -281,9 +285,9 @@ TSharedPtr<SContentBrowser> FContentBrowserSingleton::FindContentBrowserToSync(b
 	return ContentBrowserToSync;
 }
 
-void FContentBrowserSingleton::SyncBrowserToAssets(const TArray<FAssetData>& AssetDataList, bool bAllowLockedBrowsers, bool bFocusContentBrowser)
+void FContentBrowserSingleton::SyncBrowserToAssets(const TArray<FAssetData>& AssetDataList, bool bAllowLockedBrowsers, bool bFocusContentBrowser, bool bNewSpawnBrowser)
 {
-	TSharedPtr<SContentBrowser> ContentBrowserToSync = FindContentBrowserToSync(bAllowLockedBrowsers);
+	TSharedPtr<SContentBrowser> ContentBrowserToSync = FindContentBrowserToSync(bAllowLockedBrowsers, bNewSpawnBrowser);
 
 	if ( ContentBrowserToSync.IsValid() )
 	{
@@ -296,7 +300,7 @@ void FContentBrowserSingleton::SyncBrowserToAssets(const TArray<FAssetData>& Ass
 	}
 }
 
-void FContentBrowserSingleton::SyncBrowserToAssets(const TArray<UObject*>& AssetList, bool bAllowLockedBrowsers, bool bFocusContentBrowser)
+void FContentBrowserSingleton::SyncBrowserToAssets(const TArray<UObject*>& AssetList, bool bAllowLockedBrowsers, bool bFocusContentBrowser, bool bNewSpawnBrowser)
 {
 	// Convert UObject* array to FAssetData array
 	TArray<FAssetData> AssetDataList;
@@ -308,12 +312,12 @@ void FContentBrowserSingleton::SyncBrowserToAssets(const TArray<UObject*>& Asset
 		}
 	}
 
-	SyncBrowserToAssets(AssetDataList, bAllowLockedBrowsers, bFocusContentBrowser);
+	SyncBrowserToAssets(AssetDataList, bAllowLockedBrowsers, bFocusContentBrowser, bNewSpawnBrowser);
 }
 
-void FContentBrowserSingleton::SyncBrowserToFolders(const TArray<FString>& FolderList, bool bAllowLockedBrowsers, bool bFocusContentBrowser)
+void FContentBrowserSingleton::SyncBrowserToFolders(const TArray<FString>& FolderList, bool bAllowLockedBrowsers, bool bFocusContentBrowser, bool bNewSpawnBrowser)
 {
-	TSharedPtr<SContentBrowser> ContentBrowserToSync = FindContentBrowserToSync(bAllowLockedBrowsers);
+	TSharedPtr<SContentBrowser> ContentBrowserToSync = FindContentBrowserToSync(bAllowLockedBrowsers, bNewSpawnBrowser);
 
 	if ( ContentBrowserToSync.IsValid() )
 	{
@@ -326,9 +330,9 @@ void FContentBrowserSingleton::SyncBrowserToFolders(const TArray<FString>& Folde
 	}
 }
 
-void FContentBrowserSingleton::SyncBrowserTo(const FContentBrowserSelection& ItemSelection, bool bAllowLockedBrowsers, bool bFocusContentBrowser)
+void FContentBrowserSingleton::SyncBrowserTo(const FContentBrowserSelection& ItemSelection, bool bAllowLockedBrowsers, bool bFocusContentBrowser, bool bNewSpawnBrowser)
 {
-	TSharedPtr<SContentBrowser> ContentBrowserToSync = FindContentBrowserToSync(bAllowLockedBrowsers);
+	TSharedPtr<SContentBrowser> ContentBrowserToSync = FindContentBrowserToSync(bAllowLockedBrowsers, bNewSpawnBrowser);
 
 	if ( ContentBrowserToSync.IsValid() )
 	{
@@ -491,7 +495,7 @@ void FContentBrowserSingleton::FocusContentBrowser(const TSharedPtr<SContentBrow
 	}
 }
 
-void FContentBrowserSingleton::SummonNewBrowser(bool bAllowLockedBrowsers)
+FName FContentBrowserSingleton::SummonNewBrowser(bool bAllowLockedBrowsers)
 {
 	TSet<FName> OpenBrowserIDs;
 
@@ -533,6 +537,8 @@ void FContentBrowserSingleton::SummonNewBrowser(bool bAllowLockedBrowsers)
 	{
 		// No available slots... don't summon anything
 	}
+
+	return NewTabName;
 }
 
 TSharedRef<SWidget> FContentBrowserSingleton::CreateContentBrowser( const FName InstanceName, TSharedPtr<SDockTab> ContainingTab, const FContentBrowserConfig* ContentBrowserConfig )
diff --git a/Engine/Source/Editor/ContentBrowser/Private/ContentBrowserSingleton.h b/Engine/Source/Editor/ContentBrowser/Private/ContentBrowserSingleton.h
--- a/Engine/Source/Editor/ContentBrowser/Private/ContentBrowserSingleton.h
+++ b/Engine/Source/Editor/ContentBrowser/Private/ContentBrowserSingleton.h
@@ -40,10 +40,10 @@ public:
 	virtual bool HasPrimaryContentBrowser() const override;
 	virtual void FocusPrimaryContentBrowser(bool bFocusSearch) override;
 	virtual void CreateNewAsset(const FString& DefaultAssetName, const FString& PackagePath, UClass* AssetClass, UFactory* Factory) override;
-	virtual void SyncBrowserToAssets(const TArray<struct FAssetData>& AssetDataList, bool bAllowLockedBrowsers = false, bool bFocusContentBrowser = true) override;
-	virtual void SyncBrowserToAssets(const TArray<UObject*>& AssetList, bool bAllowLockedBrowsers = false, bool bFocusContentBrowser = true) override;
-	virtual void SyncBrowserToFolders(const TArray<FString>& FolderList, bool bAllowLockedBrowsers = false, bool bFocusContentBrowser = true) override;
-	virtual void SyncBrowserTo(const FContentBrowserSelection& ItemSelection, bool bAllowLockedBrowsers = false, bool bFocusContentBrowser = true) override;
+	virtual void SyncBrowserToAssets(const TArray<struct FAssetData>& AssetDataList, bool bAllowLockedBrowsers = false, bool bFocusContentBrowser = true, bool bNewSpawnBrowser = false) override;
+	virtual void SyncBrowserToAssets(const TArray<UObject*>& AssetList, bool bAllowLockedBrowsers = false, bool bFocusContentBrowser = true, bool bNewSpawnBrowser = false) override;
+	virtual void SyncBrowserToFolders(const TArray<FString>& FolderList, bool bAllowLockedBrowsers = false, bool bFocusContentBrowser = true, bool bNewSpawnBrowser = false) override;
+	virtual void SyncBrowserTo(const FContentBrowserSelection& ItemSelection, bool bAllowLockedBrowsers = false, bool bFocusContentBrowser = true, bool bNewSpawnBrowser = false) override;
 	virtual void GetSelectedAssets(TArray<FAssetData>& SelectedAssets) override;
 	virtual void CaptureThumbnailFromViewport(FViewport* InViewport, TArray<FAssetData>& SelectedAssets) override;
 	virtual void SetSelectedPaths(const TArray<FString>& FolderPaths, bool bNeedsRefresh = false) override;
@@ -69,7 +69,7 @@ public:
 private:
 
 	/** Util to get or create the content browser that should be used by the various Sync functions */
-	TSharedPtr<SContentBrowser> FindContentBrowserToSync(bool bAllowLockedBrowsers);
+	TSharedPtr<SContentBrowser> FindContentBrowserToSync(bool bAllowLockedBrowsers, bool bNewSpawnBrowser);
 
 	/** Shared code to open an asset dialog window with a config */
 	void SharedCreateAssetDialogWindow(const TSharedRef<class SAssetDialog>& AssetDialog, const FSharedAssetDialogConfig& InConfig, bool bModal) const;
@@ -86,7 +86,7 @@ private:
 	void FocusContentBrowser(const TSharedPtr<SContentBrowser>& BrowserToFocus);
 
 	/** Summons a new content browser */
-	void SummonNewBrowser(bool bAllowLockedBrowsers = false);
+	FName SummonNewBrowser(bool bAllowLockedBrowsers = false);
 
 	/** Handler for a request to spawn a new content browser tab */
 	TSharedRef<SDockTab> SpawnContentBrowserTab( const FSpawnTabArgs& SpawnTabArgs, int32 BrowserIdx );
diff --git a/Engine/Source/Editor/ContentBrowser/Private/PathContextMenu.cpp b/Engine/Source/Editor/ContentBrowser/Private/PathContextMenu.cpp
--- a/Engine/Source/Editor/ContentBrowser/Private/PathContextMenu.cpp
+++ b/Engine/Source/Editor/ContentBrowser/Private/PathContextMenu.cpp
@@ -82,6 +82,11 @@ void FPathContextMenu::SetOnFolderFavoriteToggled(const FOnFolderFavoriteToggled
 	OnFolderFavoriteToggled = InOnFolderFavoriteToggled;
 }
 
+const TArray<FString>& FPathContextMenu::GetSelectedPaths() const
+{
+	return SelectedPaths;
+}
+
 void FPathContextMenu::SetSelectedPaths(const TArray<FString>& InSelectedPaths)
 {
 	SelectedPaths = InSelectedPaths;
@@ -107,7 +112,7 @@ TSharedRef<FExtender> FPathContextMenu::MakePathViewContextMenuExtender(const TA
 	}
 	TSharedPtr<FExtender> MenuExtender = FExtender::Combine(Extenders);
 
-	MenuExtender->AddMenuExtension("NewFolder", EExtensionHook::After, TSharedPtr<FUICommandList>(), FMenuExtensionDelegate::CreateSP(this, &FPathContextMenu::MakePathViewContextMenu));
+	MenuExtender->AddMenuExtension("FolderContext", EExtensionHook::After, TSharedPtr<FUICommandList>(), FMenuExtensionDelegate::CreateSP(this, &FPathContextMenu::MakePathViewContextMenu));
 
 	return MenuExtender.ToSharedRef();
 }
diff --git a/Engine/Source/Editor/ContentBrowser/Private/PathContextMenu.h b/Engine/Source/Editor/ContentBrowser/Private/PathContextMenu.h
--- a/Engine/Source/Editor/ContentBrowser/Private/PathContextMenu.h
+++ b/Engine/Source/Editor/ContentBrowser/Private/PathContextMenu.h
@@ -38,6 +38,9 @@ public:
 	DECLARE_DELEGATE_OneParam(FOnFolderFavoriteToggled, const TArray<FString>& /*FoldersToToggle*/)
 	void SetOnFolderFavoriteToggled(const FOnFolderFavoriteToggled& InOnFolderFavoriteToggled);
 
+	/** Gets the currently selected paths */
+	const TArray<FString>& GetSelectedPaths() const;
+
 	/** Sets the currently selected paths */
 	void SetSelectedPaths(const TArray<FString>& InSelectedPaths);
 
diff --git a/Engine/Source/Editor/ContentBrowser/Private/SContentBrowser.cpp b/Engine/Source/Editor/ContentBrowser/Private/SContentBrowser.cpp
--- a/Engine/Source/Editor/ContentBrowser/Private/SContentBrowser.cpp
+++ b/Engine/Source/Editor/ContentBrowser/Private/SContentBrowser.cpp
@@ -2765,10 +2765,17 @@ TSharedPtr<SWidget> SContentBrowser::GetFolderContextMenu(const TArray<FString>&
 		FUIAction(
 			FExecuteAction::CreateSP( this, &SContentBrowser::CreateNewFolder, SelectedPaths.Num() > 0 ? SelectedPaths[0] : FString(), InOnCreateNewFolder ),
 			FCanExecuteAction::CreateLambda( [bCanCreateNewFolder] { return bCanCreateNewFolder; } )
-			),
-		"NewFolder"
+			)
 		);
 
+	MenuBuilder.AddMenuEntry(
+		LOCTEXT("ShowInNewContentBrowser", "Show in New Content Browser"),
+		LOCTEXT("ShowInNewContentBrowserTooltip", "Opens a new Content Browser at this folder location (at least 1 Content Browser window needs to be locked)"),
+		FSlateIcon(),
+		FUIAction(FExecuteAction::CreateSP(this, &SContentBrowser::OpenNewContentBrowser)),
+		"FolderContext"
+	);
+
 	return MenuBuilder.MakeWidget();
 }
 
@@ -2787,4 +2794,9 @@ void SContentBrowser::CreateNewFolder(FString FolderPath, FOnCreateNewFolder InO
 	InOnCreateNewFolder.ExecuteIfBound(DefaultFolderName.ToString(), FolderPath);
 }
 
+void SContentBrowser::OpenNewContentBrowser()
+{
+	FContentBrowserSingleton::Get().SyncBrowserToFolders(PathContextMenu->GetSelectedPaths(), false, true, true);
+}
+
 #undef LOCTEXT_NAMESPACE
diff --git a/Engine/Source/Editor/ContentBrowser/Private/SContentBrowser.h b/Engine/Source/Editor/ContentBrowser/Private/SContentBrowser.h
--- a/Engine/Source/Editor/ContentBrowser/Private/SContentBrowser.h
+++ b/Engine/Source/Editor/ContentBrowser/Private/SContentBrowser.h
@@ -405,6 +405,9 @@ private:
 	/** Sets up an inline-name for the creation of a default-named folder the specified path */
 	void CreateNewFolder(FString FolderPath, FOnCreateNewFolder OnCreateNewFolder);
 
+	/** Handler for when "Open in new Content Browser" is selected */
+	void OpenNewContentBrowser();
+
 	/** Bind our UI commands */
 	void BindCommands();
 
diff --git a/Engine/Source/Editor/ContentBrowser/Private/SPathPicker.cpp b/Engine/Source/Editor/ContentBrowser/Private/SPathPicker.cpp
--- a/Engine/Source/Editor/ContentBrowser/Private/SPathPicker.cpp
+++ b/Engine/Source/Editor/ContentBrowser/Private/SPathPicker.cpp
@@ -91,7 +91,7 @@ TSharedPtr<SWidget> SPathPicker::GetFolderContextMenu(const TArray<FString>& Sel
 			FExecuteAction::CreateSP(this, &SPathPicker::CreateNewFolder, SelectedPaths.Num() > 0 ? SelectedPaths[0] : FString(), InOnCreateNewFolder),
 			FCanExecuteAction::CreateLambda( [bCanCreateNewFolder] { return bCanCreateNewFolder; } )
 			),
-		"NewFolder"
+		"FolderContext"
 		);
 
 	return MenuBuilder.MakeWidget();
diff --git a/Engine/Source/Editor/ContentBrowser/Private/SPathView.h b/Engine/Source/Editor/ContentBrowser/Private/SPathView.h
--- a/Engine/Source/Editor/ContentBrowser/Private/SPathView.h
+++ b/Engine/Source/Editor/ContentBrowser/Private/SPathView.h
@@ -81,9 +81,6 @@ public:
 	/** Constructs this widget with InArgs */
 	virtual void Construct( const FArguments& InArgs );
 
-	/** Sets focus to the search box */
-	void FocusSearchBox();
-
 	/** Selects the closest matches to the supplied paths in the tree. "/" delimited */
 	virtual void SetSelectedPaths(const TArray<FString>& Paths);
 
diff --git a/Engine/Source/Editor/ContentBrowser/Public/IContentBrowserSingleton.h b/Engine/Source/Editor/ContentBrowser/Public/IContentBrowserSingleton.h
--- a/Engine/Source/Editor/ContentBrowser/Public/IContentBrowserSingleton.h
+++ b/Engine/Source/Editor/ContentBrowser/Public/IContentBrowserSingleton.h
@@ -510,14 +510,14 @@ public:
 	virtual void CreateNewAsset(const FString& DefaultAssetName, const FString& PackagePath, UClass* AssetClass, UFactory* Factory) = 0;
 
 	/** Selects the supplied assets in all content browsers. If bAllowLockedBrowsers is true, even locked browsers may handle the sync. Only set to true if the sync doesn't seem external to the content browser. */
-	virtual void SyncBrowserToAssets(const TArray<struct FAssetData>& AssetDataList, bool bAllowLockedBrowsers = false, bool bFocusContentBrowser = true) = 0;
-	virtual void SyncBrowserToAssets(const TArray<UObject*>& AssetList, bool bAllowLockedBrowsers = false, bool bFocusContentBrowser = true) = 0;
+	virtual void SyncBrowserToAssets(const TArray<struct FAssetData>& AssetDataList, bool bAllowLockedBrowsers = false, bool bFocusContentBrowser = true, bool bNewSpawnBrowser = false) = 0;
+	virtual void SyncBrowserToAssets(const TArray<UObject*>& AssetList, bool bAllowLockedBrowsers = false, bool bFocusContentBrowser = true, bool bNewSpawnBrowser = false) = 0;
 
 	/** Selects the supplied folders in all content browsers. If bAllowLockedBrowsers is true, even locked browsers may handle the sync. Only set to true if the sync doesn't seem external to the content browser. */
-	virtual void SyncBrowserToFolders(const TArray<FString>& FolderList, bool bAllowLockedBrowsers = false, bool bFocusContentBrowser = true) = 0;
+	virtual void SyncBrowserToFolders(const TArray<FString>& FolderList, bool bAllowLockedBrowsers = false, bool bFocusContentBrowser = true, bool bNewSpawnBrowser = false) = 0;
 
 	/** Selects the supplied items in all content browsers. If bAllowLockedBrowsers is true, even locked browsers may handle the sync. Only set to true if the sync doesn't seem external to the content browser. */
-	virtual void SyncBrowserTo(const FContentBrowserSelection& ItemSelection, bool bAllowLockedBrowsers = false, bool bFocusContentBrowser = true) = 0;
+	virtual void SyncBrowserTo(const FContentBrowserSelection& ItemSelection, bool bAllowLockedBrowsers = false, bool bFocusContentBrowser = true, bool bNewSpawnBrowser = false) = 0;
 
 	/** Generates a list of assets that are selected in the primary content browser */
 	virtual void GetSelectedAssets(TArray<FAssetData>& SelectedAssets) = 0;
