Fix for linux compile warning as error.

diff --git a/Engine/Source/Editor/ContentBrowser/Private/ContentBrowserSingleton.cpp b/Engine/Source/Editor/ContentBrowser/Private/ContentBrowserSingleton.cpp
--- a/Engine/Source/Editor/ContentBrowser/Private/ContentBrowserSingleton.cpp
+++ b/Engine/Source/Editor/ContentBrowser/Private/ContentBrowserSingleton.cpp
@@ -269,7 +269,7 @@ TSharedPtr<SContentBrowser> FContentBrowserSingleton::FindContentBrowserToSync(b
 		// Now try to find a non-locked valid browser again, now that a new one may exist
 		for (int32 BrowserIdx = 0; BrowserIdx < AllContentBrowsers.Num(); ++BrowserIdx)
 		{
-			if (AllContentBrowsers[BrowserIdx].IsValid() && (NewBrowserName == NAME_None && (bAllowLockedBrowsers || !AllContentBrowsers[BrowserIdx].Pin()->IsLocked())) || (AllContentBrowsers[BrowserIdx].Pin()->GetInstanceName() == NewBrowserName) )
+			if (AllContentBrowsers[BrowserIdx].IsValid() && ((NewBrowserName == NAME_None && (bAllowLockedBrowsers || !AllContentBrowsers[BrowserIdx].Pin()->IsLocked()))) || (AllContentBrowsers[BrowserIdx].Pin()->GetInstanceName() == NewBrowserName) )
 			{
 				ContentBrowserToSync = AllContentBrowsers[BrowserIdx].Pin();
 				break;
