Made the collection UI show the path to collections it's in with the undetermined checkbox state.

diff --git a/Engine/Source/Editor/ContentBrowser/Private/AssetContextMenu.cpp b/Engine/Source/Editor/ContentBrowser/Private/AssetContextMenu.cpp
--- a/Engine/Source/Editor/ContentBrowser/Private/AssetContextMenu.cpp
+++ b/Engine/Source/Editor/ContentBrowser/Private/AssetContextMenu.cpp
@@ -1394,14 +1394,14 @@ bool FAssetContextMenu::AddCollectionMenuOptions(FMenuBuilder& MenuBuilder)
 		static void OnCollectionClicked(TSharedRef<FCollectionAssetManagement> QuickAssetManagement, FCollectionNameType InCollectionKey)
 		{
 			// The UI actions don't give you the new check state, so we need to emulate the behavior of SCheckBox
-			// Basically, unchecked will transition to checked (adding items), and anything else will transition to unchecked (removing items)
-			if (GetCollectionCheckState(QuickAssetManagement, InCollectionKey) == ECheckBoxState::Unchecked)
+			// Basically, checked will transition to unchecked (removing items), and anything else will transition to checked (adding items)
+			if (GetCollectionCheckState(QuickAssetManagement, InCollectionKey) == ECheckBoxState::Checked)
 			{
-				QuickAssetManagement->AddCurrentAssetsToCollection(InCollectionKey);
+				QuickAssetManagement->RemoveCurrentAssetsFromCollection(InCollectionKey);
 			}
 			else
 			{
-				QuickAssetManagement->RemoveCurrentAssetsFromCollection(InCollectionKey);
+				QuickAssetManagement->AddCurrentAssetsToCollection(InCollectionKey);
 			}
 		}
 	};
diff --git a/Engine/Source/Editor/ContentBrowser/Private/CollectionAssetManagement.cpp b/Engine/Source/Editor/ContentBrowser/Private/CollectionAssetManagement.cpp
--- a/Engine/Source/Editor/ContentBrowser/Private/CollectionAssetManagement.cpp
+++ b/Engine/Source/Editor/ContentBrowser/Private/CollectionAssetManagement.cpp
@@ -161,12 +161,23 @@ void FCollectionAssetManagement::UpdateAssetManagementState()
 	if (CurrentAssetPaths.Num() == 1)
 	{
 		TArray<FCollectionNameType> MatchedCollections;
-		CollectionManagerModule.Get().GetCollectionsContainingObject(*CurrentAssetPaths.CreateConstIterator(), MatchedCollections);
+		CollectionManagerModule.Get().GetCollectionsContainingObject(*CurrentAssetPaths.CreateConstIterator(), MatchedCollections, ECollectionRecursionFlags::Self);
 
 		for (const FCollectionNameType& CollectionKey : MatchedCollections)
 		{
 			AssetManagementState.Add(CollectionKey, ECheckBoxState::Checked);
 		}
+
+		MatchedCollections.Reset();
+		CollectionManagerModule.Get().GetCollectionsContainingObject(*CurrentAssetPaths.CreateConstIterator(), MatchedCollections, ECollectionRecursionFlags::Children);
+
+		for (const FCollectionNameType& CollectionKey : MatchedCollections)
+		{
+			if (!AssetManagementState.Contains(CollectionKey))
+			{
+				AssetManagementState.Add(CollectionKey, ECheckBoxState::Undetermined);
+			}
+		}
 	}
 	else
 	{
diff --git a/Engine/Source/Editor/ContentBrowser/Private/SCollectionView.cpp b/Engine/Source/Editor/ContentBrowser/Private/SCollectionView.cpp
--- a/Engine/Source/Editor/ContentBrowser/Private/SCollectionView.cpp
+++ b/Engine/Source/Editor/ContentBrowser/Private/SCollectionView.cpp
@@ -1385,18 +1385,14 @@ void SCollectionView::OnCollectionCheckStateChanged( ECheckBoxState NewState, TS
 {
 	if ( QuickAssetManagement.IsValid() )
 	{
-		switch(NewState)
+		FCollectionNameType CollectionKey(CollectionItem->CollectionName, CollectionItem->CollectionType);
+		if (QuickAssetManagement->GetCollectionCheckState(CollectionKey) == ECheckBoxState::Checked)
 		{
-		case ECheckBoxState::Checked:
-			QuickAssetManagement->AddCurrentAssetsToCollection(FCollectionNameType(CollectionItem->CollectionName, CollectionItem->CollectionType));
-			break;
-
-		case ECheckBoxState::Unchecked:
-			QuickAssetManagement->RemoveCurrentAssetsFromCollection(FCollectionNameType(CollectionItem->CollectionName, CollectionItem->CollectionType));
-			break;
-
-		default:
-			break;
+			QuickAssetManagement->RemoveCurrentAssetsFromCollection(CollectionKey);
+		}
+		else
+		{
+			QuickAssetManagement->AddCurrentAssetsToCollection(CollectionKey);
 		}
 	}
 }
