Fixing recently opened filter triggering modified filter, fixing filters not applying on editor open

diff --git a/Engine/Source/Editor/ContentBrowser/Private/FrontendFilters.h b/Engine/Source/Editor/ContentBrowser/Private/FrontendFilters.h
--- a/Engine/Source/Editor/ContentBrowser/Private/FrontendFilters.h
+++ b/Engine/Source/Editor/ContentBrowser/Private/FrontendFilters.h
@@ -312,7 +312,7 @@ public:
 	~FFrontendFilter_Recent();
 
 	// FFrontendFilter implementation
-	virtual FString GetName() const override { return TEXT("Modified"); }
+	virtual FString GetName() const override { return TEXT("RecentlyOpened"); }
 	virtual FText GetDisplayName() const override { return LOCTEXT("FrontendFilter_Recent", "Recently Opened"); }
 	virtual FText GetToolTipText() const override { return LOCTEXT("FrontendFilter_RecentTooltip", "Show only recently opened assets."); }
 	virtual void ActiveStateChanged(bool bActive) override;
diff --git a/Engine/Source/Editor/ContentBrowser/Private/SFilterList.cpp b/Engine/Source/Editor/ContentBrowser/Private/SFilterList.cpp
--- a/Engine/Source/Editor/ContentBrowser/Private/SFilterList.cpp
+++ b/Engine/Source/Editor/ContentBrowser/Private/SFilterList.cpp
@@ -781,6 +781,7 @@ void SFilterList::LoadSettings(const FString& IniFilename, const FString& IniSec
 					if ( EnabledFrontendFilterNames.Contains(FilterName) )
 					{
 						NewFilter->SetEnabled(true, false);
+						SetFrontendFilterActive(FrontendFilter, NewFilter->IsEnabled());
 					}
 				}
 			}
