Content browser optimization, call load module once

diff --git a/Engine/Source/Editor/ContentBrowser/Private/AssetViewWidgets.cpp b/Engine/Source/Editor/ContentBrowser/Private/AssetViewWidgets.cpp
--- a/Engine/Source/Editor/ContentBrowser/Private/AssetViewWidgets.cpp
+++ b/Engine/Source/Editor/ContentBrowser/Private/AssetViewWidgets.cpp
@@ -986,13 +986,14 @@ void SAssetViewItem::UpdateSourceControlState(float InDeltaTime)
 {
 	SourceControlStateDelay += InDeltaTime;
 
-	if ( !bSourceControlStateRequested && SourceControlStateDelay > 1.0f && ISourceControlModule::Get().IsEnabled() && AssetItem.IsValid() )
+	static ISourceControlModule& SourceControlModule = ISourceControlModule::Get();
+	if ( !bSourceControlStateRequested && SourceControlStateDelay > 1.0f && SourceControlModule.IsEnabled() && AssetItem.IsValid() )
 	{
 		// Only update the SCC state for non-temporary asset items that aren't a built in script
 		if ( !AssetItem->IsTemporaryItem() && AssetItem->GetType() != EAssetItemType::Folder && !FPackageName::IsScriptPackage(CachedPackageName) && !FPackageName::IsMemoryPackage(CachedPackageName) && !FPackageName::IsTempPackage(CachedPackageName))
 		{
 			// Request the most recent SCC state for this asset
-			ISourceControlModule::Get().QueueStatusUpdate(CachedPackageFileName);
+			SourceControlModule.QueueStatusUpdate(CachedPackageFileName);
 		}
 
 		bSourceControlStateRequested = true;
diff --git a/Engine/Source/Editor/UnrealEd/Private/AssetThumbnail.cpp b/Engine/Source/Editor/UnrealEd/Private/AssetThumbnail.cpp
--- a/Engine/Source/Editor/UnrealEd/Private/AssetThumbnail.cpp
+++ b/Engine/Source/Editor/UnrealEd/Private/AssetThumbnail.cpp
@@ -89,7 +89,7 @@ public:
 		const FAssetData& AssetData = AssetThumbnail->GetAssetData();
 
 		UClass* Class = FindObjectSafe<UClass>(ANY_PACKAGE, *AssetData.AssetClass.ToString());
-		FAssetToolsModule& AssetToolsModule = FModuleManager::LoadModuleChecked<FAssetToolsModule>(TEXT("AssetTools"));
+		static FAssetToolsModule& AssetToolsModule = FModuleManager::LoadModuleChecked<FAssetToolsModule>(TEXT("AssetTools"));
 		TSharedPtr<IAssetTypeActions> AssetTypeActions;
 		if ( Class != NULL )
 		{
