Content Browser thumbnails optimization, removed asset loaded events from individual assets.

A test case went from 2700ms -> 1400ms (with 2000 texture thumbnails visible, 4000 total in folder)
OnAssetLoaded callback went from 7ms -> 0.2ms (with 2000 texture thumbnails visible, right click unloaded texture)

diff --git a/Engine/Source/Editor/ContentBrowser/Private/AssetViewWidgets.cpp b/Engine/Source/Editor/ContentBrowser/Private/AssetViewWidgets.cpp
--- a/Engine/Source/Editor/ContentBrowser/Private/AssetViewWidgets.cpp
+++ b/Engine/Source/Editor/ContentBrowser/Private/AssetViewWidgets.cpp
@@ -307,6 +307,23 @@ void SAssetTileView::Tick(const FGeometry& AllottedGeometry, const double InCurr
 	}
 }
 
+bool SAssetTileView::HasWidgetForAsset(const FName& AssetPathName)
+{
+	for (const TSharedPtr<FAssetViewItem>& Item : WidgetGenerator.ItemsWithGeneratedWidgets)
+	{
+		if (Item.IsValid() && Item->GetType() == EAssetItemType::Normal)
+		{
+			const FName& ObjectPath = StaticCastSharedPtr<FAssetViewAsset>(Item)->Data.ObjectPath;
+			if (ObjectPath == AssetPathName)
+			{
+				return true;
+			}
+		}
+	}
+
+	return false;
+}
+
 FReply SAssetListView::OnKeyDown(const FGeometry& InGeometry, const FKeyEvent& InKeyEvent)
 {
 	FReply Reply = FAssetViewModeUtils::OnViewModeKeyDown(SelectedItems, InKeyEvent);
@@ -332,6 +349,23 @@ void SAssetListView::Tick(const FGeometry& AllottedGeometry, const double InCurr
 	}
 }
 
+bool SAssetListView::HasWidgetForAsset(const FName& AssetPathName)
+{
+	for (const TSharedPtr<FAssetViewItem>& Item : WidgetGenerator.ItemsWithGeneratedWidgets)
+	{
+		if (Item.IsValid() && Item->GetType() == EAssetItemType::Normal)
+		{
+			const FName& ObjectPath = StaticCastSharedPtr<FAssetViewAsset>(Item)->Data.ObjectPath;
+			if (ObjectPath == AssetPathName)
+			{
+				return true;
+			}
+		}
+	}
+
+	return false;
+}
+
 FReply SAssetColumnView::OnKeyDown(const FGeometry& InGeometry, const FKeyEvent& InKeyEvent)
 {
 	FReply Reply = FAssetViewModeUtils::OnViewModeKeyDown(SelectedItems, InKeyEvent);
@@ -1401,24 +1435,6 @@ void SAssetViewItem::SetForceMipLevelsToBeResident(bool bForce) const
 	}
 }
 
-void SAssetViewItem::HandleAssetLoaded(UObject* InAsset) const
-{
-	if(InAsset != nullptr)
-	{
-		if(AssetItem.IsValid() && AssetItem->GetType() == EAssetItemType::Normal)
-		{
-			const FAssetData& AssetData = StaticCastSharedPtr<FAssetViewAsset>(AssetItem)->Data;
-			if(AssetData.IsValid() && AssetData.IsAssetLoaded())
-			{
-				if(InAsset == AssetData.GetAsset())
-				{
-					SetForceMipLevelsToBeResident(true);
-				}
-			}
-		}
-	}
-}
-
 bool SAssetViewItem::OnVisualizeTooltip(const TSharedPtr<SWidget>& TooltipContent)
 {
 	if(OnVisualizeAssetToolTip.IsBound() && TooltipContent.IsValid() && AssetItem->GetType() != EAssetItemType::Folder)
@@ -1442,7 +1458,7 @@ void SAssetViewItem::OnToolTipClosing()
 
 SAssetListItem::~SAssetListItem()
 {
-	FCoreUObjectDelegates::OnAssetLoaded.RemoveAll(this);
+
 }
 
 BEGIN_SLATE_FUNCTION_BUILD_OPTIMIZATION
@@ -1563,9 +1579,6 @@ void SAssetListItem::Construct( const FArguments& InArgs )
 	}
 
 	SetForceMipLevelsToBeResident(true);
-
-	// listen for asset loads so we can force mips to stream in if required
-	FCoreUObjectDelegates::OnAssetLoaded.AddSP(this, &SAssetViewItem::HandleAssetLoaded);
 }
 END_SLATE_FUNCTION_BUILD_OPTIMIZATION
 
@@ -1596,7 +1609,7 @@ FOptionalSize SAssetListItem::GetSCCImageSize() const
 
 SAssetTileItem::~SAssetTileItem()
 {
-	FCoreUObjectDelegates::OnAssetLoaded.RemoveAll(this);
+
 }
 
 BEGIN_SLATE_FUNCTION_BUILD_OPTIMIZATION
@@ -1703,9 +1716,6 @@ void SAssetTileItem::Construct( const FArguments& InArgs )
 	}
 
 	SetForceMipLevelsToBeResident(true);
-
-	// listen for asset loads so we can force mips to stream in if required
-	FCoreUObjectDelegates::OnAssetLoaded.AddSP(this, &SAssetViewItem::HandleAssetLoaded);
 }
 END_SLATE_FUNCTION_BUILD_OPTIMIZATION
 
diff --git a/Engine/Source/Editor/ContentBrowser/Private/AssetViewWidgets.h b/Engine/Source/Editor/ContentBrowser/Private/AssetViewWidgets.h
--- a/Engine/Source/Editor/ContentBrowser/Private/AssetViewWidgets.h
+++ b/Engine/Source/Editor/ContentBrowser/Private/AssetViewWidgets.h
@@ -63,6 +63,9 @@ public:
 	virtual bool SupportsKeyboardFocus() const override { return true; }
 	virtual FReply OnKeyDown( const FGeometry& InGeometry, const FKeyEvent& InKeyEvent ) override;
 	virtual void Tick( const FGeometry& AllottedGeometry, const double InCurrentTime, const float InDeltaTime ) override;
+
+	/** Returns true if a widget is currently generated for a given asset */
+	bool HasWidgetForAsset(const FName& AssetPathName);
 };
 
 /** The list view mode of the asset view */
@@ -72,6 +75,9 @@ public:
 	virtual bool SupportsKeyboardFocus() const override { return true; }
 	virtual FReply OnKeyDown( const FGeometry& InGeometry, const FKeyEvent& InKeyEvent ) override;
 	virtual void Tick(const FGeometry& AllottedGeometry, const double InCurrentTime, const float InDeltaTime) override;
+
+	/** Returns true if a widget is currently generated for a given asset */
+	bool HasWidgetForAsset(const FName& AssetPathName);
 };
 
 /** The columns view mode of the asset view */
@@ -166,9 +172,6 @@ public:
 	/** Get the name text to be displayed for this item */
 	FText GetNameText() const;
 
-	/** Delegate handling when an asset is loaded */
-	void HandleAssetLoaded(UObject* InAsset) const;
-
 protected:
 	/** Used by OnDragEnter, OnDragOver, and OnDrop to check and update the validity of the drag operation */
 	bool ValidateDragDrop( const FGeometry& MyGeometry, const FDragDropEvent& DragDropEvent, bool& OutIsKnownDragOperation ) const;
diff --git a/Engine/Source/Editor/ContentBrowser/Private/SAssetView.cpp b/Engine/Source/Editor/ContentBrowser/Private/SAssetView.cpp
--- a/Engine/Source/Editor/ContentBrowser/Private/SAssetView.cpp
+++ b/Engine/Source/Editor/ContentBrowser/Private/SAssetView.cpp
@@ -49,6 +49,7 @@
 #include "HAL/PlatformApplicationMisc.h"
 #include "DesktopPlatformModule.h"
 #include "Misc/FileHelper.h"
+#include "Materials/Material.h"
 
 #define LOCTEXT_NAMESPACE "ContentBrowser"
 
@@ -2649,10 +2650,44 @@ void SAssetView::OnAssetRenamed(const FAssetData& AssetData, const FString& OldO
 
 void SAssetView::OnAssetLoaded(UObject* Asset)
 {
-	if ( Asset != NULL )
+	if (Asset == nullptr)
 	{
-		RecentlyLoadedOrChangedAssets.Add( FName(*Asset->GetPathName()), Asset );
+		return;
 	}
+
+	FName AssetPathName = FName(*Asset->GetPathName());
+	RecentlyLoadedOrChangedAssets.Add( AssetPathName, Asset );
+
+	UTexture2D* Texture2D = Cast<UTexture2D>(Asset);
+	UMaterial* Material = Texture2D ? nullptr : Cast<UMaterial>(Asset);
+	if ((Texture2D && !Texture2D->bForceMiplevelsToBeResident) || Material)
+	{
+		bool bHasWidgetForAsset = false;
+		switch (GetCurrentViewType())
+		{
+		case EAssetViewType::List:
+			bHasWidgetForAsset = ListView->HasWidgetForAsset(AssetPathName);
+			break;
+		case EAssetViewType::Tile:
+			bHasWidgetForAsset = TileView->HasWidgetForAsset(AssetPathName);
+			break;
+		default:
+			bHasWidgetForAsset = false;
+			break;
+		}
+
+		if (bHasWidgetForAsset)
+		{
+			if (Texture2D)
+			{
+				Texture2D->bForceMiplevelsToBeResident = true;
+			}
+			else if (Material)
+			{
+				Material->SetForceMipLevelsToBeResident(true, true, -1.0f);
+			}
+		}
+	};
 }
 
 void SAssetView::OnObjectPropertyChanged(UObject* Object, FPropertyChangedEvent& PropertyChangedEvent)
