New content browser filter to show files not under source control (slow method)

diff --git a/Engine/Source/Editor/ContentBrowser/Private/FrontendFilters.cpp b/Engine/Source/Editor/ContentBrowser/Private/FrontendFilters.cpp
--- a/Engine/Source/Editor/ContentBrowser/Private/FrontendFilters.cpp
+++ b/Engine/Source/Editor/ContentBrowser/Private/FrontendFilters.cpp
@@ -20,6 +20,7 @@
 #include "ContentBrowserModule.h"
 #include "MRUFavoritesList.h"
 #include "Settings/ContentBrowserSettings.h"
+#include "HAL/FileManager.h"
 
 /** Helper functions for frontend filters */
 namespace FrontendFilterHelper
@@ -717,6 +718,80 @@ void FFrontendFilter_CheckedOut::SourceControlOperationComplete(const FSourceCon
 	BroadcastChangedEvent();
 }
 
+/////////////////////////////////////////
+// FFrontendFilter_NotSourceControlled
+/////////////////////////////////////////
+
+FFrontendFilter_NotSourceControlled::FFrontendFilter_NotSourceControlled(TSharedPtr<FFrontendFilterCategory> InCategory) 
+	: FFrontendFilter(InCategory),
+	bSourceControlEnabled(false),
+	bCompletedSourceControlQuery(false)
+{
+
+}
+
+void FFrontendFilter_NotSourceControlled::ActiveStateChanged(bool bActive)
+{
+	if (bActive)
+	{
+		RequestStatus();
+	}
+	else
+	{
+		bCompletedSourceControlQuery = false;
+	}
+}
+
+void FFrontendFilter_NotSourceControlled::SetCurrentFilter(const FARFilter& InBaseFilter)
+{
+	bSourceControlEnabled = ISourceControlModule::Get().IsEnabled();
+}
+
+bool FFrontendFilter_NotSourceControlled::PassesFilter(FAssetFilterType InItem) const
+{
+	if (!bSourceControlEnabled)
+	{
+		return true;
+	}
+
+	if (!bCompletedSourceControlQuery)
+	{
+		return false;
+	}
+	
+	FSourceControlStatePtr SourceControlState = ISourceControlModule::Get().GetProvider().GetState(SourceControlHelpers::PackageFilename(InItem.PackageName.ToString()), EStateCacheUsage::Use);
+	return !SourceControlState.IsValid() || !SourceControlState->IsSourceControlled();
+}
+
+void FFrontendFilter_NotSourceControlled::RequestStatus()
+{
+	bCompletedSourceControlQuery = false;
+
+	ISourceControlProvider& SourceControlProvider = ISourceControlModule::Get().GetProvider();
+	if ( ISourceControlModule::Get().IsEnabled() )
+	{
+		bSourceControlEnabled = true;
+
+		// Request the state of files at filter construction time to make sure files have the correct state for the filter
+		TSharedRef<FUpdateStatus, ESPMode::ThreadSafe> UpdateStatusOperation = ISourceControlOperation::Create<FUpdateStatus>();
+			
+		TArray<FString> Filenames;
+		Filenames.Add(FPaths::ConvertRelativePathToFull(FPaths::ProjectContentDir()));
+		UpdateStatusOperation->SetCheckingAllFiles(false);
+		SourceControlProvider.Execute(UpdateStatusOperation, Filenames, EConcurrency::Asynchronous, FSourceControlOperationComplete::CreateSP(this, &FFrontendFilter_NotSourceControlled::SourceControlOperationComplete));
+	}
+	else
+	{
+		bSourceControlEnabled = false;
+	}
+}
+
+void FFrontendFilter_NotSourceControlled::SourceControlOperationComplete(const FSourceControlOperationRef& InOperation, ECommandResult::Type InResult)
+{
+	bCompletedSourceControlQuery = true;
+	BroadcastChangedEvent();
+}
+
 /////////////////////////////////////////
 // FFrontendFilter_Modified
 /////////////////////////////////////////
diff --git a/Engine/Source/Editor/ContentBrowser/Private/FrontendFilters.h b/Engine/Source/Editor/ContentBrowser/Private/FrontendFilters.h
--- a/Engine/Source/Editor/ContentBrowser/Private/FrontendFilters.h
+++ b/Engine/Source/Editor/ContentBrowser/Private/FrontendFilters.h
@@ -106,6 +106,34 @@ private:
 	void SourceControlOperationComplete(const FSourceControlOperationRef& InOperation, ECommandResult::Type InResult);
 };
 
+/** A filter that displays assets not tracked by source control */
+class FFrontendFilter_NotSourceControlled : public FFrontendFilter, public TSharedFromThis<FFrontendFilter_NotSourceControlled>
+{
+public:
+	FFrontendFilter_NotSourceControlled(TSharedPtr<FFrontendFilterCategory> InCategory);
+
+	// FFrontendFilter implementation
+	virtual FString GetName() const override { return TEXT("NotSourceControlled"); }
+	virtual FText GetDisplayName() const override { return LOCTEXT("FrontendFilter_NotSourceControlled", "Not Source Controlled"); }
+	virtual FText GetToolTipText() const override { return LOCTEXT("FrontendFilter_NotSourceControlledTooltip", "Show only assets that are not tracked by source control."); }
+	virtual void ActiveStateChanged(bool bActive) override;
+	virtual void SetCurrentFilter(const FARFilter& InBaseFilter);
+
+	// IFilter implementation
+	virtual bool PassesFilter(FAssetFilterType InItem) const override;
+
+private:
+
+	/** Request the source control status for this filter */
+	void RequestStatus();
+
+	/** Callback when source control operation has completed */
+	void SourceControlOperationComplete(const FSourceControlOperationRef& InOperation, ECommandResult::Type InResult);
+
+	bool bSourceControlEnabled;
+	bool bCompletedSourceControlQuery;
+};
+
 /** A filter that displays only modified assets */
 class FFrontendFilter_Modified : public FFrontendFilter
 {
diff --git a/Engine/Source/Editor/ContentBrowser/Private/SFilterList.cpp b/Engine/Source/Editor/ContentBrowser/Private/SFilterList.cpp
--- a/Engine/Source/Editor/ContentBrowser/Private/SFilterList.cpp
+++ b/Engine/Source/Editor/ContentBrowser/Private/SFilterList.cpp
@@ -417,6 +417,7 @@ void SFilterList::Construct( const FArguments& InArgs )
 	AllFrontendFilters.Add( MakeShareable(new FFrontendFilter_NotUsedInAnyLevel(DefaultCategory)) );
 	AllFrontendFilters.Add( MakeShareable(new FFrontendFilter_ArbitraryComparisonOperation(DefaultCategory)) );
 	AllFrontendFilters.Add(MakeShareable(new FFrontendFilter_Recent(DefaultCategory)));
+	AllFrontendFilters.Add( MakeShareable(new FFrontendFilter_NotSourceControlled(DefaultCategory)) );
 
 	// Add any global user-defined frontend filters
 	for (TObjectIterator<UContentBrowserFrontEndFilterExtension> ExtensionIt(RF_NoFlags); ExtensionIt; ++ExtensionIt)
