Content browser filter: Not Source Controlled updated

diff --git a/Engine/Source/Editor/ContentBrowser/Private/FrontendFilters.cpp b/Engine/Source/Editor/ContentBrowser/Private/FrontendFilters.cpp
--- a/Engine/Source/Editor/ContentBrowser/Private/FrontendFilters.cpp
+++ b/Engine/Source/Editor/ContentBrowser/Private/FrontendFilters.cpp
@@ -725,7 +725,8 @@ void FFrontendFilter_CheckedOut::SourceControlOperationComplete(const FSourceCon
 FFrontendFilter_NotSourceControlled::FFrontendFilter_NotSourceControlled(TSharedPtr<FFrontendFilterCategory> InCategory) 
 	: FFrontendFilter(InCategory),
 	bSourceControlEnabled(false),
-	bCompletedSourceControlQuery(false)
+	bIsRequestStatusRunning(false),
+	bInitialRequestCompleted(false)
 {
 
 }
@@ -734,17 +735,30 @@ void FFrontendFilter_NotSourceControlled::ActiveStateChanged(bool bActive)
 {
 	if (bActive)
 	{
-		RequestStatus();
-	}
-	else
-	{
-		bCompletedSourceControlQuery = false;
+		if (!bIsRequestStatusRunning)
+		{
+			RequestStatus();
+		}
 	}
 }
 
 void FFrontendFilter_NotSourceControlled::SetCurrentFilter(const FARFilter& InBaseFilter)
 {
 	bSourceControlEnabled = ISourceControlModule::Get().IsEnabled();
+
+	if (InBaseFilter.PackagePaths.Num() == 1)
+	{
+		UE_LOG(LogTemp, Log, TEXT("FFrontendFilter_NotSourceControlled::SetCurrentFilter() SCEnabled: %d, '%s'"), bSourceControlEnabled, *InBaseFilter.PackagePaths[0].ToString());
+	}
+	else if (InBaseFilter.PackagePaths.Num() > 1)
+	{
+		UE_LOG(LogTemp, Log, TEXT("FFrontendFilter_NotSourceControlled::SetCurrentFilter() SCEnabled: %d, '%s', ..."), bSourceControlEnabled, *InBaseFilter.PackagePaths[0].ToString());
+	}
+	else
+	{
+		UE_LOG(LogTemp, Log, TEXT("FFrontendFilter_NotSourceControlled::SetCurrentFilter() SCEnabled: %d"), bSourceControlEnabled);
+	}
+
 }
 
 bool FFrontendFilter_NotSourceControlled::PassesFilter(FAssetFilterType InItem) const
@@ -754,23 +768,39 @@ bool FFrontendFilter_NotSourceControlled::PassesFilter(FAssetFilterType InItem)
 		return true;
 	}
 
-	if (!bCompletedSourceControlQuery)
+	// Hide all items until the first status request finishes
+	if (!bInitialRequestCompleted)
 	{
 		return false;
 	}
-	
+
 	FSourceControlStatePtr SourceControlState = ISourceControlModule::Get().GetProvider().GetState(SourceControlHelpers::PackageFilename(InItem.PackageName.ToString()), EStateCacheUsage::Use);
-	return !SourceControlState.IsValid() || !SourceControlState->IsSourceControlled();
+	if (!SourceControlState.IsValid())
+	{
+		return false;
+	}
+
+	if (SourceControlState->IsUnknown())
+	{
+		return true;
+	}
+
+	if (SourceControlState->IsSourceControlled())
+	{
+		return false;
+	}
+
+	return true;
 }
 
 void FFrontendFilter_NotSourceControlled::RequestStatus()
 {
-	bCompletedSourceControlQuery = false;
-
 	ISourceControlProvider& SourceControlProvider = ISourceControlModule::Get().GetProvider();
-	if ( ISourceControlModule::Get().IsEnabled() )
+	bSourceControlEnabled = ISourceControlModule::Get().IsEnabled();
+	if ( bSourceControlEnabled )
 	{
 		bSourceControlEnabled = true;
+		bIsRequestStatusRunning = true;
 
 		// Request the state of files at filter construction time to make sure files have the correct state for the filter
 		TSharedRef<FUpdateStatus, ESPMode::ThreadSafe> UpdateStatusOperation = ISourceControlOperation::Create<FUpdateStatus>();
@@ -780,15 +810,13 @@ void FFrontendFilter_NotSourceControlled::RequestStatus()
 		UpdateStatusOperation->SetCheckingAllFiles(false);
 		SourceControlProvider.Execute(UpdateStatusOperation, Filenames, EConcurrency::Asynchronous, FSourceControlOperationComplete::CreateSP(this, &FFrontendFilter_NotSourceControlled::SourceControlOperationComplete));
 	}
-	else
-	{
-		bSourceControlEnabled = false;
-	}
 }
 
 void FFrontendFilter_NotSourceControlled::SourceControlOperationComplete(const FSourceControlOperationRef& InOperation, ECommandResult::Type InResult)
 {
-	bCompletedSourceControlQuery = true;
+	bIsRequestStatusRunning = false;
+	bInitialRequestCompleted = true;
+
 	BroadcastChangedEvent();
 }
 
diff --git a/Engine/Source/Editor/ContentBrowser/Private/FrontendFilters.h b/Engine/Source/Editor/ContentBrowser/Private/FrontendFilters.h
--- a/Engine/Source/Editor/ContentBrowser/Private/FrontendFilters.h
+++ b/Engine/Source/Editor/ContentBrowser/Private/FrontendFilters.h
@@ -131,7 +131,8 @@ private:
 	void SourceControlOperationComplete(const FSourceControlOperationRef& InOperation, ECommandResult::Type InResult);
 
 	bool bSourceControlEnabled;
-	bool bCompletedSourceControlQuery;
+	bool bIsRequestStatusRunning;
+	bool bInitialRequestCompleted;
 };
 
 /** A filter that displays only modified assets */
