Remove a debug log

diff --git a/Engine/Source/Editor/ContentBrowser/Private/FrontendFilters.cpp b/Engine/Source/Editor/ContentBrowser/Private/FrontendFilters.cpp
--- a/Engine/Source/Editor/ContentBrowser/Private/FrontendFilters.cpp
+++ b/Engine/Source/Editor/ContentBrowser/Private/FrontendFilters.cpp
@@ -745,20 +745,6 @@ void FFrontendFilter_NotSourceControlled::ActiveStateChanged(bool bActive)
 void FFrontendFilter_NotSourceControlled::SetCurrentFilter(const FARFilter& InBaseFilter)
 {
 	bSourceControlEnabled = ISourceControlModule::Get().IsEnabled();
-
-	if (InBaseFilter.PackagePaths.Num() == 1)
-	{
-		UE_LOG(LogTemp, Log, TEXT("FFrontendFilter_NotSourceControlled::SetCurrentFilter() SCEnabled: %d, '%s'"), bSourceControlEnabled, *InBaseFilter.PackagePaths[0].ToString());
-	}
-	else if (InBaseFilter.PackagePaths.Num() > 1)
-	{
-		UE_LOG(LogTemp, Log, TEXT("FFrontendFilter_NotSourceControlled::SetCurrentFilter() SCEnabled: %d, '%s', ..."), bSourceControlEnabled, *InBaseFilter.PackagePaths[0].ToString());
-	}
-	else
-	{
-		UE_LOG(LogTemp, Log, TEXT("FFrontendFilter_NotSourceControlled::SetCurrentFilter() SCEnabled: %d"), bSourceControlEnabled);
-	}
-
 }
 
 bool FFrontendFilter_NotSourceControlled::PassesFilter(FAssetFilterType InItem) const
@@ -806,6 +792,7 @@ void FFrontendFilter_NotSourceControlled::RequestStatus()
 		TSharedRef<FUpdateStatus, ESPMode::ThreadSafe> UpdateStatusOperation = ISourceControlOperation::Create<FUpdateStatus>();
 			
 		TArray<FString> Filenames;
+		Filenames.Add(FPaths::ConvertRelativePathToFull(FPaths::EngineContentDir()));
 		Filenames.Add(FPaths::ConvertRelativePathToFull(FPaths::ProjectContentDir()));
 		UpdateStatusOperation->SetCheckingAllFiles(false);
 		SourceControlProvider.Execute(UpdateStatusOperation, Filenames, EConcurrency::Asynchronous, FSourceControlOperationComplete::CreateSP(this, &FFrontendFilter_NotSourceControlled::SourceControlOperationComplete));
