Fixed new asset not being able to create objects if another window had focus

diff --git a/Engine/Source/Editor/ContentBrowser/Private/SAssetView.cpp b/Engine/Source/Editor/ContentBrowser/Private/SAssetView.cpp
--- a/Engine/Source/Editor/ContentBrowser/Private/SAssetView.cpp
+++ b/Engine/Source/Editor/ContentBrowser/Private/SAssetView.cpp
@@ -766,6 +766,13 @@ void SAssetView::CreateNewAsset(const FString& DefaultAssetName, const FString&
 	// we should only be creating one deferred asset per tick
 	check(!DeferredAssetToCreate.IsValid());
 
+	// Asset creation requires focus to give object a name, otherwise object will not be created
+	TSharedPtr<SWindow> OwnerWindow = FSlateApplication::Get().FindWidgetWindow(AsShared());
+	if (OwnerWindow.IsValid() && !OwnerWindow->HasAnyUserFocusOrFocusedDescendants())
+	{
+		FSlateApplication::Get().SetUserFocus(FSlateApplication::Get().GetUserIndexForKeyboard(), AsShared(), EFocusCause::SetDirectly);
+	}
+
 	// Make sure we are showing the location of the new asset (we may have created it in a folder)
 	OnPathSelected.Execute(PackagePath);
 
