[NO NON-PII COMMIT MESSAGE]

diff --git a/Engine/Source/Editor/ContentBrowser/Private/SAssetPicker.cpp b/Engine/Source/Editor/ContentBrowser/Private/SAssetPicker.cpp
--- a/Engine/Source/Editor/ContentBrowser/Private/SAssetPicker.cpp
+++ b/Engine/Source/Editor/ContentBrowser/Private/SAssetPicker.cpp
@@ -106,7 +106,7 @@ void SAssetPicker::Construct( const FArguments& InArgs )
 			HorizontalBox->AddSlot()
 			.AutoWidth()
 			[
-				SNew( SComboButton )
+				SAssignNew(FilterComboButtonPtr, SComboButton)
 				.ComboButtonStyle( FEditorStyle::Get(), "GenericFilters.ComboButtonStyle" )
 				.ForegroundColor(FLinearColor::White)
 				.ToolTipText( LOCTEXT( "AddFilterToolTip", "Add an asset filter." ) )
@@ -296,6 +296,12 @@ EActiveTimerReturnType SAssetPicker::SetFocusPostConstruct( double InCurrentTime
 
 FReply SAssetPicker::HandleKeyDownFromSearchBox(const FGeometry& MyGeometry, const FKeyEvent& InKeyEvent)
 {
+	// Hide the filter list
+	if (FilterComboButtonPtr.IsValid())
+	{
+		FilterComboButtonPtr->SetIsOpen(false);
+	}
+
 	// Up and down move thru the filtered list
 	int32 SelectionDelta = 0;
 
diff --git a/Engine/Source/Editor/ContentBrowser/Private/SAssetPicker.h b/Engine/Source/Editor/ContentBrowser/Private/SAssetPicker.h
--- a/Engine/Source/Editor/ContentBrowser/Private/SAssetPicker.h
+++ b/Engine/Source/Editor/ContentBrowser/Private/SAssetPicker.h
@@ -18,6 +18,7 @@ class FUICommandList;
 class SAssetSearchBox;
 class SAssetView;
 class SFilterList;
+class SComboButton;
 enum class ECheckBoxState : uint8;
 
 /**
@@ -133,6 +134,9 @@ private:
 	/** The filter list */
 	TSharedPtr<SFilterList> FilterListPtr;
 
+	/** The filter combo button */
+	TSharedPtr<SComboButton> FilterComboButtonPtr;
+
 	/** Called to when an asset is selected or the none button is pressed */
 	FOnAssetSelected OnAssetSelected;
 
