Clarifying order in Move/Copy menu

diff --git a/Engine/Source/Editor/ContentBrowser/Private/DragDropHandler.cpp b/Engine/Source/Editor/ContentBrowser/Private/DragDropHandler.cpp
--- a/Engine/Source/Editor/ContentBrowser/Private/DragDropHandler.cpp
+++ b/Engine/Source/Editor/ContentBrowser/Private/DragDropHandler.cpp
@@ -117,18 +117,17 @@ void DragDropHandler::HandleDropOnAssetFolder(const TSharedRef<SWidget>& ParentW
 	const FText MoveCopyHeaderString = FText::Format(LOCTEXT("AssetViewDropMenuHeading", "Move/Copy to {0}"), TargetDisplayName);
 	MenuBuilder.BeginSection("PathAssetMoveCopy", MoveCopyHeaderString);
 	{
-		MenuBuilder.AddMenuEntry(
-			LOCTEXT("DragDropCopy", "Copy Here"),
-			LOCTEXT("DragDropCopyTooltip", "Copy the dragged items to this folder, preserving the structure of any copied folders."),
-			FSlateIcon(),
-			FUIAction(FExecuteAction::CreateLambda([=]() { CopyActionHandler.ExecuteIfBound(FinalAssetList, FinalAssetPaths, TargetPath); }))
-			);
-
 		MenuBuilder.AddMenuEntry(
 			LOCTEXT("DragDropMove", "Move Here"),
 			LOCTEXT("DragDropMoveTooltip", "Move the dragged items to this folder, preserving the structure of any copied folders."),
 			FSlateIcon(),
 			FUIAction(FExecuteAction::CreateLambda([=]() { MoveActionHandler.ExecuteIfBound(FinalAssetList, FinalAssetPaths, TargetPath); }))
+		);
+		MenuBuilder.AddMenuEntry(
+			LOCTEXT("DragDropCopy", "Copy Here"),
+			LOCTEXT("DragDropCopyTooltip", "Copy the dragged items to this folder, preserving the structure of any copied folders."),
+			FSlateIcon(),
+			FUIAction(FExecuteAction::CreateLambda([=]() { CopyActionHandler.ExecuteIfBound(FinalAssetList, FinalAssetPaths, TargetPath); }))
 			);
 		MenuBuilder.AddMenuEntry(
 			LOCTEXT("DragDropAdvancedCopy", "Advanced Copy Here"),
