Focus on folder creation the same way with assets if owner window doesnt have any focus

Revert adding popup windows who are not set to be focused to the revert stack

diff --git a/Engine/Source/Editor/ContentBrowser/Private/SAssetView.cpp b/Engine/Source/Editor/ContentBrowser/Private/SAssetView.cpp
--- a/Engine/Source/Editor/ContentBrowser/Private/SAssetView.cpp
+++ b/Engine/Source/Editor/ContentBrowser/Private/SAssetView.cpp
@@ -727,6 +727,13 @@ void SAssetView::OnCreateNewFolder(const FString& FolderName, const FString& Fol
 	// we should only be creating one deferred folder per tick
 	check(!DeferredFolderToCreate.IsValid());
 
+	// Folder creation requires focus to give object a name, otherwise object will not be created
+	TSharedPtr<SWindow> OwnerWindow = FSlateApplication::Get().FindWidgetWindow(AsShared());
+	if (OwnerWindow.IsValid() && !OwnerWindow->HasAnyUserFocusOrFocusedDescendants())
+	{
+		FSlateApplication::Get().SetUserFocus(FSlateApplication::Get().GetUserIndexForKeyboard(), AsShared(), EFocusCause::SetDirectly);
+	}
+
 	// Make sure we are showing the location of the new folder (we may have created it in a folder)
 	OnPathSelected.Execute(FolderPath);
 
diff --git a/Engine/Source/Runtime/ApplicationCore/Private/Linux/LinuxApplication.cpp b/Engine/Source/Runtime/ApplicationCore/Private/Linux/LinuxApplication.cpp
--- a/Engine/Source/Runtime/ApplicationCore/Private/Linux/LinuxApplication.cpp
+++ b/Engine/Source/Runtime/ApplicationCore/Private/Linux/LinuxApplication.cpp
@@ -134,8 +134,8 @@ void FLinuxApplication::InitializeWindow(	const TSharedRef< FGenericWindow >& In
 	Window->Initialize( this, InDefinition, ParentWindow, bShowImmediately );
 	Windows.Add(Window);
 
-	// Add the windows that should revert activation when closed
-	if (Window->IsFocusWhenFirstShown() || Window->IsPopupMenuWindow())
+	// Add the windows into the focus stack.
+	if (Window->IsFocusWhenFirstShown())
 	{
 		RevertFocusStack.Add(Window);
 	}
