[NO NON-PII COMMIT MESSAGE]

diff --git a/Engine/Source/Editor/ContentBrowser/Private/AssetContextMenu.cpp b/Engine/Source/Editor/ContentBrowser/Private/AssetContextMenu.cpp
--- a/Engine/Source/Editor/ContentBrowser/Private/AssetContextMenu.cpp
+++ b/Engine/Source/Editor/ContentBrowser/Private/AssetContextMenu.cpp
@@ -6,6 +6,8 @@
 #include "Framework/Commands/UIAction.h"
 #include "Textures/SlateIcon.h"
 #include "Engine/Blueprint.h"
+#include "Engine/UserDefinedStruct.h"
+#include "Engine/UserDefinedEnum.h"
 #include "Misc/MessageDialog.h"
 #include "HAL/PlatformApplicationMisc.h"
 #include "HAL/FileManager.h"
@@ -2156,6 +2158,22 @@ void FAssetContextMenu::ExecuteReload()
 				continue;
 			}
 
+			if (AssetData.AssetClass == UUserDefinedStruct::StaticClass()->GetFName())
+			{
+				FNotificationInfo Notification(LOCTEXT("CannotReloadUserStruct", "User created structures cannot be safely reloaded."));
+				Notification.ExpireDuration = 3.0f;
+				FSlateNotificationManager::Get().AddNotification(Notification);
+				continue;
+			}
+
+			if (AssetData.AssetClass == UUserDefinedEnum::StaticClass()->GetFName())
+			{
+				FNotificationInfo Notification(LOCTEXT("CannotReloadUserEnum", "User created enumerations cannot be safely reloaded."));
+				Notification.ExpireDuration = 3.0f;
+				FSlateNotificationManager::Get().AddNotification(Notification);
+				continue;
+			}
+
 			PackagesToReload.AddUnique(AssetData.GetPackage());
 		}
 
