PR #5574: Prevent asset view from eating all other modifier combinations with ctrl + c (Contributed by [REDACTED_PERSON])

diff --git a/Engine/Source/Editor/ContentBrowser/Private/AssetViewWidgets.cpp b/Engine/Source/Editor/ContentBrowser/Private/AssetViewWidgets.cpp
--- a/Engine/Source/Editor/ContentBrowser/Private/AssetViewWidgets.cpp
+++ b/Engine/Source/Editor/ContentBrowser/Private/AssetViewWidgets.cpp
@@ -46,7 +46,9 @@
 FReply FAssetViewModeUtils::OnViewModeKeyDown( const TSet< TSharedPtr<FAssetViewItem> >& SelectedItems, const FKeyEvent& InKeyEvent )
 {
 	// All asset views use Ctrl-C to copy references to assets
-	if ( InKeyEvent.IsControlDown() && InKeyEvent.GetCharacter() == 'C' )
+	if ( InKeyEvent.IsControlDown() && InKeyEvent.GetCharacter() == 'C' 
+		&& !InKeyEvent.IsShiftDown() && !InKeyEvent.IsAltDown()
+		)
 	{
 		FAssetRegistryModule& AssetRegistryModule = FModuleManager::LoadModuleChecked<FAssetRegistryModule>(TEXT("AssetRegistry"));
 
