Optimize content browser filters/search during startup (while assets are being discovered)

diff --git a/Engine/Source/Editor/ContentBrowser/Private/SAssetView.cpp b/Engine/Source/Editor/ContentBrowser/Private/SAssetView.cpp
--- a/Engine/Source/Editor/ContentBrowser/Private/SAssetView.cpp
+++ b/Engine/Source/Editor/ContentBrowser/Private/SAssetView.cpp
@@ -1101,6 +1101,14 @@ void SAssetView::AdjustActiveSelection(int32 SelectionDelta)
 
 void SAssetView::ProcessRecentlyLoadedOrChangedAssets()
 {
+	if (RecentlyLoadedOrChangedAssets.Num() == 0)
+	{
+		return;
+	}
+
+	// Gather recently loaded or changed assets
+	TArray<FAssetData> AssetDatas;
+	TArray<int32> FilteredAssetItemIndices;
 	for (int32 AssetIdx = FilteredAssetItems.Num() - 1; AssetIdx >= 0 && RecentlyLoadedOrChangedAssets.Num() > 0; --AssetIdx)
 	{
 		if (FilteredAssetItems[AssetIdx]->GetType() != EAssetItemType::Folder)
@@ -1109,50 +1117,71 @@ void SAssetView::ProcessRecentlyLoadedOrChangedAssets()
 				
 			// Find the updated version of the asset data from the set
 			// This is the version of the data we should use to update our view
-			FAssetData RecentlyLoadedOrChangedAsset;
 			if (const FAssetData* RecentlyLoadedOrChangedAssetPtr = RecentlyLoadedOrChangedAssets.Find(ItemAsAsset->Data))
 			{
-				RecentlyLoadedOrChangedAsset = *RecentlyLoadedOrChangedAssetPtr;
+				if (RecentlyLoadedOrChangedAssetPtr->IsValid())
+				{
+					AssetDatas.Add(*RecentlyLoadedOrChangedAssetPtr);
+					FilteredAssetItemIndices.Add(AssetIdx);
+				}
+
 				RecentlyLoadedOrChangedAssets.Remove(ItemAsAsset->Data);
 			}
+		}
+	}
 
-			if (RecentlyLoadedOrChangedAsset.IsValid())
-			{
-				bool bShouldRemoveAsset = false;
+	// Run backend filter
+	TSet<FAssetData> PassedBackendFilterSet;
+	if (AssetDatas.Num() > 0)
+	{
+		TArray<FAssetData> AssetDatasCopy = AssetDatas;
+		RunAssetsThroughBackendFilter(AssetDatasCopy);
+		PassedBackendFilterSet.Append(AssetDatasCopy);
+	}
 
-				if (!PassesCurrentBackendFilter(RecentlyLoadedOrChangedAsset))
-				{
-					bShouldRemoveAsset = true;
-				}
+	// Update or remove
+	for (int32 i = 0; i < FilteredAssetItemIndices.Num(); ++i)
+	{
+		const int32 AssetIdx = FilteredAssetItemIndices[i];
+		const TSharedPtr<FAssetViewAsset>& ItemAsAsset = StaticCastSharedPtr<FAssetViewAsset>(FilteredAssetItems[AssetIdx]);
+		const FAssetData& RecentlyLoadedOrChangedAsset = AssetDatas[i];
 
-				if (!bShouldRemoveAsset && OnShouldFilterAsset.IsBound() && OnShouldFilterAsset.Execute(RecentlyLoadedOrChangedAsset))
-				{
-					bShouldRemoveAsset = true;
-				}
+		bool bShouldRemoveAsset = false;
 
-				if (!bShouldRemoveAsset && (IsFrontendFilterActive() && !PassesCurrentFrontendFilter(RecentlyLoadedOrChangedAsset)))
-				{
-					bShouldRemoveAsset = true;
-				}
+		if (!PassedBackendFilterSet.Contains(ItemAsAsset->Data))
+		{
+			bShouldRemoveAsset = true;
+		}
 
-				if (bShouldRemoveAsset)
-				{
-					FilteredAssetItems.RemoveAt(AssetIdx);
-				}
-				else
-				{
-					// Update the asset data on the item
-					ItemAsAsset->SetAssetData(RecentlyLoadedOrChangedAsset);
+		if (!bShouldRemoveAsset && OnShouldFilterAsset.IsBound() && OnShouldFilterAsset.Execute(RecentlyLoadedOrChangedAsset))
+		{
+			bShouldRemoveAsset = true;
+		}
 
-					// Update the custom column data
-					ItemAsAsset->CacheCustomColumns(CustomColumns, true, true, true);
-				}
+		if (!bShouldRemoveAsset && (IsFrontendFilterActive() && !PassesCurrentFrontendFilter(RecentlyLoadedOrChangedAsset)))
+		{
+			bShouldRemoveAsset = true;
+		}
 
-				RefreshList();
-			}
+		if (bShouldRemoveAsset)
+		{
+			FilteredAssetItems.RemoveAt(AssetIdx);
+		}
+		else
+		{
+			// Update the asset data on the item
+			ItemAsAsset->SetAssetData(RecentlyLoadedOrChangedAsset);
+
+			// Update the custom column data
+			ItemAsAsset->CacheCustomColumns(CustomColumns, true, true, true);
 		}
 	}
 
+	if (FilteredAssetItemIndices.Num() > 0)
+	{
+		RefreshList();
+	}
+
 	if (FilteredRecentlyAddedAssets.Num() == 0 && RecentlyAddedAssets.Num() == 0)
 	{
 		//No more assets coming in so if we haven't found them now we aren't going to
