In use by any level content browser filter now checks more than one layer deep.

diff --git a/Engine/Source/Editor/ContentBrowser/Private/FrontendFilters.cpp b/Engine/Source/Editor/ContentBrowser/Private/FrontendFilters.cpp
--- a/Engine/Source/Editor/ContentBrowser/Private/FrontendFilters.cpp
+++ b/Engine/Source/Editor/ContentBrowser/Private/FrontendFilters.cpp
@@ -33,17 +33,31 @@ namespace FrontendFilterHelper
 	 */
 	void GetDependencies(const FARFilter& InAssetRegistryFilter, const IAssetRegistry& AssetRegistry, TSet<FName>& OutDependencySet)
 	{
-		TArray<FAssetData> FoundAssets;
-		AssetRegistry.GetAssets(InAssetRegistryFilter, FoundAssets);
+		TArray<FName> PackageNamesToProcess;
+		{
+			TArray<FAssetData> FoundAssets;
+			AssetRegistry.GetAssets(InAssetRegistryFilter, FoundAssets);
+			for (const FAssetData& AssetData : FoundAssets)
+			{
+				PackageNamesToProcess.Add(AssetData.PackageName);
+				OutDependencySet.Add(AssetData.PackageName);
+			}
+		}
 
-		for (const FAssetData& AssetData : FoundAssets)
+		TArray<FAssetIdentifier> AssetDependencies;
+		while (PackageNamesToProcess.Num() > 0)
 		{
-			// Store all the dependencies of all the levels
-			TArray<FAssetIdentifier> AssetDependencies;
-			AssetRegistry.GetDependencies(FAssetIdentifier(AssetData.PackageName), AssetDependencies);
+			const FName PackageName = PackageNamesToProcess.Pop(false);
+			AssetDependencies.Reset();
+			AssetRegistry.GetDependencies(FAssetIdentifier(PackageName), AssetDependencies);
 			for (const FAssetIdentifier& Dependency : AssetDependencies)
 			{
-				OutDependencySet.Add(Dependency.PackageName);
+				bool bIsAlreadyInSet = false;
+				OutDependencySet.Add(Dependency.PackageName, &bIsAlreadyInSet);
+				if (bIsAlreadyInSet == false)
+				{
+					PackageNamesToProcess.Add(Dependency.PackageName);
+				}
 			}
 		}
 	}
