#UE4 Allow projects to change the sort order of root level folders in SPathViews in the editor

diff --git a/Engine/Source/Editor/ContentBrowser/Private/ContentBrowserSingleton.cpp b/Engine/Source/Editor/ContentBrowser/Private/ContentBrowserSingleton.cpp
--- a/Engine/Source/Editor/ContentBrowser/Private/ContentBrowserSingleton.cpp
+++ b/Engine/Source/Editor/ContentBrowser/Private/ContentBrowserSingleton.cpp
@@ -29,6 +29,7 @@
 #include "EmptyFolderVisibilityManager.h"
 #include "CollectionAssetRegistryBridge.h"
 #include "ContentBrowserCommands.h"
+#include "CoreGlobals.h"
 
 #define LOCTEXT_NAMESPACE "ContentBrowser"
 
@@ -68,6 +69,8 @@ FContentBrowserSingleton::FContentBrowserSingleton()
 	FEditorDelegates::LoadSelectedAssetsIfNeeded.AddRaw(this, &FContentBrowserSingleton::OnEditorLoadSelectedAssetsIfNeeded);
 
 	FContentBrowserCommands::Register();
+
+	PopulateConfigValues();
 }
 
 FContentBrowserSingleton::~FContentBrowserSingleton()
@@ -453,6 +456,18 @@ TSharedRef<FEmptyFolderVisibilityManager> FContentBrowserSingleton::GetEmptyFold
 	return EmptyFolderVisibilityManager;
 }
 
+const FContentBrowserPluginSettings& FContentBrowserSingleton::GetPluginSettings(FName PluginName) const
+{
+	const FContentBrowserPluginSettings* LocalSettings = PluginSettings.FindByPredicate([PluginName](const FContentBrowserPluginSettings& Setting) { return Setting.PluginName == PluginName; });
+	if (LocalSettings)
+	{
+		return *LocalSettings;
+	}
+
+	static FContentBrowserPluginSettings DefaultSettings;
+	return DefaultSettings;
+}
+
 void FContentBrowserSingleton::SharedCreateAssetDialogWindow(const TSharedRef<SAssetDialog>& AssetDialog, const FSharedAssetDialogConfig& InConfig, bool bModal) const
 {
 	const FVector2D DefaultWindowSize(1152.0f, 648.0f);
@@ -706,4 +721,25 @@ void FContentBrowserSingleton::ForceShowPluginContent(bool bEnginePlugin)
 	}
 }
 
+void FContentBrowserSingleton::PopulateConfigValues()
+{
+	const FString ContentBrowserSection = TEXT("ContentBrowser");
+
+	// PluginSettings
+	{
+		const FString PluginSettingsName = TEXT("PluginSettings");
+		TArray<FString> PluginSettingsStrings;
+		GConfig->GetArray(*ContentBrowserSection, *PluginSettingsName, PluginSettingsStrings, GEditorIni);
+		for (const FString& PluginSettingString : PluginSettingsStrings)
+		{
+			FContentBrowserPluginSettings PluginSetting;
+			const TCHAR* Result = FContentBrowserPluginSettings::StaticStruct()->ImportText(*PluginSettingString, &PluginSetting, nullptr, PPF_None, GLog, PluginSettingsName);
+			if (Result != nullptr)
+			{
+				PluginSettings.Emplace(PluginSetting);
+			}
+		}
+	}
+}
+
 #undef LOCTEXT_NAMESPACE
diff --git a/Engine/Source/Editor/ContentBrowser/Private/ContentBrowserSingleton.h b/Engine/Source/Editor/ContentBrowser/Private/ContentBrowserSingleton.h
--- a/Engine/Source/Editor/ContentBrowser/Private/ContentBrowserSingleton.h
+++ b/Engine/Source/Editor/ContentBrowser/Private/ContentBrowserSingleton.h
@@ -6,6 +6,7 @@
 #include "CoreMinimal.h"
 #include "AssetData.h"
 #include "IContentBrowserSingleton.h"
+#include "ContentBrowserSingleton.generated.h"
 
 class FCollectionAssetRegistryBridge;
 class FEmptyFolderVisibilityManager;
@@ -18,6 +19,23 @@ class UFactory;
 
 #define MAX_CONTENT_BROWSERS 4
 
+USTRUCT()
+struct FContentBrowserPluginSettings
+{
+	GENERATED_BODY()
+
+	UPROPERTY()
+	FName PluginName;
+
+	/** Used to control the order of plugin root folders in the path view. A higher priority sorts higher in the list. Game and Engine folders are priority 1.0 */
+	UPROPERTY()
+	float RootFolderSortPriority;
+
+	FContentBrowserPluginSettings()
+		: RootFolderSortPriority(0.f)
+	{}
+};
+
 /**
  * Content browser module singleton implementation class
  */
@@ -65,6 +83,9 @@ public:
 
 	TSharedRef<FEmptyFolderVisibilityManager> GetEmptyFolderVisibilityManager();
 
+	/** Gets the settings for the plugin with the specified name */
+	const FContentBrowserPluginSettings& GetPluginSettings(FName PluginName) const;
+
 	/** Single storage location for content browser favorites */
 	TArray<FString> FavoriteFolderPaths;
 
@@ -102,6 +123,9 @@ private:
 	/** Returns a localized name for the tab/menu entry with index */
 	static FText GetContentBrowserLabelWithIndex( int32 BrowserIdx );
 
+	/** Populates properties that come from ini files */
+	void PopulateConfigValues();
+
 public:
 	/** The tab identifier/instance name for content browser tabs */
 	FName ContentBrowserTabIDs[MAX_CONTENT_BROWSERS];
@@ -119,6 +143,8 @@ private:
 
 	TSharedRef<FCollectionAssetRegistryBridge> CollectionAssetRegistryBridge;
 
+	TArray<FContentBrowserPluginSettings> PluginSettings;
+
 	/** An incrementing int32 which is used when making unique settings strings */
 	int32 SettingsStringID;
 };
diff --git a/Engine/Source/Editor/ContentBrowser/Private/SPathView.cpp b/Engine/Source/Editor/ContentBrowser/Private/SPathView.cpp
--- a/Engine/Source/Editor/ContentBrowser/Private/SPathView.cpp
+++ b/Engine/Source/Editor/ContentBrowser/Private/SPathView.cpp
@@ -1180,25 +1180,60 @@ void SPathView::SortRootItems()
 
 	// We have some manual sorting requirements that game must come before engine, and engine before everything else - we do that here after sorting everything by name
 	// The array below is in the inverse order as we iterate through and move each match to the beginning of the root items array
-	static const FString InverseSortOrder[] = {
-		TEXT("Classes_Engine"),
-		TEXT("Engine"),
-		TEXT("Classes_Game"),
+	const TArray<FString> SpecialDefaultFolders = {
 		TEXT("Game"),
+		TEXT("Classes_Game"),
+		TEXT("Engine"),
+		TEXT("Classes_Engine"),
 	};
-	for(const FString& SortItem : InverseSortOrder)
+
+	const FString ClassesPrefix = TEXT("Classes_");
+
+	struct FRootItemSortInfo
+	{
+		FString FolderName;
+		float Priority;
+		int32 SpecialDefaultFolderPriority;
+		bool bIsClassesFolder;
+	};
+
+	TMap<FTreeItem*, FRootItemSortInfo> SortInfoMap;
+	for (const TSharedPtr<FTreeItem>& RootItem : TreeRootItems)
 	{
-		const int32 FoundItemIndex = TreeRootItems.IndexOfByPredicate([&SortItem](const TSharedPtr<FTreeItem>& TreeItem) -> bool
+		FRootItemSortInfo SortInfo;
+		SortInfo.bIsClassesFolder = RootItem->FolderName.StartsWith(ClassesPrefix);
+		SortInfo.FolderName = SortInfo.bIsClassesFolder ? RootItem->FolderName.RightChop(ClassesPrefix.Len()) : RootItem->FolderName;
+		int32 SpecialDefaultFolderIdx = SpecialDefaultFolders.IndexOfByKey(RootItem->FolderName);
+		SortInfo.SpecialDefaultFolderPriority = SpecialDefaultFolderIdx != INDEX_NONE ? SpecialDefaultFolders.Num() - SpecialDefaultFolderIdx : 0;
+		SortInfo.Priority = SpecialDefaultFolderIdx == INDEX_NONE ? FContentBrowserSingleton::Get().GetPluginSettings(FName(*SortInfo.FolderName)).RootFolderSortPriority : 1.f;
+		SortInfoMap.Add(RootItem.Get(), SortInfo);
+	}
+
+	
+	TreeRootItems.Sort([&SortInfoMap](const TSharedPtr<FTreeItem>& RootItemA, const TSharedPtr<FTreeItem>& RootItemB) {
+		const FRootItemSortInfo& SortInfoA = SortInfoMap.FindChecked(RootItemA.Get());
+		const FRootItemSortInfo& SortInfoB = SortInfoMap.FindChecked(RootItemB.Get());
+		if (SortInfoA.Priority != SortInfoB.Priority)
 		{
-			return TreeItem->FolderName == SortItem;
-		});
-		if(FoundItemIndex != INDEX_NONE)
+			// Not the same priority, use priority to sort
+			return SortInfoA.Priority > SortInfoB.Priority;
+		}
+		else if (SortInfoA.SpecialDefaultFolderPriority != SortInfoB.SpecialDefaultFolderPriority)
 		{
-			TSharedPtr<FTreeItem> ItemToMove = TreeRootItems[FoundItemIndex];
-			TreeRootItems.RemoveAt(FoundItemIndex);
-			TreeRootItems.Insert(ItemToMove, 0);
+			// Special folders use the index to sort. Non special folders are all set to 0.
+			return SortInfoA.SpecialDefaultFolderPriority > SortInfoB.SpecialDefaultFolderPriority;
 		}
-	}
+		else if (SortInfoA.FolderName != SortInfoB.FolderName)
+		{
+			// Two non special folders of the same priorty, sort alphabetically
+			return SortInfoA.FolderName < SortInfoB.FolderName;
+		}
+		else
+		{
+			// Classes folders have the same name so sort them adjacent but under non-classes
+			return !SortInfoA.bIsClassesFolder;
+		}
+	});
 
 	TreeViewPtr->RequestTreeRefresh();
 }
