Add ability for asset actions to provide submenu array to nest under asset creation context menu

diff --git a/Engine/Source/Developer/AssetTools/Public/AssetTypeActions_Base.h b/Engine/Source/Developer/AssetTools/Public/AssetTypeActions_Base.h
--- a/Engine/Source/Developer/AssetTools/Public/AssetTypeActions_Base.h
+++ b/Engine/Source/Developer/AssetTools/Public/AssetTypeActions_Base.h
@@ -124,6 +124,11 @@ PRAGMA_ENABLE_DEPRECATION_WARNINGS
 		return TPT_None;
 	}
 
+	virtual const TArray<FText>& GetSubMenus() const override
+	{
+		static const TArray<FText> SubMenus;
+		return SubMenus;
+	}
 
 	virtual TSharedPtr<class SWidget> GetThumbnailOverlay(const FAssetData& AssetData) const override
 	{
diff --git a/Engine/Source/Developer/AssetTools/Public/IAssetTypeActions.h b/Engine/Source/Developer/AssetTools/Public/IAssetTypeActions.h
--- a/Engine/Source/Developer/AssetTools/Public/IAssetTypeActions.h
+++ b/Engine/Source/Developer/AssetTools/Public/IAssetTypeActions.h
@@ -84,9 +84,12 @@ public:
 	/** Begins a merge between the specified assets */
 	virtual void Merge(UObject* BaseAsset, UObject* RemoteAsset, UObject* LocalAsset, const FOnMergeResolved& ResolutionCallback) = 0;
 
-	/** Returns the categories that this asset type. The return value is one or more flags from EAssetTypeCategories.  */
+	/** Returns the categories that this asset type appears in. The return value is one or more flags from EAssetTypeCategories.  */
 	virtual uint32 GetCategories() = 0;
 
+	/** Returns array of sub-menu names that this asset type is parented under in the Asset Creation Context Menu. */
+	virtual const TArray<FText>& GetSubMenus() const = 0;
+
 	/** @return True if we should force world-centric mode for newly-opened assets */
 	virtual bool ShouldForceWorldCentric() = 0;
 
diff --git a/Engine/Source/Editor/ContentBrowser/Private/NewAssetOrClassContextMenu.cpp b/Engine/Source/Editor/ContentBrowser/Private/NewAssetOrClassContextMenu.cpp
--- a/Engine/Source/Editor/ContentBrowser/Private/NewAssetOrClassContextMenu.cpp
+++ b/Engine/Source/Editor/ContentBrowser/Private/NewAssetOrClassContextMenu.cpp
@@ -26,15 +26,29 @@
 
 #define LOCTEXT_NAMESPACE "ContentBrowser"
 
-struct FFactoryItem
+
+void FCategorySubMenuItem::SortSubMenus(FCategorySubMenuItem* SubMenu)
 {
-	UFactory* Factory;
-	FText DisplayName;
+	if (!SubMenu)
+	{
+		SubMenu = this;
+	}
 
-	FFactoryItem(UFactory* InFactory, const FText& InDisplayName)
-		: Factory(InFactory), DisplayName(InDisplayName)
-	{}
-};
+	// Sort the factories by display name
+	SubMenu->Factories.Sort([](const FFactoryItem& A, const FFactoryItem& B) -> bool
+	{
+		return A.DisplayName.CompareToCaseIgnored(B.DisplayName) < 0;
+	});
+
+	for (TPair<FString, TSharedPtr<FCategorySubMenuItem>>& Pair : SubMenu->Children)
+	{
+		if (Pair.Value.IsValid())
+		{
+			FCategorySubMenuItem* MenuData = Pair.Value.Get();
+			SortSubMenus(MenuData);
+		}
+	}
+}
 
 TArray<FFactoryItem> FindFactoriesInCategory(EAssetTypeCategories::Type AssetTypeCategory)
 {
@@ -389,30 +403,88 @@ void FNewAssetOrClassContextMenu::CreateNewAssetMenuCategory(FMenuBuilder& MenuB
 {
 	// Find UFactory classes that can create new objects in this category.
 	TArray<FFactoryItem> FactoriesInThisCategory = FindFactoriesInCategory(AssetTypeCategory);
+	if (FactoriesInThisCategory.Num() == 0)
+	{
+		return;
+	}
 
-	// Sort the list
-	struct FCompareFactoryDisplayNames
+	TSharedPtr<FCategorySubMenuItem> ParentMenuData = MakeShareable(new FCategorySubMenuItem);
+	for (FFactoryItem& Item : FactoriesInThisCategory)
 	{
-		FORCEINLINE bool operator()( const FFactoryItem& A, const FFactoryItem& B ) const
+		FCategorySubMenuItem* SubMenu = ParentMenuData.Get();
+		const TArray<FText>& CategoryNames = Item.Factory->GetMenuCategorySubMenus();
+		for (FText CategoryName : CategoryNames)
 		{
-			return A.DisplayName.CompareToCaseIgnored(B.DisplayName) < 0;
+			const FString SourceString = CategoryName.BuildSourceString();
+			if (TSharedPtr<FCategorySubMenuItem> SubMenuData = SubMenu->Children.FindRef(SourceString))
+			{
+				check(SubMenuData.IsValid());
+				SubMenu = SubMenuData.Get();
+			}
+			else
+			{
+				TSharedPtr<FCategorySubMenuItem> NewSubMenu = MakeShared<FCategorySubMenuItem>();
+				NewSubMenu->Name = CategoryName;
+				SubMenu->Children.Add(SourceString, NewSubMenu);
+				SubMenu = NewSubMenu.Get();
+			}
 		}
-	};
-	FactoriesInThisCategory.Sort( FCompareFactoryDisplayNames() );
+		SubMenu->Factories.Add(Item);
+	}
+	ParentMenuData->SortSubMenus();
+	CreateNewAssetMenus(MenuBuilder, ParentMenuData, InPath, InOnNewAssetRequested, InCanExecuteAction);
+}
 
-	// Add menu entries for each one
-	for ( auto FactoryIt = FactoriesInThisCategory.CreateConstIterator(); FactoryIt; ++FactoryIt )
+void FNewAssetOrClassContextMenu::CreateNewAssetMenus(FMenuBuilder& MenuBuilder, TSharedPtr<FCategorySubMenuItem> SubMenuData, FString InPath, FOnNewAssetRequested InOnNewAssetRequested, FCanExecuteAction InCanExecuteAction)
+{
+	for (const FFactoryItem& FactoryItem : SubMenuData->Factories)
 	{
-		UFactory* Factory = (*FactoryIt).Factory;
-		TWeakObjectPtr<UClass> WeakFactoryClass = Factory->GetClass();
+		TWeakObjectPtr<UClass> WeakFactoryClass = FactoryItem.Factory->GetClass();
 
 		MenuBuilder.AddMenuEntry(
 			FUIAction(
-				FExecuteAction::CreateStatic( &FNewAssetOrClassContextMenu::ExecuteNewAsset, InPath, WeakFactoryClass, InOnNewAssetRequested ),
+				FExecuteAction::CreateStatic(&FNewAssetOrClassContextMenu::ExecuteNewAsset, InPath, WeakFactoryClass, InOnNewAssetRequested),
+				InCanExecuteAction
+			),
+			SNew(SFactoryMenuEntry, FactoryItem.Factory)
+		);
+	}
+
+	if (SubMenuData->Children.Num() == 0)
+	{
+		return;
+	}
+
+	MenuBuilder.AddMenuSeparator();
+
+	TArray<TSharedPtr<FCategorySubMenuItem>> SortedMenus;
+	SubMenuData->Children.GenerateValueArray(SortedMenus);
+	SortedMenus.Sort([](const TSharedPtr<FCategorySubMenuItem>& A, const TSharedPtr<FCategorySubMenuItem>& B) -> bool
+	{
+		return A->Name.CompareToCaseIgnored(B->Name) < 0;
+	});
+
+	for (TSharedPtr<FCategorySubMenuItem>& ChildMenuData : SortedMenus)
+	{
+		check(ChildMenuData.IsValid());
+
+		MenuBuilder.AddSubMenu(
+			ChildMenuData->Name,
+			FText::GetEmpty(),
+			FNewMenuDelegate::CreateStatic(
+				&FNewAssetOrClassContextMenu::CreateNewAssetMenus,
+				ChildMenuData,
+				InPath,
+				InOnNewAssetRequested,
+				InCanExecuteAction
+			),
+			FUIAction(
+				FExecuteAction(),
 				InCanExecuteAction
-				),
-			SNew( SFactoryMenuEntry, Factory )
-			);
+			),
+			NAME_None,
+			EUserInterfaceActionType::Button
+		);
 	}
 }
 
diff --git a/Engine/Source/Editor/ContentBrowser/Private/NewAssetOrClassContextMenu.h b/Engine/Source/Editor/ContentBrowser/Private/NewAssetOrClassContextMenu.h
--- a/Engine/Source/Editor/ContentBrowser/Private/NewAssetOrClassContextMenu.h
+++ b/Engine/Source/Editor/ContentBrowser/Private/NewAssetOrClassContextMenu.h
@@ -5,9 +5,30 @@
 #include "CoreMinimal.h"
 #include "Framework/Commands/UIAction.h"
 #include "AssetTypeCategories.h"
+#include "Templates/SharedPointer.h"
+
 
 class FMenuBuilder;
 
+struct FFactoryItem
+{
+	UFactory* Factory;
+	FText DisplayName;
+
+	FFactoryItem(UFactory* InFactory, const FText& InDisplayName)
+		: Factory(InFactory), DisplayName(InDisplayName)
+	{}
+};
+
+struct FCategorySubMenuItem
+{
+	FText Name;
+	TArray<FFactoryItem> Factories;
+	TMap<FString, TSharedPtr<FCategorySubMenuItem>> Children;
+
+	void SortSubMenus(FCategorySubMenuItem* NextNode = nullptr);
+};
+
 class FNewAssetOrClassContextMenu
 {
 public:
@@ -43,6 +64,8 @@ private:
 	/** Handle creating a new asset from an asset category */
 	static void CreateNewAssetMenuCategory(FMenuBuilder& MenuBuilder, EAssetTypeCategories::Type AssetTypeCategory, FString InPath, FOnNewAssetRequested InOnNewAssetRequested, FCanExecuteAction InCanExecuteAction);
 
+	static void CreateNewAssetMenus(FMenuBuilder& MenuBuilder, TSharedPtr<FCategorySubMenuItem> SubMenuData, FString InPath, FOnNewAssetRequested InOnNewAssetRequested, FCanExecuteAction InCanExecuteAction);
+
 	/** Handle when the "Import" button is clicked */
 	static void ExecuteImportAsset(FOnImportAssetRequested InOnImportAssetRequested, FString InPath);
 
diff --git a/Engine/Source/Editor/UnrealEd/Classes/Factories/Factory.h b/Engine/Source/Editor/UnrealEd/Classes/Factories/Factory.h
--- a/Engine/Source/Editor/UnrealEd/Classes/Factories/Factory.h
+++ b/Engine/Source/Editor/UnrealEd/Classes/Factories/Factory.h
@@ -122,6 +122,9 @@ public:
 	/** When shown in menus, this is the category containing this factory. Return type is a BitFlag mask using EAssetTypeCategories. */
 	virtual uint32 GetMenuCategories() const;
 
+	/** Branch of sub-menus containing factory under each provided category. */
+	virtual const TArray<FText>& GetMenuCategorySubMenus() const;
+
 	/** Returns the tooltip text description of this factory */
 	virtual FText GetToolTip() const;
 
diff --git a/Engine/Source/Editor/UnrealEd/Private/Factories/Factory.cpp b/Engine/Source/Editor/UnrealEd/Private/Factories/Factory.cpp
--- a/Engine/Source/Editor/UnrealEd/Private/Factories/Factory.cpp
+++ b/Engine/Source/Editor/UnrealEd/Private/Factories/Factory.cpp
@@ -258,6 +258,23 @@ uint32 UFactory::GetMenuCategories() const
 	return EAssetTypeCategories::Misc;
 }
 
+const TArray<FText>& UFactory::GetMenuCategorySubMenus() const
+{
+	FAssetToolsModule& AssetToolsModule = FModuleManager::LoadModuleChecked<FAssetToolsModule>("AssetTools");
+	UClass* LocalSupportedClass = GetSupportedClass();
+
+	if (LocalSupportedClass)
+	{
+		TWeakPtr<IAssetTypeActions> AssetTypeActions = AssetToolsModule.Get().GetAssetTypeActionsForClass(LocalSupportedClass);
+		if (AssetTypeActions.IsValid())
+		{
+			return AssetTypeActions.Pin()->GetSubMenus();
+		}
+	}
+
+	const static TArray<FText> DefaultSubCategories;
+	return DefaultSubCategories;
+}
 
 FText UFactory::GetToolTip() const
 {
