Missing symbol fix for vs2019 non-unity build

diff --git a/Engine/Source/Editor/ContentBrowser/Private/NewAssetOrClassContextMenu.cpp b/Engine/Source/Editor/ContentBrowser/Private/NewAssetOrClassContextMenu.cpp
--- a/Engine/Source/Editor/ContentBrowser/Private/NewAssetOrClassContextMenu.cpp
+++ b/Engine/Source/Editor/ContentBrowser/Private/NewAssetOrClassContextMenu.cpp
@@ -27,6 +27,12 @@
 #define LOCTEXT_NAMESPACE "ContentBrowser"
 
 
+FFactoryItem::FFactoryItem(UFactory* InFactory, const FText& InDisplayName)
+	: Factory(InFactory)
+	, DisplayName(InDisplayName)
+{
+}
+
 void FCategorySubMenuItem::SortSubMenus(FCategorySubMenuItem* SubMenu)
 {
 	if (!SubMenu)
diff --git a/Engine/Source/Editor/ContentBrowser/Private/NewAssetOrClassContextMenu.h b/Engine/Source/Editor/ContentBrowser/Private/NewAssetOrClassContextMenu.h
--- a/Engine/Source/Editor/ContentBrowser/Private/NewAssetOrClassContextMenu.h
+++ b/Engine/Source/Editor/ContentBrowser/Private/NewAssetOrClassContextMenu.h
@@ -9,15 +9,14 @@
 
 
 class FMenuBuilder;
+class UFactory;
 
 struct FFactoryItem
 {
 	UFactory* Factory;
 	FText DisplayName;
 
-	FFactoryItem(UFactory* InFactory, const FText& InDisplayName)
-		: Factory(InFactory), DisplayName(InDisplayName)
-	{}
+	FFactoryItem(UFactory* InFactory, const FText& InDisplayName);
 };
 
 struct FCategorySubMenuItem
