#UE4 Fix bug where packages opened for diffing were emitting an error due to being queried for tags in SAssetView

diff --git a/Engine/Source/Editor/ContentBrowser/Private/SAssetView.cpp b/Engine/Source/Editor/ContentBrowser/Private/SAssetView.cpp
--- a/Engine/Source/Editor/ContentBrowser/Private/SAssetView.cpp
+++ b/Engine/Source/Editor/ContentBrowser/Private/SAssetView.cpp
@@ -2898,6 +2898,12 @@ void SAssetView::OnAssetLoaded(UObject* Asset)
 		return;
 	}
 
+	if (Asset->GetOutermost()->HasAnyPackageFlags(PKG_ForDiffing))
+	{
+		// No need to consider packages loaded for diffing purposes
+		return;
+	}
+
 	FName AssetPathName = FName(*Asset->GetPathName());
 	RecentlyLoadedOrChangedAssets.Add( FAssetData(Asset) );
 
