Fixed crash when creating a new asset with CollectGarbageEveryFrame enabled

diff --git a/Engine/Source/Editor/ContentBrowser/Private/SAssetView.cpp b/Engine/Source/Editor/ContentBrowser/Private/SAssetView.cpp
--- a/Engine/Source/Editor/ContentBrowser/Private/SAssetView.cpp
+++ b/Engine/Source/Editor/ContentBrowser/Private/SAssetView.cpp
@@ -736,7 +736,7 @@ void SAssetView::OnCreateNewFolder(const FString& FolderName, const FString& Fol
 	// Make sure we are showing the location of the new folder (we may have created it in a folder)
 	OnPathSelected.Execute(FolderPath);
 
-	DeferredFolderToCreate = MakeShareable(new FCreateDeferredFolderData());
+	DeferredFolderToCreate = MakeUnique<FCreateDeferredFolderData>();
 	DeferredFolderToCreate->FolderName = FolderName;
 	DeferredFolderToCreate->FolderPath = FolderPath;
 }
@@ -783,7 +783,7 @@ void SAssetView::CreateNewAsset(const FString& DefaultAssetName, const FString&
 	OnPathSelected.Execute(PackagePath);
 
 	// Defer asset creation until next tick, so we get a chance to refresh the view
-	DeferredAssetToCreate = MakeShareable(new FCreateDeferredAssetData());
+	DeferredAssetToCreate = MakeUnique<FCreateDeferredAssetData>();
 	DeferredAssetToCreate->DefaultAssetName = DefaultAssetName;
 	DeferredAssetToCreate->PackagePath = PackagePath;
 	DeferredAssetToCreate->AssetClass = AssetClass;
@@ -1780,6 +1780,14 @@ void SAssetView::OnFocusChanging( const FWeakWidgetPath& PreviousFocusPath, cons
 	ResetQuickJump();
 }
 
+void SAssetView::AddReferencedObjects(FReferenceCollector& Collector)
+{
+	if (DeferredAssetToCreate)
+	{
+		DeferredAssetToCreate->AddReferencedObjects(Collector);
+	}
+}
+
 TSharedRef<SAssetTileView> SAssetView::CreateTileView()
 {
 	return SNew(SAssetTileView)
diff --git a/Engine/Source/Editor/ContentBrowser/Private/SAssetView.h b/Engine/Source/Editor/ContentBrowser/Private/SAssetView.h
--- a/Engine/Source/Editor/ContentBrowser/Private/SAssetView.h
+++ b/Engine/Source/Editor/ContentBrowser/Private/SAssetView.h
@@ -3,6 +3,7 @@
 #pragma once
 
 #include "CoreMinimal.h"
+#include "UObject/GCObject.h"
 #include "Misc/Attribute.h"
 #include "Input/Reply.h"
 #include "Layout/Visibility.h"
@@ -41,7 +42,7 @@ DECLARE_DELEGATE(FOnSearchOptionChanged);
 /**
  * A widget to display a list of filtered assets
  */
-class SAssetView : public SCompoundWidget
+class SAssetView : public SCompoundWidget, public FGCObject
 {
 public:
 	SLATE_BEGIN_ARGS( SAssetView )
@@ -292,6 +293,10 @@ public:
 	virtual FReply OnMouseWheel( const FGeometry& MyGeometry, const FPointerEvent& MouseEvent ) override;
 	virtual void OnFocusChanging( const FWeakWidgetPath& PreviousFocusPath, const FWidgetPath& NewWidgetPath, const FFocusEvent& InFocusEvent ) override;
 
+	//~ FGCObject inherited
+	virtual void AddReferencedObjects(FReferenceCollector& Collector) override;
+	virtual FString GetReferencerName() const override { return TEXT("SAssetView"); }
+
 	/** Opens the selected assets or folders, depending on the selection */
 	void OnOpenAssetsOrFolders();
 
@@ -1047,10 +1052,16 @@ private:
 
 		/** The factory to use */
 		UFactory* Factory;
+
+		void AddReferencedObjects(FReferenceCollector& Collector)
+		{
+			Collector.AddReferencedObject(AssetClass);
+			Collector.AddReferencedObject(Factory);
+		}
 	};
 
 	/** Asset pending deferred creation */
-	TSharedPtr<FCreateDeferredAssetData> DeferredAssetToCreate;
+	TUniquePtr<FCreateDeferredAssetData> DeferredAssetToCreate;
 
 	/** A struct to hold data for the deferred creation of a folder */
 	struct FCreateDeferredFolderData
@@ -1063,7 +1074,7 @@ private:
 	};
 
 	/** Folder pending deferred creation */
-	TSharedPtr<FCreateDeferredFolderData> DeferredFolderToCreate;
+	TUniquePtr<FCreateDeferredFolderData> DeferredFolderToCreate;
 
 	/** Struct holding the data for the asset quick-jump */
 	struct FQuickJumpData
diff --git a/Engine/Source/Editor/ContentBrowser/Private/SContentBrowser.cpp b/Engine/Source/Editor/ContentBrowser/Private/SContentBrowser.cpp
--- a/Engine/Source/Editor/ContentBrowser/Private/SContentBrowser.cpp
+++ b/Engine/Source/Editor/ContentBrowser/Private/SContentBrowser.cpp
@@ -54,6 +54,7 @@
 #include "Widgets/Docking/SDockTab.h"
 #include "Framework/Commands/GenericCommands.h"
 #include "IAddContentDialogModule.h"
+#include "UObject/GCObjectScopeGuard.h"
 #include "Engine/Selection.h"
 #include "NativeClassHierarchy.h"
 #include "AddToProjectConfig.h"
@@ -1639,8 +1640,9 @@ void SContentBrowser::NewAssetRequested(const FString& SelectedPath, TWeakObject
 	if ( ensure(SelectedPath.Len() > 0) && ensure(FactoryClass.IsValid()) )
 	{
 		UFactory* NewFactory = NewObject<UFactory>(GetTransientPackage(), FactoryClass.Get());
+
 		// This factory may get gc'd as a side effect of various delegates potentially calling CollectGarbage so protect against it from being gc'd out from under us
-		NewFactory->AddToRoot();
+		FGCObjectScopeGuard FactoryGCGuard(NewFactory);
 
 		FEditorDelegates::OnConfigureNewAssetProperties.Broadcast(NewFactory);
 		if ( NewFactory->ConfigureProperties() )
@@ -1653,7 +1655,6 @@ void SContentBrowser::NewAssetRequested(const FString& SelectedPath, TWeakObject
 			AssetToolsModule.Get().CreateUniqueAssetName(SelectedPath + TEXT("/") + NewFactory->GetDefaultNewAssetName(), TEXT(""), PackageNameToUse, DefaultAssetName);
 			CreateNewAsset(DefaultAssetName, SelectedPath, NewFactory->GetSupportedClass(), NewFactory);
 		}
-		NewFactory->RemoveFromRoot();
 	}
 }
 
