Added a hotfix to avoid a possible crash when closing the editor.

Under some currently unknown circumstances, this widget sometime outlive the SlateApplicationBase::Shutdown.
If during the shutdown of a module a path is unmounted, the broadcast will make this widget try to register active timer.
Registering a active require the SlateApplicationBase to be in a working state (initialized) otherwise it will crash.

diff --git a/Engine/Source/Editor/ContentBrowser/Private/SPathView.cpp b/Engine/Source/Editor/ContentBrowser/Private/SPathView.cpp
--- a/Engine/Source/Editor/ContentBrowser/Private/SPathView.cpp
+++ b/Engine/Source/Editor/ContentBrowser/Private/SPathView.cpp
@@ -25,6 +25,8 @@
 #include "NativeClassHierarchy.h"
 #include "EmptyFolderVisibilityManager.h"
 
+#include "Application/SlateApplicationBase.h"
+
 #define LOCTEXT_NAMESPACE "ContentBrowser"
 
 SPathView::~SPathView()
@@ -1579,8 +1581,16 @@ void SPathView::OnFolderPopulated(const FString& Path)
 
 void SPathView::OnContentPathMountedOrDismounted( const FString& AssetPath, const FString& FilesystemPath )
 {
-	// A new content path has appeared, so we should refresh out root set of paths
-	RegisterActiveTimer(0.f, FWidgetActiveTimerDelegate::CreateSP(this, &SPathView::TriggerRepopulate));
+	/**
+	 * Hotfix
+	 * For some reason this widget sometime outlive the slate application shutdown
+	 * Validating that Slate application base is initialized will at least avoid the possible crash
+	 */
+	if ( FSlateApplicationBase::IsInitialized() )
+	{
+		// A new content path has appeared, so we should refresh out root set of paths
+		RegisterActiveTimer(0.f, FWidgetActiveTimerDelegate::CreateSP(this, &SPathView::TriggerRepopulate));
+	}
 }
 
 void SPathView::OnClassHierarchyUpdated()
