It is now possible to define a whitelist of asset classes that an editor session is allowed to use in the DefaultEditor.ini section [AssetTools] in an array called SupportedAssetTypes. Types that are not supported will not appear in the content browser or in context menus to create new assets.

#rnx

diff --git a/Engine/Source/Developer/AssetTools/Private/AssetTools.cpp b/Engine/Source/Developer/AssetTools/Private/AssetTools.cpp
--- a/Engine/Source/Developer/AssetTools/Private/AssetTools.cpp
+++ b/Engine/Source/Developer/AssetTools/Private/AssetTools.cpp
@@ -130,13 +130,13 @@
 #include "SAdvancedCopyReportDialog.h"
 #include "AssetToolsSettings.h"
 #include "AssetVtConversion.h"
+#include "Misc/ConfigCacheIni.h"
 #if WITH_EDITOR
 #include "Subsystems/AssetEditorSubsystem.h"
 #endif
 
 #define LOCTEXT_NAMESPACE "AssetTools"
 
- 
 TScriptInterface<IAssetTools> UAssetToolsHelpers::GetAssetTools()
 {
 	return &UAssetToolsImpl::Get();
@@ -154,6 +154,14 @@ UAssetToolsImpl::UAssetToolsImpl(const FObjectInitializer& ObjectInitializer)
 	, AssetFixUpRedirectors(MakeShareable(new FAssetFixUpRedirectors))
 	, NextUserCategoryBit(EAssetTypeCategories::FirstUser)
 {
+	TArray<FString> SupportedTypesArray;
+	GConfig->GetArray(TEXT("AssetTools"), TEXT("SupportedAssetTypes"), SupportedTypesArray, GEditorIni);
+
+	for (const FString& Type : SupportedTypesArray)
+	{
+		SupportedAssetTypes.Add(*Type);
+	}
+
 	// Register the built-in advanced categories
 	AllocatedCategoryBits.Add(TEXT("_BuiltIn_0"), FAdvancedAssetCategory(EAssetTypeCategories::Animation, LOCTEXT("AnimationAssetCategory", "Animation")));
 	AllocatedCategoryBits.Add(TEXT("_BuiltIn_1"), FAdvancedAssetCategory(EAssetTypeCategories::Blueprint, LOCTEXT("BlueprintAssetCategory", "Blueprints")));
@@ -236,7 +244,7 @@ UAssetToolsImpl::UAssetToolsImpl(const FObjectInitializer& ObjectInitializer)
 	RegisterAssetTypeActions(MakeShareable(new FAssetTypeActions_Texture2D));
 	RegisterAssetTypeActions(MakeShareable(new FAssetTypeActions_TextureCube));
 	RegisterAssetTypeActions(MakeShareable(new FAssetTypeActions_Texture2DArray));
-	RegisterAssetTypeActions(MakeShareable(new FAssetTypeActions_VolumeTexture) );
+	RegisterAssetTypeActions( MakeShareable(new FAssetTypeActions_VolumeTexture) );
 	RegisterAssetTypeActions(MakeShareable(new FAssetTypeActions_TextureRenderTarget));
 	RegisterAssetTypeActions(MakeShareable(new FAssetTypeActions_TextureRenderTarget2D));
 	RegisterAssetTypeActions(MakeShareable(new FAssetTypeActions_TextureRenderTargetCube));
@@ -253,6 +261,23 @@ UAssetToolsImpl::UAssetToolsImpl(const FObjectInitializer& ObjectInitializer)
 
 void UAssetToolsImpl::RegisterAssetTypeActions(const TSharedRef<IAssetTypeActions>& NewActions)
 {
+	if (SupportedAssetTypes.Num() > 0)
+	{
+		const UClass* SupportedClass = NewActions->GetSupportedClass();
+		if (SupportedClass != nullptr)
+		{
+			const FName ClassName = SupportedClass->GetFName();
+			if (!SupportedAssetTypes.Contains(ClassName))
+			{
+				NewActions->SetSupported(false);
+			}
+		}
+		else
+		{
+			NewActions->SetSupported(false);
+		}
+	}
+
 	AssetTypeActionsList.Add(NewActions);
 }
 
@@ -269,7 +294,7 @@ void UAssetToolsImpl::GetAssetTypeActionsList( TArray<TWeakPtr<IAssetTypeActions
 	}
 }
 
-TWeakPtr<IAssetTypeActions> UAssetToolsImpl::GetAssetTypeActionsForClass( UClass* Class ) const
+TWeakPtr<IAssetTypeActions> UAssetToolsImpl::GetAssetTypeActionsForClass(const UClass* Class) const
 {
 	TSharedPtr<IAssetTypeActions> MostDerivedAssetTypeActions;
 
@@ -290,7 +315,7 @@ TWeakPtr<IAssetTypeActions> UAssetToolsImpl::GetAssetTypeActionsForClass( UClass
 	return MostDerivedAssetTypeActions;
 }
 
-TArray<TWeakPtr<IAssetTypeActions>> UAssetToolsImpl::GetAssetTypeActionsListForClass(UClass* Class) const
+TArray<TWeakPtr<IAssetTypeActions>> UAssetToolsImpl::GetAssetTypeActionsListForClass(const UClass* Class) const
 {
 	TArray<TWeakPtr<IAssetTypeActions>> ResultAssetTypeActionsList;
 
@@ -3068,4 +3093,44 @@ void UAssetToolsImpl::AdvancedCopyPackages_ReportConfirmed(FAdvancedCopyParams C
 	AdvancedCopyPackages(CopyParams, DestinationMap);
 }
 
+bool UAssetToolsImpl::IsAssetClassSupported(const UClass* AssetClass) const
+{
+	TWeakPtr<IAssetTypeActions> AssetTypeActions = GetAssetTypeActionsForClass(AssetClass);
+	if (!AssetTypeActions.IsValid())
+	{
+		return false;
+	}
+
+	if (AssetTypeActions.Pin()->IsSupported() == false)
+	{
+		return false;
+	}
+
+	return true;
+}
+
+TArray<UFactory*> UAssetToolsImpl::GetNewAssetFactories() const
+{
+	TArray<UFactory*> Factories;
+
+	for (TObjectIterator<UClass> It; It; ++It)
+	{
+		UClass* Class = *It;
+		if (Class->IsChildOf(UFactory::StaticClass()) &&
+			!Class->HasAnyClassFlags(CLASS_Abstract))
+		{
+			UFactory* Factory = Class->GetDefaultObject<UFactory>();
+
+			if (Factory->ShouldShowInNewMenu() &&
+				ensure(!Factory->GetDisplayName().IsEmpty()) &&
+				IsAssetClassSupported(Factory->GetSupportedClass()))
+			{
+				Factories.Add(Factory);
+			}
+		}
+	}
+
+	return MoveTemp(Factories);
+}
+
 #undef LOCTEXT_NAMESPACE
diff --git a/Engine/Source/Developer/AssetTools/Private/AssetTools.h b/Engine/Source/Developer/AssetTools/Private/AssetTools.h
--- a/Engine/Source/Developer/AssetTools/Private/AssetTools.h
+++ b/Engine/Source/Developer/AssetTools/Private/AssetTools.h
@@ -63,8 +63,8 @@ public:
 	virtual void RegisterAssetTypeActions(const TSharedRef<IAssetTypeActions>& NewActions) override;
 	virtual void UnregisterAssetTypeActions(const TSharedRef<IAssetTypeActions>& ActionsToRemove) override;
 	virtual void GetAssetTypeActionsList( TArray<TWeakPtr<IAssetTypeActions>>& OutAssetTypeActionsList ) const override;
-	virtual TWeakPtr<IAssetTypeActions> GetAssetTypeActionsForClass( UClass* Class ) const override;
-	virtual TArray<TWeakPtr<IAssetTypeActions>> GetAssetTypeActionsListForClass(UClass* Class) const override;
+	virtual TWeakPtr<IAssetTypeActions> GetAssetTypeActionsForClass(const UClass* Class) const override;
+	virtual TArray<TWeakPtr<IAssetTypeActions>> GetAssetTypeActionsListForClass(const UClass* Class) const override;
 	virtual EAssetTypeCategories::Type RegisterAdvancedAssetCategory(FName CategoryKey, FText CategoryDisplayName) override;
 	virtual EAssetTypeCategories::Type FindAdvancedAssetCategory(FName CategoryKey) const override;
 	virtual void GetAllAdvancedAssetCategories(TArray<FAdvancedAssetCategory>& OutCategoryList) const override;
@@ -91,7 +91,6 @@ public:
 	virtual void ExportAssets(const TArray<UObject*>& AssetsToExport, const FString& ExportPath) const override;
 	virtual void ExportAssetsWithDialog(const TArray<UObject*>& AssetsToExport, bool bPromptForIndividualFilenames) override;
 	virtual void ExportAssetsWithDialog(const TArray<FString>& AssetsToExport, bool bPromptForIndividualFilenames) override;
-
 	virtual void CreateUniqueAssetName(const FString& InBasePackageName, const FString& InSuffix, FString& OutPackageName, FString& OutAssetName) override;
 	virtual bool AssetUsesGenericThumbnail( const FAssetData& AssetData ) const override;
 	virtual void DiffAgainstDepot(UObject* InObject, const FString& InPackagePath, const FString& InPackageName) const override;
@@ -110,10 +109,11 @@ public:
 	virtual bool ValidateFlattenedAdvancedCopyDestinations(const TMap<FString, FString>& FlattenedPackagesAndDestinations) const override;
 	virtual void GetAllAdvancedCopySources(FName SelectedPackage, FAdvancedCopyParams& CopyParams, TArray<FName>& OutPackageNamesToCopy, TMap<FName, FName>& DependencyMap, const class UAdvancedCopyCustomization* CopyCustomization) const override;
 	virtual void InitAdvancedCopyFromCopyParams(FAdvancedCopyParams CopyParams) const override;
-
 	virtual void OpenEditorForAssets(const TArray<UObject*>& Assets) override;
 
 	virtual void ConvertVirtualTextures(const TArray<UTexture2D*>& Textures, bool bConvertBackToNonVirtual, const TArray<UMaterial*>* RelatedMaterials = nullptr) const override;
+	virtual bool IsAssetClassSupported(const UClass* AssetClass) const override;
+	virtual TArray<UFactory*> GetNewAssetFactories() const override;
 public:
 	/** Gets the asset tools singleton as a FAssetTools for asset tools module use */
 	static UAssetToolsImpl& Get();
@@ -176,6 +176,8 @@ private:
 	/** The categories that have been allocated already */
 	TMap<FName, FAdvancedAssetCategory> AllocatedCategoryBits;
 	
+	TSet<FName> SupportedAssetTypes;
+	
 	/** The next user category bit to allocate (set to 0 when there are no more bits left) */
 	uint32 NextUserCategoryBit;
 };
diff --git a/Engine/Source/Developer/AssetTools/Public/AssetTypeActions_Base.h b/Engine/Source/Developer/AssetTools/Public/AssetTypeActions_Base.h
--- a/Engine/Source/Developer/AssetTools/Public/AssetTypeActions_Base.h
+++ b/Engine/Source/Developer/AssetTools/Public/AssetTypeActions_Base.h
@@ -167,6 +167,16 @@ PRAGMA_ENABLE_DEPRECATION_WARNINGS
 		return FText::GetEmpty();
 	}
 
+	virtual void SetSupported(bool bInSupported) final
+	{
+		bIsSupported = bInSupported;
+	}
+
+	virtual bool IsSupported() const final
+	{
+		return bIsSupported;
+	}
+
 protected:
 
 	// Here are some convenience functions for common asset type actions logic
@@ -219,4 +229,7 @@ protected:
 
 		return TypedObjects;
 	}
+
+private:
+	bool bIsSupported = true;
 };
diff --git a/Engine/Source/Developer/AssetTools/Public/IAssetTools.h b/Engine/Source/Developer/AssetTools/Public/IAssetTools.h
--- a/Engine/Source/Developer/AssetTools/Public/IAssetTools.h
+++ b/Engine/Source/Developer/AssetTools/Public/IAssetTools.h
@@ -163,10 +163,10 @@ public:
 	virtual void GetAssetTypeActionsList(TArray<TWeakPtr<IAssetTypeActions>>& OutAssetTypeActionsList) const = 0;
 
 	/** Gets the appropriate AssetTypeActions for the supplied class */
-	virtual TWeakPtr<IAssetTypeActions> GetAssetTypeActionsForClass(UClass* Class) const = 0;
+	virtual TWeakPtr<IAssetTypeActions> GetAssetTypeActionsForClass(const UClass* Class) const = 0;
 
 	/** Gets the list of appropriate AssetTypeActions for the supplied class */
-	virtual TArray<TWeakPtr<IAssetTypeActions>> GetAssetTypeActionsListForClass(UClass* Class) const = 0;
+	virtual TArray<TWeakPtr<IAssetTypeActions>> GetAssetTypeActionsListForClass(const UClass* Class) const = 0;
 
 	/**
 	* Allocates a Category bit for a user-defined Category, or EAssetTypeCategories::Misc if all available bits are allocated.
@@ -398,6 +398,12 @@ public:
 	 * @param RelatedMaterials			An optional array of materials to update after the conversion, this is useful during import when not all dependencies and caches are up to date.
 	 */
 	virtual void ConvertVirtualTextures(const TArray<UTexture2D*>& Textures, bool bConvertBackToNonVirtual, const TArray<UMaterial*>* RelatedMaterials = nullptr) const = 0;
+
+	/** Is the given asset class supported? */
+	virtual bool IsAssetClassSupported(const UClass* AssetClass) const = 0;
+
+	/** Find all supported asset factories. */
+	virtual TArray<UFactory*> GetNewAssetFactories() const = 0;
 };
 
 UCLASS(transient)
diff --git a/Engine/Source/Developer/AssetTools/Public/IAssetTypeActions.h b/Engine/Source/Developer/AssetTools/Public/IAssetTypeActions.h
--- a/Engine/Source/Developer/AssetTools/Public/IAssetTypeActions.h
+++ b/Engine/Source/Developer/AssetTools/Public/IAssetTypeActions.h
@@ -45,7 +45,7 @@ public:
 	/** Returns the name of this type */
 	virtual FText GetName() const = 0;
 
-	/** Checks to see if the specified object is handled by this type. */
+	/** Get the supported class of this type. */
 	virtual UClass* GetSupportedClass() const = 0;
 
 	/** Returns the color associated with this type */
@@ -119,4 +119,10 @@ public:
 
 	/** Builds the filter for this class*/
 	virtual void BuildBackendFilter(struct FARFilter& InFilter) = 0;
+
+	/** Sets whether or not this asset type is a supported type for this editor session. */
+	virtual void SetSupported(bool bInSupported) = 0;
+
+	/** Is this asset type supported in the current session? */
+	virtual bool IsSupported() const = 0;
 };
diff --git a/Engine/Source/Editor/ContentBrowser/Private/NewAssetOrClassContextMenu.cpp b/Engine/Source/Editor/ContentBrowser/Private/NewAssetOrClassContextMenu.cpp
--- a/Engine/Source/Editor/ContentBrowser/Private/NewAssetOrClassContextMenu.cpp
+++ b/Engine/Source/Editor/ContentBrowser/Private/NewAssetOrClassContextMenu.cpp
@@ -23,6 +23,7 @@
 #include "Widgets/SToolTip.h"
 #include "IDocumentation.h"
 #include "ClassIconFinder.h"
+#include "AssetToolsModule.h"
 
 #define LOCTEXT_NAMESPACE "ContentBrowser"
 
@@ -59,23 +60,17 @@ void FCategorySubMenuItem::SortSubMenus(FCategorySubMenuItem* SubMenu)
 TArray<FFactoryItem> FindFactoriesInCategory(EAssetTypeCategories::Type AssetTypeCategory)
 {
 	TArray<FFactoryItem> FactoriesInThisCategory;
-	for (TObjectIterator<UClass> It; It; ++It)
-	{
-		UClass* Class = *It;
-		if (Class->IsChildOf(UFactory::StaticClass()) && !Class->HasAnyClassFlags(CLASS_Abstract))
-		{
-			UFactory* Factory = Class->GetDefaultObject<UFactory>();
-			if (Factory->ShouldShowInNewMenu() && ensure(!Factory->GetDisplayName().IsEmpty()))
+
+	const IAssetTools& AssetTools = FModuleManager::LoadModuleChecked<FAssetToolsModule>("AssetTools").Get();
+	TArray<UFactory*> Factories = AssetTools.GetNewAssetFactories();
+	for (UFactory* Factory : Factories)
 			{
 				uint32 FactoryCategories = Factory->GetMenuCategories();
-
 				if (FactoryCategories & AssetTypeCategory)
 				{
 					new(FactoriesInThisCategory)FFactoryItem(Factory, Factory->GetDisplayName());
 				}
 			}
-		}
-	}
 
 	return FactoriesInThisCategory;
 }
diff --git a/Engine/Source/Editor/ContentBrowser/Private/SAssetView.cpp b/Engine/Source/Editor/ContentBrowser/Private/SAssetView.cpp
--- a/Engine/Source/Editor/ContentBrowser/Private/SAssetView.cpp
+++ b/Engine/Source/Editor/ContentBrowser/Private/SAssetView.cpp
@@ -283,6 +283,44 @@ namespace AssetViewUtils
 			return INDEX_NONE;
 		}
 	};
+
+	static FARFilter CreateSupportedFilter()
+	{
+		FAssetToolsModule& AssetToolsModule = FModuleManager::LoadModuleChecked<FAssetToolsModule>("AssetTools");
+		TArray<TWeakPtr<IAssetTypeActions>> AssetTypeActionsList;
+		AssetToolsModule.Get().GetAssetTypeActionsList(AssetTypeActionsList);
+
+		bool bAnyUnsupported = false;
+		FARFilter SupportedFilter;
+
+		for (const TWeakPtr<IAssetTypeActions>& WeakTypeActions : AssetTypeActionsList)
+		{
+			if (const TSharedPtr<IAssetTypeActions> TypeActions = WeakTypeActions.Pin())
+			{
+				if (!TypeActions->IsSupported())
+				{
+					bAnyUnsupported = true;
+				}
+				else
+				{
+					SupportedFilter.bRecursiveClasses = true;
+					const UClass* SupportedClass = TypeActions->GetSupportedClass();
+					if (SupportedClass != nullptr)
+					{
+						const FName ClassName = SupportedClass->GetFName();
+						SupportedFilter.ClassNames.AddUnique(ClassName);
+					}
+				}
+			}
+		}
+
+		if (!bAnyUnsupported)
+		{
+			return FARFilter();
+		}
+
+		return SupportedFilter;
+	}
 } // namespace AssetViewUtils
 
 SAssetView::~SAssetView()
@@ -493,6 +531,8 @@ void SAssetView::Construct( const FArguments& InArgs )
 
 	NumVisibleColumns = 0;
 
+	SupportedFilter = AssetViewUtils::CreateSupportedFilter();
+
 	FEditorWidgetsModule& EditorWidgetsModule = FModuleManager::LoadModuleChecked<FEditorWidgetsModule>("EditorWidgets");
 	TSharedRef<SWidget> AssetDiscoveryIndicator = EditorWidgetsModule.CreateAssetDiscoveryIndicator(EAssetDiscoveryIndicatorScaleMode::Scale_Vertical);
 
@@ -1925,7 +1965,7 @@ void SAssetView::RefreshSourceItems()
 
 	TArray<FAssetData>& Items = OnShouldFilterAsset.IsBound() ? QueriedAssetItems : AssetItems;
 
-	const bool bShowAll = SourcesData.IsEmpty() && BackendFilter.IsEmpty();
+	const bool bShowAll = SourcesData.IsEmpty() && SupportedFilter.IsEmpty() && BackendFilter.IsEmpty();
 
 	bool bShowClasses = false;
 	TArray<FName> ClassPathsToShow;
@@ -1992,6 +2032,8 @@ void SAssetView::RefreshSourceItems()
 		const bool bIsDynamicCollection = SourcesData.IsDynamicCollection();
 		FARFilter Filter = SourcesData.MakeFilter(bRecurse, bUsingFolders);
 
+		Filter.Append(SupportedFilter);
+
 		// Add the backend filters from the filter list
 		Filter.Append(BackendFilter);
 
@@ -2143,7 +2185,7 @@ void SAssetView::RefreshFilteredItems()
 	TMap< FName, TSharedPtr< FAssetViewAsset > > ItemToObjectPath;
 	for (int Index = 0; Index < FilteredAssetItems.Num(); Index++)
 	{
-		if(FilteredAssetItems[Index].IsValid() && FilteredAssetItems[Index]->GetType() != EAssetItemType::Folder)
+		if (FilteredAssetItems[Index].IsValid() && FilteredAssetItems[Index]->GetType() != EAssetItemType::Folder)
 		{
 			TSharedPtr<FAssetViewAsset> Item = StaticCastSharedPtr<FAssetViewAsset>(FilteredAssetItems[Index]);
 
@@ -2168,18 +2210,19 @@ void SAssetView::RefreshFilteredItems()
 	const bool bGatherAssetTypeCount = CurrentViewType == EAssetViewType::Column;
 	TMap<FName, int32> AssetTypeCount;
 
-	if ( bIsFrontendFilterActive && FrontendFilters.IsValid() )
+	if (bIsFrontendFilterActive && FrontendFilters.IsValid())
 	{
 		const bool bRecurse = ShouldFilterRecursively();
 		const bool bUsingFolders = IsShowingFolders();
 		FARFilter CombinedFilter = SourcesData.MakeFilter(bRecurse, bUsingFolders);
+		CombinedFilter.Append(SupportedFilter);
 		CombinedFilter.Append(BackendFilter);
 
 		// Let the frontend filters know the currently used filter in case it is necessary to conditionally filter based on path or class filters
-		for ( int32 FilterIdx = 0; FilterIdx < FrontendFilters->Num(); ++FilterIdx )
+		for (int32 FilterIdx = 0; FilterIdx < FrontendFilters->Num(); ++FilterIdx)
 		{
 			// There are only FFrontendFilters in this collection
-			const TSharedPtr<FFrontendFilter>& Filter = StaticCastSharedPtr<FFrontendFilter>( FrontendFilters->GetFilterAtIndex(FilterIdx) );
+			const TSharedPtr<FFrontendFilter>& Filter = StaticCastSharedPtr<FFrontendFilter>(FrontendFilters->GetFilterAtIndex(FilterIdx));
 			if ( Filter.IsValid() )
 			{
 				Filter->SetCurrentFilter(CombinedFilter);
@@ -2187,67 +2230,14 @@ void SAssetView::RefreshFilteredItems()
 		}
 	}
 
-	if ( bIsFrontendFilterActive && bGatherAssetTypeCount )
-	{
 		// Check the frontend filter for every asset and keep track of how many assets were found of each type
-		for (int32 AssetIdx = 0; AssetIdx < AssetItems.Num(); ++AssetIdx)
-		{
-			const FAssetData& AssetData = AssetItems[AssetIdx];
-			if ( PassesCurrentFrontendFilter(AssetData) )
-			{
-				const TSharedPtr< FAssetViewAsset >* AssetItem = ItemToObjectPath.Find( AssetData.ObjectPath );
-
-				if ( AssetItem != NULL )
-				{
-					FilteredAssetItems.Add(*AssetItem);
-				}
-				else
-				{
-					FilteredAssetItems.Add(MakeShareable(new FAssetViewAsset(AssetData)));
-				}
-
-				int32* TypeCount = AssetTypeCount.Find(AssetData.AssetClass);
-				if ( TypeCount )
-				{
-					(*TypeCount)++;
-				}
-				else
-				{
-					AssetTypeCount.Add(AssetData.AssetClass, 1);
-				}
-			}
-		}
-	}
-	else if ( bIsFrontendFilterActive && !bGatherAssetTypeCount )
+	for (const FAssetData& AssetData : AssetItems)
 	{
-		// Check the frontend filter for every asset and don't worry about asset type counts
-		for (int32 AssetIdx = 0; AssetIdx < AssetItems.Num(); ++AssetIdx)
+		if (bIsFrontendFilterActive == false ||
+			PassesCurrentFrontendFilter(AssetData) )
 		{
-			const FAssetData& AssetData = AssetItems[AssetIdx];
-			if ( PassesCurrentFrontendFilter(AssetData) )
-			{
-				const TSharedPtr< FAssetViewAsset >* AssetItem = ItemToObjectPath.Find( AssetData.ObjectPath );
-
-				if ( AssetItem != NULL )
-				{
-					FilteredAssetItems.Add(*AssetItem);
-				}
-				else
-				{
-					FilteredAssetItems.Add(MakeShareable(new FAssetViewAsset(AssetData)));
-				}
-			}
-		}
-	}
-	else if ( !bIsFrontendFilterActive && bGatherAssetTypeCount )
-	{
-		// Don't need to check the frontend filter for every asset but keep track of how many assets were found of each type
-		for (int32 AssetIdx = 0; AssetIdx < AssetItems.Num(); ++AssetIdx)
-		{
-			const FAssetData& AssetData = AssetItems[AssetIdx];
 			const TSharedPtr< FAssetViewAsset >* AssetItem = ItemToObjectPath.Find( AssetData.ObjectPath );
-
-			if ( AssetItem != NULL )
+			if (AssetItem != nullptr)
 			{
 				FilteredAssetItems.Add(*AssetItem);
 			}
@@ -2256,48 +2246,28 @@ void SAssetView::RefreshFilteredItems()
 				FilteredAssetItems.Add(MakeShareable(new FAssetViewAsset(AssetData)));
 			}
 
-			int32* TypeCount = AssetTypeCount.Find(AssetData.AssetClass);
-			if ( TypeCount )
-			{
-				(*TypeCount)++;
-			}
-			else
-			{
-				AssetTypeCount.Add(AssetData.AssetClass, 1);
-			}
-		}
-	}
-	else if ( !bIsFrontendFilterActive && !bGatherAssetTypeCount )
-	{
-		// Don't check the frontend filter and don't count the number of assets of each type. Just add all assets.
-		for (int32 AssetIdx = 0; AssetIdx < AssetItems.Num(); ++AssetIdx)
-		{
-			const FAssetData& AssetData = AssetItems[AssetIdx];
-			const TSharedPtr< FAssetViewAsset >* AssetItem = ItemToObjectPath.Find( AssetData.ObjectPath );
-
-			if ( AssetItem != NULL )
-			{
-				FilteredAssetItems.Add(*AssetItem);
-			}
-			else
+			if (bGatherAssetTypeCount)
 			{
-				FilteredAssetItems.Add(MakeShareable(new FAssetViewAsset(AssetData)));
+				int32* TypeCount = AssetTypeCount.Find(AssetData.AssetClass);
+				if ( TypeCount )
+				{
+					(*TypeCount)++;
+				}
+				else
+				{
+					AssetTypeCount.Add(AssetData.AssetClass, 1);
+				}
 			}
 		}
 	}
-	else
-	{
-		// The above cases should handle all combinations of bIsFrontendFilterActive and bGatherAssetTypeCount
-		ensure(0);
-	}
 
-	if ( bGatherAssetTypeCount )
+	if (bGatherAssetTypeCount)
 	{
 		int32 HighestCount = 0;
 		FName HighestType;
-		for ( auto TypeIt = AssetTypeCount.CreateConstIterator(); TypeIt; ++TypeIt )
+		for (auto TypeIt = AssetTypeCount.CreateConstIterator(); TypeIt; ++TypeIt)
 		{
-			if ( TypeIt.Value() > HighestCount )
+			if (TypeIt.Value() > HighestCount)
 			{
 				HighestType = TypeIt.Key();
 				HighestCount = TypeIt.Value();
@@ -2672,14 +2642,14 @@ void SAssetView::ProcessRecentlyAddedAssets()
 		bool bNeedsRefresh = false;
 
 		TSet<FName> ExistingObjectPaths;
-		for ( auto AssetIt = AssetItems.CreateConstIterator(); AssetIt; ++AssetIt )
+		for (const FAssetData& AssetItem : AssetItems)
 		{
-			ExistingObjectPaths.Add((*AssetIt).ObjectPath);
+			ExistingObjectPaths.Add(AssetItem.ObjectPath);
 		}
 
-		for ( auto AssetIt = QueriedAssetItems.CreateConstIterator(); AssetIt; ++AssetIt )
+		for (const FAssetData& AssetItem : AssetItems)
 		{
-			ExistingObjectPaths.Add((*AssetIt).ObjectPath);
+			ExistingObjectPaths.Add(AssetItem.ObjectPath);
 		}
 
 		int32 AssetIdx = 0;
@@ -2988,14 +2958,6 @@ bool SAssetView::PassesCurrentFrontendFilter(const FAssetData& Item) const
 	return true;
 }
 
-bool SAssetView::PassesCurrentBackendFilter(const FAssetData& Item) const
-{
-	TArray<FAssetData> AssetDataList;
-	AssetDataList.Add(Item);
-	RunAssetsThroughBackendFilter(AssetDataList);
-	return AssetDataList.Num() == 1;
-}
-
 void SAssetView::RunAssetsThroughBackendFilter(TArray<FAssetData>& InOutAssetDataList) const
 {
 	const bool bRecurse = ShouldFilterRecursively();
@@ -3010,7 +2972,7 @@ void SAssetView::RunAssetsThroughBackendFilter(TArray<FAssetData>& InOutAssetDat
 	}
 	else
 	{
-		// Actually append the backend filter
+		Filter.Append(SupportedFilter);
 		Filter.Append(BackendFilter);
 
 		FAssetRegistryModule& AssetRegistryModule = FModuleManager::LoadModuleChecked<FAssetRegistryModule>(TEXT("AssetRegistry"));
diff --git a/Engine/Source/Editor/ContentBrowser/Private/SAssetView.h b/Engine/Source/Editor/ContentBrowser/Private/SAssetView.h
--- a/Engine/Source/Editor/ContentBrowser/Private/SAssetView.h
+++ b/Engine/Source/Editor/ContentBrowser/Private/SAssetView.h
@@ -435,9 +435,6 @@ private:
 	/** Returns true if the specified asset data item passes all applied frontend (non asset registry) filters */
 	bool PassesCurrentFrontendFilter(const FAssetData& Item) const;
 
-	/** Returns true if the specified asset data item passes all applied backend (asset registry) filters */
-	bool PassesCurrentBackendFilter(const FAssetData& Item) const;
-
 	/** Removes asset data items from the given array that don't pass all applied backend (asset registry) filters */
 	void RunAssetsThroughBackendFilter(TArray<FAssetData>& InOutAssetDataList) const;
 
@@ -828,6 +825,7 @@ private:
 	/** The current base source filter for the view */
 	FSourcesData SourcesData;
 	FARFilter BackendFilter;
+	FARFilter SupportedFilter;
 	TSharedPtr<FAssetFilterCollectionType> FrontendFilters;
 
 	/** If true, the source items will be refreshed next frame. Very slow. */
diff --git a/Engine/Source/Editor/ContentBrowser/Private/SFilterList.cpp b/Engine/Source/Editor/ContentBrowser/Private/SFilterList.cpp
--- a/Engine/Source/Editor/ContentBrowser/Private/SFilterList.cpp
+++ b/Engine/Source/Editor/ContentBrowser/Private/SFilterList.cpp
@@ -643,14 +643,6 @@ TSharedRef<SWidget> SFilterList::ExternalMakeAddFilterMenu(EAssetTypeCategories:
 	return MakeAddFilterMenu(MenuExpansion);
 }
 
-void SFilterList::EnableAllFilters()
-{
-	for (auto FilterIt = Filters.CreateConstIterator(); FilterIt; ++FilterIt)
-	{
-		(*FilterIt)->SetEnabled(true);
-	}
-}
-
 void SFilterList::DisableAllFilters()
 {
 	for (const TSharedRef<SFilter>& Filter : Filters)
@@ -1133,7 +1125,7 @@ void SFilterList::CreateOtherFiltersMenuCategory(FMenuBuilder& MenuBuilder, TSha
 	}
 }
 
-bool IsFilteredByPicker(TArray<UClass*> FilterClassList, UClass* TestClass)
+bool IsFilteredByPicker(const TArray<UClass*>& FilterClassList, UClass* TestClass)
 {
 	if (FilterClassList.Num() == 0)
 	{
diff --git a/Engine/Source/Editor/PropertyEditor/Private/PropertyCustomizationHelpers.cpp b/Engine/Source/Editor/PropertyEditor/Private/PropertyCustomizationHelpers.cpp
--- a/Engine/Source/Editor/PropertyEditor/Private/PropertyCustomizationHelpers.cpp
+++ b/Engine/Source/Editor/PropertyEditor/Private/PropertyCustomizationHelpers.cpp
@@ -35,6 +35,7 @@
 #include "EditorFontGlyphs.h"
 #include "DetailCategoryBuilder.h"
 #include "IDetailGroup.h"
+#include "AssetToolsModule.h"
 
 #define LOCTEXT_NAMESPACE "PropertyCustomizationHelpers"
 
@@ -421,32 +422,29 @@ namespace PropertyCustomizationHelpers
 
 	TArray<UFactory*> GetNewAssetFactoriesForClasses(const TArray<const UClass*>& Classes, const TArray<const UClass*>& DisallowedClasses)
 	{
-		TArray<UFactory*> Factories;
-		for (TObjectIterator<UClass> It; It; ++It)
-		{
-			UClass* Class = *It;
-			if (Class->IsChildOf(UFactory::StaticClass()) && !Class->HasAnyClassFlags(CLASS_Abstract))
-			{
-				UFactory* Factory = Class->GetDefaultObject<UFactory>();
-				if (Factory->ShouldShowInNewMenu() && ensure(!Factory->GetDisplayName().IsEmpty()))
+		const IAssetTools& AssetTools = FModuleManager::LoadModuleChecked<FAssetToolsModule>("AssetTools").Get();
+		TArray<UFactory*> AllFactories = AssetTools.GetNewAssetFactories();
+		TArray<UFactory*> FilteredFactories;
+
+		for (UFactory* Factory : AllFactories)
 				{
 					UClass* SupportedClass = Factory->GetSupportedClass();
+			auto IsChildOfLambda = [SupportedClass](const UClass* InClass) { return SupportedClass->IsChildOf(InClass); };
+
 					if (SupportedClass != nullptr 
-						&& Classes.ContainsByPredicate([=](const UClass* InClass) { return SupportedClass->IsChildOf(InClass); })
-						&& !DisallowedClasses.ContainsByPredicate([=](const UClass* InClass) { return SupportedClass->IsChildOf(InClass); }))
+				&& Classes.ContainsByPredicate(IsChildOfLambda)
+				&& !DisallowedClasses.ContainsByPredicate(IsChildOfLambda))
 					{
-						Factories.Add(Factory);
-					}
-				}
+				FilteredFactories.Add(Factory);
 			}
 		}
 
-		Factories.Sort([](UFactory& A, UFactory& B) -> bool
+		FilteredFactories.Sort([](UFactory& A, UFactory& B) -> bool
 		{
 			return A.GetDisplayName().CompareToCaseIgnored(B.GetDisplayName()) < 0;
 		});
 
-		return Factories;
+		return FilteredFactories;
 	}
 }
 
