Added new asset context option to copy the file path on disk of the selected assets

diff --git a/Engine/Source/Editor/ContentBrowser/Private/AssetContextMenu.cpp b/Engine/Source/Editor/ContentBrowser/Private/AssetContextMenu.cpp
--- a/Engine/Source/Editor/ContentBrowser/Private/AssetContextMenu.cpp
+++ b/Engine/Source/Editor/ContentBrowser/Private/AssetContextMenu.cpp
@@ -276,6 +276,9 @@ void FAssetContextMenu::AddMenuOptions(UToolMenu* InMenu)
 	// Add reference options
 	AddReferenceMenuOptions(InMenu);
 
+	// Add copy full path options
+	AddCopyFilePathMenuOptions(InMenu);
+
 	// Add collection options
 	AddCollectionMenuOptions(InMenu);
 
@@ -1286,6 +1289,22 @@ bool FAssetContextMenu::AddReferenceMenuOptions(UToolMenu* Menu)
 	return true;
 }
 
+bool FAssetContextMenu::AddCopyFilePathMenuOptions(UToolMenu* Menu)
+{
+	FToolMenuSection& Section = Menu->AddSection("AssetContextReferences", LOCTEXT("ReferencesMenuHeading", "References"));
+	{
+		Section.AddMenuEntry(
+			"CopyFilePath",
+			LOCTEXT("CopyFilePath", "Copy File Path"),
+			LOCTEXT("CopyFilePathTooltip", "Copies the file paths on disk for the selected assets to the clipboard."),
+			FSlateIcon(),
+			FUIAction(FExecuteAction::CreateSP(this, &FAssetContextMenu::ExecuteCopyFilePath))
+		);
+	}
+
+	return true;
+}
+
 bool FAssetContextMenu::AddDocumentationMenuOptions(UToolMenu* Menu)
 {
 	bool bAddedOption = false;
@@ -2524,6 +2543,11 @@ void FAssetContextMenu::ExecuteCopyReference()
 	ContentBrowserUtils::CopyAssetReferencesToClipboard(SelectedAssets);
 }
 
+void FAssetContextMenu::ExecuteCopyFilePath()
+{
+	ContentBrowserUtils::CopyFilePathsToClipboard(SelectedAssets);
+}
+
 void FAssetContextMenu::ExecuteCopyTextToClipboard(FString InText)
 {
 	FPlatformApplicationMisc::ClipboardCopy(*InText);
diff --git a/Engine/Source/Editor/ContentBrowser/Private/AssetContextMenu.h b/Engine/Source/Editor/ContentBrowser/Private/AssetContextMenu.h
--- a/Engine/Source/Editor/ContentBrowser/Private/AssetContextMenu.h
+++ b/Engine/Source/Editor/ContentBrowser/Private/AssetContextMenu.h
@@ -166,6 +166,9 @@ private:
 	/** Adds asset reference menu options to a menu builder. Returns true if any options were added. */
 	bool AddReferenceMenuOptions(UToolMenu* Menu);
 
+	/** Adds copy file path menu options to a menu builder. Returns true if any options were added. */
+	bool AddCopyFilePathMenuOptions(UToolMenu* Menu);
+
 	/** Adds asset documentation menu options to a menu builder. Returns true if any options were added. */
 	bool AddDocumentationMenuOptions(UToolMenu* Menu);
 	
@@ -235,6 +238,9 @@ private:
 	/** Handler for CopyReference */
 	void ExecuteCopyReference();
 
+	/** Handler for CopyFilePath */
+	void ExecuteCopyFilePath();
+
 	/** Handler to copy the given text to the clipboard */
 	void ExecuteCopyTextToClipboard(FString InText);
 
diff --git a/Engine/Source/Editor/ContentBrowser/Private/ContentBrowserUtils.cpp b/Engine/Source/Editor/ContentBrowser/Private/ContentBrowserUtils.cpp
--- a/Engine/Source/Editor/ContentBrowser/Private/ContentBrowserUtils.cpp
+++ b/Engine/Source/Editor/ContentBrowser/Private/ContentBrowserUtils.cpp
@@ -1004,6 +1004,33 @@ void ContentBrowserUtils::CopyAssetReferencesToClipboard(const TArray<FAssetData
 	FPlatformApplicationMisc::ClipboardCopy( *ClipboardText );
 }
 
+void ContentBrowserUtils::CopyFilePathsToClipboard(const TArray<FAssetData>& AssetsToCopy)
+{
+	FString ClipboardText;
+	for (const FAssetData& Asset : AssetsToCopy)
+	{
+		if (ClipboardText.Len() > 0)
+		{
+			ClipboardText += LINE_TERMINATOR;
+		}
+		FString PackageFileName;
+		FString PackageFile;
+		if (FPackageName::TryConvertLongPackageNameToFilename(Asset.PackageName.ToString(), PackageFileName) &&
+			FPackageName::FindPackageFileWithoutExtension(PackageFileName, PackageFile))
+		{
+			ClipboardText += FPaths::ConvertRelativePathToFull(PackageFile);
+		}
+		else
+		{
+			// Add a message for when a user tries to copy the path to a file that doesn't exist on disk of the form
+			// <AssetName>: No file on disk
+			ClipboardText += Asset.AssetName.ToString() + FString(": No file on disk");
+		}
+	}
+
+	FPlatformApplicationMisc::ClipboardCopy(*ClipboardText);
+}
+
 void ContentBrowserUtils::CaptureThumbnailFromViewport(FViewport* InViewport, const TArray<FAssetData>& InAssetsToAssign)
 {
 	//capture the thumbnail
diff --git a/Engine/Source/Editor/ContentBrowser/Private/ContentBrowserUtils.h b/Engine/Source/Editor/ContentBrowser/Private/ContentBrowserUtils.h
--- a/Engine/Source/Editor/ContentBrowser/Private/ContentBrowserUtils.h
+++ b/Engine/Source/Editor/ContentBrowser/Private/ContentBrowserUtils.h
@@ -131,6 +131,9 @@ namespace ContentBrowserUtils
 	/** Copies references to the specified assets to the clipboard */
 	void CopyAssetReferencesToClipboard(const TArray<FAssetData>& AssetsToCopy);
 
+	/** Copies file paths on disk to the specified assets to the clipboard */
+	void CopyFilePathsToClipboard(const TArray<FAssetData>& AssetsToCopy);
+
 	/**
 	 * Capture active viewport to thumbnail and assigns that thumbnail to incoming assets
 	 *
