Fixed build failure introduced by CL 10074271 accidentally removing the definition for SFilterList::EnableAllFilters().

diff --git a/Engine/Source/Editor/ContentBrowser/Private/SFilterList.cpp b/Engine/Source/Editor/ContentBrowser/Private/SFilterList.cpp
--- a/Engine/Source/Editor/ContentBrowser/Private/SFilterList.cpp
+++ b/Engine/Source/Editor/ContentBrowser/Private/SFilterList.cpp
@@ -643,6 +643,16 @@ TSharedRef<SWidget> SFilterList::ExternalMakeAddFilterMenu(EAssetTypeCategories:
 	return MakeAddFilterMenu(MenuExpansion);
 }
 
+void SFilterList::EnableAllFilters()
+{
+	for (const TSharedRef<SFilter>& Filter : Filters)
+	{
+		Filter->SetEnabled(true, false);
+	}
+
+	OnFilterChanged.ExecuteIfBound();
+}
+
 void SFilterList::DisableAllFilters()
 {
 	for (const TSharedRef<SFilter>& Filter : Filters)
