Crash overriding merge destination with asset of different type

diff --git a/Engine/Source/Editor/ContentBrowser/Private/ContentBrowserUtils.cpp b/Engine/Source/Editor/ContentBrowser/Private/ContentBrowserUtils.cpp
--- a/Engine/Source/Editor/ContentBrowser/Private/ContentBrowserUtils.cpp
+++ b/Engine/Source/Editor/ContentBrowser/Private/ContentBrowserUtils.cpp
@@ -1727,6 +1727,11 @@ static const auto CVarMaxFullPathLength =
 	IConsoleManager::Get().RegisterConsoleVariable( TEXT("MaxAssetFullPath"), FPlatformMisc::GetMaxPathLength(), TEXT("Maximum full path name of an asset.") )->AsVariableInt();
 
 bool ContentBrowserUtils::IsValidObjectPathForCreate(const FString& ObjectPath, FText& OutErrorMessage, bool bAllowExistingAsset)
+{
+	return IsValidObjectPathForCreate(ObjectPath, nullptr, OutErrorMessage, bAllowExistingAsset);
+}
+
+bool ContentBrowserUtils::IsValidObjectPathForCreate(const FString& ObjectPath, const UClass* ObjectClass, FText& OutErrorMessage, bool bAllowExistingAsset)
 {
 	const FString ObjectName = FPackageName::ObjectPathToObjectName(ObjectPath);
 
@@ -1773,14 +1778,22 @@ bool ContentBrowserUtils::IsValidObjectPathForCreate(const FString& ObjectPath,
 		return false;
 	}
 
-	// Check for an existing asset, unless it we were asked not to.
-	if ( !bAllowExistingAsset )
+	// Check for an existing asset
+	FAssetRegistryModule& AssetRegistryModule = FModuleManager::LoadModuleChecked<FAssetRegistryModule>("AssetRegistry");
+	FAssetData ExistingAsset = AssetRegistryModule.Get().GetAssetByObjectPath(FName(*ObjectPath));
+	if (ExistingAsset.IsValid())
 	{
-		FAssetRegistryModule& AssetRegistryModule = FModuleManager::LoadModuleChecked<FAssetRegistryModule>("AssetRegistry");
-		FAssetData ExistingAsset = AssetRegistryModule.Get().GetAssetByObjectPath(FName(*ObjectPath));
-		if (ExistingAsset.IsValid())
+		// An asset of a different type already exists at this location, inform the user and continue
+		if (ObjectClass && !ExistingAsset.GetClass()->IsChildOf(ObjectClass))
+		{
+			OutErrorMessage = FText::Format(LOCTEXT("RenameAssetOtherTypeAlreadyExists", "An asset of type '{0}' already exists at this location with the name '{1}'."), FText::FromName(ExistingAsset.AssetClass), FText::FromString(ObjectName));
+			
+			// Return false to indicate that the user should enter a new name
+			return false;
+		}
+        // This asset already exists at this location, warn user if asked to.
+		else if (!bAllowExistingAsset)
 		{
-			// This asset already exists at this location, inform the user and continue
 			OutErrorMessage = FText::Format( LOCTEXT("RenameAssetAlreadyExists", "An asset already exists at this location with the name '{0}'."), FText::FromString( ObjectName ) );
 
 			// Return false to indicate that the user should enter a new name
diff --git a/Engine/Source/Editor/ContentBrowser/Private/ContentBrowserUtils.h b/Engine/Source/Editor/ContentBrowser/Private/ContentBrowserUtils.h
--- a/Engine/Source/Editor/ContentBrowser/Private/ContentBrowserUtils.h
+++ b/Engine/Source/Editor/ContentBrowser/Private/ContentBrowserUtils.h
@@ -256,6 +256,9 @@ namespace ContentBrowserUtils
 
 	/** Returns true if the specified path is available for object creation */
 	bool IsValidObjectPathForCreate(const FString& ObjectPath, FText& OutErrorMessage, bool bAllowExistingAsset = false);
+	
+	/** Returns true if the specified path is available for object creation */
+	bool IsValidObjectPathForCreate(const FString& ObjectPath, const UClass* ObjectClass, FText& OutErrorMessage, bool bAllowExistingAsset = false);
 
 	/** Returns true if the specified folder name in the specified path is available for folder creation */
 	bool IsValidFolderPathForCreate(const FString& FolderPath, const FString& NewFolderName, FText& OutErrorMessage);
diff --git a/Engine/Source/Editor/ContentBrowser/Private/SAssetDialog.cpp b/Engine/Source/Editor/ContentBrowser/Private/SAssetDialog.cpp
--- a/Engine/Source/Editor/ContentBrowser/Private/SAssetDialog.cpp
+++ b/Engine/Source/Editor/ContentBrowser/Private/SAssetDialog.cpp
@@ -857,11 +857,26 @@ void SAssetDialog::UpdateInputValidity()
 			const FString ObjectPath = GetObjectPathForSave();
 			FText ErrorMessage;
 			const bool bAllowExistingAsset = (ExistingAssetPolicy == ESaveAssetDialogExistingAssetPolicy::AllowButWarn);
-			if ( !ContentBrowserUtils::IsValidObjectPathForCreate(ObjectPath, ErrorMessage, bAllowExistingAsset) )
+
+			FName AssetClassName = AssetClassNames.Num() == 1 ? AssetClassNames[0] : NAME_None;
+			UClass* AssetClass = AssetClassName != NAME_None ? FindObject<UClass>(ANY_PACKAGE, *AssetClassName.ToString(), true) : nullptr;
+
+			if ( !ContentBrowserUtils::IsValidObjectPathForCreate(ObjectPath, AssetClass, ErrorMessage, bAllowExistingAsset) )
 			{
 				LastInputValidityErrorText = ErrorMessage;
 				bLastInputValidityCheckSuccessful = false;
 			}
+			else if(bAllowExistingAsset && AssetClassNames.Num() > 1) // If for some reason we have multiple names, perform additional logic here...
+			{
+				FAssetRegistryModule& AssetRegistryModule = FModuleManager::LoadModuleChecked<FAssetRegistryModule>("AssetRegistry");
+				FAssetData ExistingAsset = AssetRegistryModule.Get().GetAssetByObjectPath(FName(*ObjectPath));
+				if (ExistingAsset.IsValid() && !AssetClassNames.Contains(ExistingAsset.AssetClass))
+				{
+					const FString ObjectName = FPackageName::ObjectPathToObjectName(ObjectPath);
+					LastInputValidityErrorText = FText::Format(LOCTEXT("AssetDialog_AssetAlreadyExists", "An asset of type '{0}' already exists at this location with the name '{1}'."), FText::FromName(ExistingAsset.AssetClass), FText::FromString(ObjectName));
+					bLastInputValidityCheckSuccessful = false;
+				}
+			}
 		}
 	}
 }
