SAssetView can now only show allowed classes in cases of asset dropdowns in the Details view.

This is accomplished by filtering out the recursive classes list to only the intersection between the supported asset classes and the backend filter's classes list if one is set. FARFilter::Append() creates the union of these lists, which would cause an excessive amount of classes to be displayed.

diff --git a/Engine/Source/Editor/ContentBrowser/Private/SAssetView.cpp b/Engine/Source/Editor/ContentBrowser/Private/SAssetView.cpp
--- a/Engine/Source/Editor/ContentBrowser/Private/SAssetView.cpp
+++ b/Engine/Source/Editor/ContentBrowser/Private/SAssetView.cpp
@@ -761,6 +761,24 @@ void SAssetView::SetBackendFilter(const FARFilter& InBackendFilter)
 	RequestSlowFullListRefresh();
 }
 
+void SAssetView::AppendBackendFilter(FARFilter& FilterToAppendTo) const
+{
+	FilterToAppendTo.Append(SupportedFilter);
+	FilterToAppendTo.Append(BackendFilter);
+
+	if (SupportedFilter.bRecursiveClasses && BackendFilter.bRecursiveClasses)
+	{
+		FilterToAppendTo.ClassNames.RemoveAll(
+			[this](const FName& Class)
+			{
+				return !SupportedFilter.ClassNames.Contains(Class) ||
+					!BackendFilter.ClassNames.Contains(Class);
+			});
+
+		FilterToAppendTo.bRecursiveClasses = FilterToAppendTo.ClassNames.Num() > 0;
+	}
+}
+
 void SAssetView::OnCreateNewFolder(const FString& FolderName, const FString& FolderPath)
 {
 	// we should only be creating one deferred folder per tick
@@ -783,7 +801,7 @@ void SAssetView::OnCreateNewFolder(const FString& FolderName, const FString& Fol
 
 void SAssetView::DeferredCreateNewFolder()
 {
-	if(DeferredFolderToCreate.IsValid())
+	if (DeferredFolderToCreate.IsValid())
 	{
 		TSharedPtr<FAssetViewFolder> NewItem = MakeShareable(new FAssetViewFolder(DeferredFolderToCreate->FolderPath / DeferredFolderToCreate->FolderName));
 		NewItem->bNewFolder = true;
@@ -1248,7 +1266,7 @@ void SAssetView::Tick( const FGeometry& AllottedGeometry, const double InCurrent
 		AssetThumbnailPool->Tick(InDeltaTime);
 	}
 
-	if ( bPendingUpdateThumbnails )
+	if (bPendingUpdateThumbnails)
 	{
 		UpdateThumbnails();
 		bPendingUpdateThumbnails = false;
@@ -1482,19 +1500,19 @@ void SAssetView::CalculateThumbnailHintColorAndOpacity()
 	ThumbnailHintColorAndOpacity = FLinearColor( 1.0, 1.0, 1.0, Opacity );
 }
 
-void SAssetView::ProcessQueriedItems( const double TickStartTime )
+void SAssetView::ProcessQueriedItems(const double TickStartTime)
 {
 	const bool bFlushFullBuffer = TickStartTime < 0;
 
 	bool ListNeedsRefresh = false;
 	int32 AssetIndex = 0;
-	for ( AssetIndex = QueriedAssetItems.Num() - 1; AssetIndex >= 0 ; AssetIndex--)
+	for (AssetIndex = QueriedAssetItems.Num() - 1; AssetIndex >= 0; AssetIndex--)
 	{
-		if ( !OnShouldFilterAsset.Execute( QueriedAssetItems[AssetIndex] ) )
+		if (!OnShouldFilterAsset.Execute(QueriedAssetItems[AssetIndex]))
 		{
-			AssetItems.Add( QueriedAssetItems[AssetIndex] );
+			AssetItems.Add(QueriedAssetItems[AssetIndex]);
 
-			if ( !IsFrontendFilterActive() || PassesCurrentFrontendFilter(QueriedAssetItems[AssetIndex]))
+			if (!IsFrontendFilterActive() || PassesCurrentFrontendFilter(QueriedAssetItems[AssetIndex]))
 			{
 				const FAssetData& AssetData = QueriedAssetItems[AssetIndex];
 				FilteredAssetItems.Add(MakeShareable(new FAssetViewAsset(AssetData)));
@@ -1504,7 +1522,7 @@ void SAssetView::ProcessQueriedItems( const double TickStartTime )
 		}
 
 		// Check to see if we have run out of time in this tick
-		if ( !bFlushFullBuffer && (FPlatformTime::Seconds() - TickStartTime) > MaxSecondsPerFrame)
+		if (!bFlushFullBuffer && (FPlatformTime::Seconds() - TickStartTime) > MaxSecondsPerFrame)
 		{
 			break;
 		}
@@ -1513,14 +1531,14 @@ void SAssetView::ProcessQueriedItems( const double TickStartTime )
 	// Trim the results array
 	if (AssetIndex > 0)
 	{
-		QueriedAssetItems.RemoveAt( AssetIndex, QueriedAssetItems.Num() - AssetIndex );
+		QueriedAssetItems.RemoveAt(AssetIndex, QueriedAssetItems.Num() - AssetIndex);
 	}
 	else
 	{
 		QueriedAssetItems.Reset();
 	}
 
-	if ( ListNeedsRefresh )
+	if (ListNeedsRefresh)
 	{
 		RefreshList();
 	}
@@ -2032,11 +2050,8 @@ void SAssetView::RefreshSourceItems()
 		const bool bIsDynamicCollection = SourcesData.IsDynamicCollection();
 		FARFilter Filter = SourcesData.MakeFilter(bRecurse, bUsingFolders);
 
-		Filter.Append(SupportedFilter);
-
-		// Add the backend filters from the filter list
-		Filter.Append(BackendFilter);
-
+		AppendBackendFilter(Filter);
+		
 		bWereItemsRecursivelyFiltered = bRecurse;
 
 		// Move any class paths into their own array
@@ -2181,19 +2196,19 @@ bool SAssetView::ShouldFilterRecursively() const
 
 void SAssetView::RefreshFilteredItems()
 {
-	//Build up a map of the existing AssetItems so we can preserve them while filtering
+	// Build up a map of the existing AssetItems so we can preserve them while filtering
 	TMap< FName, TSharedPtr< FAssetViewAsset > > ItemToObjectPath;
-	for (int Index = 0; Index < FilteredAssetItems.Num(); Index++)
+	for (const TSharedPtr<FAssetViewItem>& AssetItem : FilteredAssetItems)
 	{
-		if (FilteredAssetItems[Index].IsValid() && FilteredAssetItems[Index]->GetType() != EAssetItemType::Folder)
+		if (AssetItem.IsValid() && AssetItem->GetType() != EAssetItemType::Folder)
 		{
-			TSharedPtr<FAssetViewAsset> Item = StaticCastSharedPtr<FAssetViewAsset>(FilteredAssetItems[Index]);
+			TSharedPtr<FAssetViewAsset> Item = StaticCastSharedPtr<FAssetViewAsset>(AssetItem);
 
 			// Clear custom column data
 			Item->CustomColumnData.Reset();
 			Item->CustomColumnDisplayText.Reset();
 
-			ItemToObjectPath.Add( Item->Data.ObjectPath, Item );
+			ItemToObjectPath.Add(Item->Data.ObjectPath, Item);
 		}
 	}
 
@@ -2215,22 +2230,21 @@ void SAssetView::RefreshFilteredItems()
 		const bool bRecurse = ShouldFilterRecursively();
 		const bool bUsingFolders = IsShowingFolders();
 		FARFilter CombinedFilter = SourcesData.MakeFilter(bRecurse, bUsingFolders);
-		CombinedFilter.Append(SupportedFilter);
-		CombinedFilter.Append(BackendFilter);
+		AppendBackendFilter(CombinedFilter);
 
 		// Let the frontend filters know the currently used filter in case it is necessary to conditionally filter based on path or class filters
 		for (int32 FilterIdx = 0; FilterIdx < FrontendFilters->Num(); ++FilterIdx)
 		{
 			// There are only FFrontendFilters in this collection
 			const TSharedPtr<FFrontendFilter>& Filter = StaticCastSharedPtr<FFrontendFilter>(FrontendFilters->GetFilterAtIndex(FilterIdx));
-			if ( Filter.IsValid() )
+			if (Filter.IsValid())
 			{
 				Filter->SetCurrentFilter(CombinedFilter);
 			}
 		}
 	}
 
-		// Check the frontend filter for every asset and keep track of how many assets were found of each type
+	// Check the frontend filter for every asset and keep track of how many assets were found of each type
 	for (const FAssetData& AssetData : AssetItems)
 	{
 		if (bIsFrontendFilterActive == false ||
@@ -2249,7 +2263,7 @@ void SAssetView::RefreshFilteredItems()
 			if (bGatherAssetTypeCount)
 			{
 				int32* TypeCount = AssetTypeCount.Find(AssetData.AssetClass);
-				if ( TypeCount )
+				if (TypeCount)
 				{
 					(*TypeCount)++;
 				}
@@ -2357,7 +2371,7 @@ void SAssetView::RefreshFolders()
 	}
 
 	// Add folders for any child collections of the currently selected collections
-	if(SourcesData.HasCollections())
+	if (SourcesData.HasCollections())
 	{
 		FCollectionManagerModule& CollectionManagerModule = FCollectionManagerModule::GetModule();
 		
@@ -2367,7 +2381,7 @@ void SAssetView::RefreshFolders()
 			ChildCollections.Reset();
 			CollectionManagerModule.Get().GetChildCollections(Collection.Name, Collection.Type, ChildCollections);
 
-			for(const FCollectionNameType& ChildCollection : ChildCollections)
+			for (const FCollectionNameType& ChildCollection : ChildCollections)
 			{
 				// Use "Collections" as the root of the path to avoid this being confused with other asset view folders - see ContentBrowserUtils::IsCollectionPath
 				FoldersToAdd.Add(FString::Printf(TEXT("/Collections/%s/%s"), ECollectionShareType::ToString(ChildCollection.Type), *ChildCollection.Name.ToString()));
@@ -2375,9 +2389,9 @@ void SAssetView::RefreshFolders()
 		}
 	}
 
-	if(FoldersToAdd.Num() > 0)
+	if (FoldersToAdd.Num() > 0)
 	{
-		for(const FString& FolderPath : FoldersToAdd)
+		for (const FString& FolderPath : FoldersToAdd)
 		{
 			FilteredAssetItems.Add(MakeShareable(new FAssetViewFolder(FolderPath)));
 			Folders.Add(FolderPath);
@@ -2762,13 +2776,13 @@ void SAssetView::OnAssetRegistryPathAdded(const FString& Path)
 void SAssetView::OnAssetRegistryPathRemoved(const FString& Path)
 {
 	FString* Folder = Folders.Find(Path);
-	if(Folder != NULL)
+	if (Folder != nullptr)
 	{
 		Folders.Remove(Path);
 
 		for (int32 AssetIdx = 0; AssetIdx < FilteredAssetItems.Num(); ++AssetIdx)
 		{
-			if(FilteredAssetItems[AssetIdx]->GetType() == EAssetItemType::Folder)
+			if (FilteredAssetItems[AssetIdx]->GetType() == EAssetItemType::Folder)
 			{
 				if ( StaticCastSharedPtr<FAssetViewFolder>(FilteredAssetItems[AssetIdx])->FolderPath == Path )
 				{
@@ -2964,21 +2978,20 @@ void SAssetView::RunAssetsThroughBackendFilter(TArray<FAssetData>& InOutAssetDat
 	const bool bUsingFolders = IsShowingFolders();
 	const bool bIsDynamicCollection = SourcesData.IsDynamicCollection();
 	FARFilter Filter = SourcesData.MakeFilter(bRecurse, bUsingFolders);
-	
-	if ( SourcesData.HasCollections() && Filter.ObjectPaths.Num() == 0 && !bIsDynamicCollection )
+
+	if (SourcesData.HasCollections() && Filter.ObjectPaths.Num() == 0 && !bIsDynamicCollection)
 	{
 		// This is an empty collection, no asset will pass the check
 		InOutAssetDataList.Reset();
 	}
 	else
 	{
-		Filter.Append(SupportedFilter);
-		Filter.Append(BackendFilter);
+		AppendBackendFilter(Filter);
 
 		FAssetRegistryModule& AssetRegistryModule = FModuleManager::LoadModuleChecked<FAssetRegistryModule>(TEXT("AssetRegistry"));
 		AssetRegistryModule.Get().RunAssetsThroughFilter(InOutAssetDataList, Filter);
 
-		if ( SourcesData.HasCollections() && !bIsDynamicCollection )
+		if (SourcesData.HasCollections() && !bIsDynamicCollection)
 		{
 			// Include objects from child collections if we're recursing
 			const ECollectionRecursionFlags::Flags CollectionRecursionMode = (Filter.bRecursivePaths) ? ECollectionRecursionFlags::SelfAndChildren : ECollectionRecursionFlags::Self;
@@ -2990,11 +3003,11 @@ void SAssetView::RunAssetsThroughBackendFilter(TArray<FAssetData>& InOutAssetDat
 				CollectionManagerModule.Get().GetObjectsInCollection(Collection.Name, Collection.Type, CollectionObjectPaths, CollectionRecursionMode);
 			}
 
-			for ( int32 AssetDataIdx = InOutAssetDataList.Num() - 1; AssetDataIdx >= 0; --AssetDataIdx )
+			for (int32 AssetDataIdx = InOutAssetDataList.Num() - 1; AssetDataIdx >= 0; --AssetDataIdx)
 			{
 				const FAssetData& AssetData = InOutAssetDataList[AssetDataIdx];
 
-				if ( !CollectionObjectPaths.Contains( AssetData.ObjectPath ) )
+				if (!CollectionObjectPaths.Contains(AssetData.ObjectPath))
 				{
 					InOutAssetDataList.RemoveAtSwap(AssetDataIdx);
 				}
diff --git a/Engine/Source/Editor/ContentBrowser/Private/SAssetView.h b/Engine/Source/Editor/ContentBrowser/Private/SAssetView.h
--- a/Engine/Source/Editor/ContentBrowser/Private/SAssetView.h
+++ b/Engine/Source/Editor/ContentBrowser/Private/SAssetView.h
@@ -787,9 +787,13 @@ private:
 
 	/** Creates the row header context menu allowing for hiding individually clicked columns*/
 	TSharedRef<SWidget> CreateRowHeaderMenuContent(const FString ColumnName);
+
 	/** Will compute the max row size from all its children for the specified column id*/
 	FVector2D GetMaxRowSizeForColumn(const FName& ColumnId);
 
+	/** Append the current effective backend filter (intersection of BackendFilter and SupportedFilter) to the given filter. */
+	void AppendBackendFilter(FARFilter& FilterToAppendTo) const;
+
 public:
 	/** Delegate that handles if any folder paths changed as a result of a move, rename, etc. in the asset view*/
 	FOnFolderPathChanged OnFolderPathChanged;
