Fix asset context menus not working after re-opening content browser
#rnx

diff --git a/Engine/Source/Editor/ContentBrowser/Private/AssetContextMenu.cpp b/Engine/Source/Editor/ContentBrowser/Private/AssetContextMenu.cpp
--- a/Engine/Source/Editor/ContentBrowser/Private/AssetContextMenu.cpp
+++ b/Engine/Source/Editor/ContentBrowser/Private/AssetContextMenu.cpp
@@ -249,7 +249,14 @@ void FAssetContextMenu::RegisterContextMenu(const FName MenuName)
 			}
 		}));
 
-		Menu->AddDynamicSection("AddMenuOptions", FNewToolMenuDelegate::CreateSP(this, &FAssetContextMenu::AddMenuOptions));
+		Menu->AddDynamicSection("AddMenuOptions", FNewToolMenuDelegate::CreateLambda([](UToolMenu* InMenu)
+		{
+			UContentBrowserAssetContextMenuContext* Context = InMenu->FindContext<UContentBrowserAssetContextMenuContext>();
+			if (Context && Context->AssetContextMenu.IsValid())
+			{
+				Context->AssetContextMenu.Pin()->AddMenuOptions(InMenu);
+			}
+		}));
 	}
 }
 
