Content browser folders to display can now be filtered using delegates registered by plugins

diff --git a/Engine/Source/Editor/ContentBrowser/Private/ContentBrowserSingleton.cpp b/Engine/Source/Editor/ContentBrowser/Private/ContentBrowserSingleton.cpp
--- a/Engine/Source/Editor/ContentBrowser/Private/ContentBrowserSingleton.cpp
+++ b/Engine/Source/Editor/ContentBrowser/Private/ContentBrowserSingleton.cpp
@@ -721,6 +721,20 @@ void FContentBrowserSingleton::ForceShowPluginContent(bool bEnginePlugin)
 	}
 }
 
+bool FContentBrowserSingleton::PathViewPathPassesFilter(const FString& InPath) const
+{
+	static FContentBrowserModule& ContentBrowserModule = FModuleManager::LoadModuleChecked<FContentBrowserModule>("ContentBrowser");
+	for (const auto& It : ContentBrowserModule.GetDirectoryPathFilteredDelegates())
+	{
+		if (It.IsBound() && It.Execute(InPath))
+		{
+			return false;
+		}
+	}
+
+	return true;
+}
+
 void FContentBrowserSingleton::PopulateConfigValues()
 {
 	const FString ContentBrowserSection = TEXT("ContentBrowser");
diff --git a/Engine/Source/Editor/ContentBrowser/Private/ContentBrowserSingleton.h b/Engine/Source/Editor/ContentBrowser/Private/ContentBrowserSingleton.h
--- a/Engine/Source/Editor/ContentBrowser/Private/ContentBrowserSingleton.h
+++ b/Engine/Source/Editor/ContentBrowser/Private/ContentBrowserSingleton.h
@@ -86,6 +86,9 @@ public:
 	/** Gets the settings for the plugin with the specified name */
 	const FContentBrowserPluginSettings& GetPluginSettings(FName PluginName) const;
 
+	/** Returns true if should show directory path */
+	bool PathViewPathPassesFilter(const FString& InPath) const;
+
 	/** Single storage location for content browser favorites */
 	TArray<FString> FavoriteFolderPaths;
 
diff --git a/Engine/Source/Editor/ContentBrowser/Private/SPathView.cpp b/Engine/Source/Editor/ContentBrowser/Private/SPathView.cpp
--- a/Engine/Source/Editor/ContentBrowser/Private/SPathView.cpp
+++ b/Engine/Source/Editor/ContentBrowser/Private/SPathView.cpp
@@ -24,6 +24,7 @@
 #include "Widgets/Input/SSearchBox.h"
 #include "NativeClassHierarchy.h"
 #include "EmptyFolderVisibilityManager.h"
+#include "ContentBrowserModule.h"
 
 #include "Application/SlateApplicationBase.h"
 
@@ -74,6 +75,8 @@ void SPathView::Construct( const FArguments& InArgs )
 		RegisterActiveTimer( 0.f, FWidgetActiveTimerDelegate::CreateSP( this, &SPathView::SetFocusPostConstruct ) );
 	}
 
+	ContentBrowserSingleton = &FContentBrowserSingleton::Get();
+
 	// Listen for when view settings are changed
 	UContentBrowserSettings::OnSettingChanged().AddSP(this, &SPathView::HandleSettingChanged);
 
@@ -335,6 +338,11 @@ TSharedPtr<FTreeItem> SPathView::AddPath(const FString& Path, bool bUserNamed)
 		return TSharedPtr<FTreeItem>();
 	}
 
+	if (!ContentBrowserSingleton->PathViewPathPassesFilter(Path))
+	{
+		return TSharedPtr<FTreeItem>();
+	}
+
 	TArray<FString> PathItemList;
 	Path.ParseIntoArray(PathItemList, TEXT("/"), /*InCullEmpty=*/true);
 
@@ -1855,6 +1863,11 @@ TSharedPtr<FTreeItem> SFavoritePathView::AddPath(const FString& Path, bool bUser
 		return TSharedPtr<FTreeItem>();
 	}
 
+	if (!ContentBrowserSingleton->PathViewPathPassesFilter(Path))
+	{
+		return TSharedPtr<FTreeItem>();
+	}
+
 	TArray<FString> PathItemList;
 	Path.ParseIntoArray(PathItemList, TEXT("/"), /*InCullEmpty=*/true);
 
diff --git a/Engine/Source/Editor/ContentBrowser/Private/SPathView.h b/Engine/Source/Editor/ContentBrowser/Private/SPathView.h
--- a/Engine/Source/Editor/ContentBrowser/Private/SPathView.h
+++ b/Engine/Source/Editor/ContentBrowser/Private/SPathView.h
@@ -19,6 +19,7 @@
 class FSourcesSearch;
 struct FHistoryData;
 struct FTreeItem;
+class FContentBrowserSingleton;
 
 typedef TTextFilter< const FString& > FolderTextFilter;
 
@@ -326,6 +327,9 @@ protected:
 
 	TSharedPtr<SWidget> PathViewWidget;
 
+	/** Cached pointer to content browser singleton */
+	FContentBrowserSingleton* ContentBrowserSingleton;
+
 private:
 
 	/** The paths that were last reported by OnPathExpanded event. Used in preserving expansion when filtering folders */
diff --git a/Engine/Source/Editor/ContentBrowser/Public/ContentBrowserModule.h b/Engine/Source/Editor/ContentBrowser/Public/ContentBrowserModule.h
--- a/Engine/Source/Editor/ContentBrowser/Public/ContentBrowserModule.h
+++ b/Engine/Source/Editor/ContentBrowser/Public/ContentBrowserModule.h
@@ -53,6 +53,8 @@ public:
 	DECLARE_MULTICAST_DELEGATE_OneParam( FOnSourcesViewChanged, bool /*bExpanded*/ );
 	/** */
 	DECLARE_MULTICAST_DELEGATE_OneParam( FOnAssetPathChanged, const FString& /*NewPath*/ );
+	/** */
+	DECLARE_DELEGATE_RetVal_OneParam(bool, FIsDirectoryPathFiltered, const FString&);
 
 	/**
 	 * Called right after the plugin DLL has been loaded and the plugin object has been created
@@ -97,6 +99,9 @@ public:
 	/** Delegates to be called to extend the drag-and-drop support of the asset view */
 	virtual TArray<FAssetViewDragAndDropExtender>& GetAssetViewDragAndDropExtenders() { return AssetViewDragAndDropExtenders; }
 
+	/** Delegates to call to hide folders displayed in content browser */
+	virtual TArray<FIsDirectoryPathFiltered>& GetDirectoryPathFilteredDelegates() { return DirectoryPathFilteredDelegates; }
+
 	/** Delegate accessors */
 	FOnFilterChanged& GetOnFilterChanged() { return OnFilterChanged; } 
 	FOnSearchBoxChanged& GetOnSearchBoxChanged() { return OnSearchBoxChanged; } 
@@ -128,6 +133,9 @@ private:
 	TArray<FContentBrowserMenuExtender> AssetViewViewMenuExtenders;
 	TArray<FContentBrowserCommandExtender> ContentBrowserCommandExtenders;
 
+	/** All delegates that hide folders displayed in content browser */
+	TArray<FIsDirectoryPathFiltered> DirectoryPathFilteredDelegates;
+
 	/** All delegates generating extra state indicators */
 	TArray<FAssetViewExtraStateGenerator> AssetViewExtraStateGenerators;
 
