Added a "docked" mode for the collections view

When docked, the collections view appears in a splitter below the paths view and is disconnected from the omni-search

diff --git a/Engine/Source/Editor/ContentBrowser/Private/SAssetView.cpp b/Engine/Source/Editor/ContentBrowser/Private/SAssetView.cpp
--- a/Engine/Source/Editor/ContentBrowser/Private/SAssetView.cpp
+++ b/Engine/Source/Editor/ContentBrowser/Private/SAssetView.cpp
@@ -458,6 +458,7 @@ void SAssetView::Construct( const FArguments& InArgs )
 	bCanShowDevelopersFolder = InArgs._CanShowDevelopersFolder;
 
 	bCanShowFavorites = InArgs._CanShowFavorites;
+	bCanDockCollections = InArgs._CanDockCollections;
 	bPreloadAssetsForContextMenu = InArgs._PreloadAssetsForContextMenu;
 
 	SelectionMode = InArgs._SelectionMode;
@@ -3161,6 +3162,19 @@ TSharedRef<SWidget> SAssetView::GetViewButtonContent()
 			EUserInterfaceActionType::ToggleButton
 		);
 
+		MenuBuilder.AddMenuEntry(
+			LOCTEXT("DockCollectionsOptions", "Dock Collections"),
+			LOCTEXT("DockCollectionsOptionToolTip", "Dock the collections view under the path view?"),
+			FSlateIcon(),
+			FUIAction(
+				FExecuteAction::CreateSP(this, &SAssetView::ToggleDockCollections),
+				FCanExecuteAction::CreateSP(this, &SAssetView::IsToggleDockCollectionsAllowed),
+				FIsActionChecked::CreateSP(this, &SAssetView::HasDockedCollections)
+			),
+			NAME_None,
+			EUserInterfaceActionType::ToggleButton
+		);
+
 		MenuBuilder.AddMenuEntry(
 			LOCTEXT("FilterRecursivelyOption", "Filter Recursively"),
 			LOCTEXT("FilterRecursivelyOptionToolTip", "Should filters apply recursively in the view?"),
@@ -3523,6 +3537,23 @@ bool SAssetView::IsShowingFavorites() const
 	return IsToggleShowFavoritesAllowed() && GetDefault<UContentBrowserSettings>()->GetDisplayFavorites();
 }
 
+void SAssetView::ToggleDockCollections()
+{
+	const bool bDockCollections = GetDefault<UContentBrowserSettings>()->GetDockCollections();
+	GetMutableDefault<UContentBrowserSettings>()->SetDockCollections(!bDockCollections);
+	GetMutableDefault<UContentBrowserSettings>()->PostEditChange();
+}
+
+bool SAssetView::IsToggleDockCollectionsAllowed() const
+{
+	return bCanDockCollections;
+}
+
+bool SAssetView::HasDockedCollections() const
+{
+	return IsToggleDockCollectionsAllowed() && GetDefault<UContentBrowserSettings>()->GetDockCollections();
+}
+
 void SAssetView::ToggleShowCppContent()
 {
 	const bool bDisplayCppFolders = GetDefault<UContentBrowserSettings>()->GetDisplayCppFolders();
diff --git a/Engine/Source/Editor/ContentBrowser/Private/SAssetView.h b/Engine/Source/Editor/ContentBrowser/Private/SAssetView.h
--- a/Engine/Source/Editor/ContentBrowser/Private/SAssetView.h
+++ b/Engine/Source/Editor/ContentBrowser/Private/SAssetView.h
@@ -59,6 +59,7 @@ public:
 		, _CanShowRealTimeThumbnails(false)
 		, _CanShowDevelopersFolder(false)
 		, _CanShowFavorites(false)
+		, _CanDockCollections(false)
 		, _PreloadAssetsForContextMenu(true)
 		, _SelectionMode( ESelectionMode::Multi )
 		, _AllowDragging(true)
@@ -165,6 +166,9 @@ public:
 		/** Indicates if the 'Show Favorites' option should be enabled or disabled */
 		SLATE_ARGUMENT(bool, CanShowFavorites)
 
+		/** Indicates if the 'Dock Collections' option should be enabled or disabled */
+		SLATE_ARGUMENT(bool, CanDockCollections)
+
 		/** Indicates if the context menu is going to load the assets, and if so to preload before the context menu is shown, and warn about the pending load. */
 		SLATE_ARGUMENT( bool, PreloadAssetsForContextMenu )
 
@@ -519,6 +523,15 @@ private:
 	/** @return true when we are showing favorites */
 	bool IsShowingFavorites() const;
 
+	/** Toggle whether the collections view should be docked under the paths view */
+	void ToggleDockCollections();
+
+	/** Whether or not it's possible to dock the collections view */
+	bool IsToggleDockCollectionsAllowed() const;
+
+	/** @return true when the collections view is docked */
+	bool HasDockedCollections() const;
+
 	/** Toggle whether C++ content should be shown or not */
 	void ToggleShowCppContent();
 
@@ -980,6 +993,9 @@ private:
 	/** Indicates if the 'Show Favorites' option should be enabled or disabled */
 	bool bCanShowFavorites;
 
+	/** Indicates if the 'Dock Collections' option should be enabled or disabled */
+	bool bCanDockCollections;
+
 	/** Indicates if the context menu is going to load the assets, and if so to preload before the context menu is shown, and warn about the pending load. */
 	bool bPreloadAssetsForContextMenu;
 
diff --git a/Engine/Source/Editor/ContentBrowser/Private/SCollectionView.cpp b/Engine/Source/Editor/ContentBrowser/Private/SCollectionView.cpp
--- a/Engine/Source/Editor/ContentBrowser/Private/SCollectionView.cpp
+++ b/Engine/Source/Editor/ContentBrowser/Private/SCollectionView.cpp
@@ -83,6 +83,7 @@
 	bAllowCollectionButtons = InArgs._AllowCollectionButtons;
 	bAllowRightClickMenu = InArgs._AllowRightClickMenu;
 	bAllowCollectionDrag = InArgs._AllowCollectionDrag;
+	bAllowExternalSearch = false;
 	bDraggedOver = false;
 
 	bQueueCollectionItemsUpdate = false;
@@ -121,46 +122,9 @@
 		CollectionListContextMenuOpening = FOnContextMenuOpening::CreateSP( this, &SCollectionView::MakeCollectionTreeContextMenu );
 	}
 
-	SearchPtr = InArgs._ExternalSearch;
-	if (!SearchPtr)
-	{
-		SearchPtr = MakeShared<FSourcesSearch>();
-		SearchPtr->Initialize();
-		SearchPtr->SetHintText(LOCTEXT("CollectionsViewSearchBoxHint", "Search Collections"));
-	}
-	SearchPtr->OnSearchChanged().AddSP(this, &SCollectionView::SetCollectionsSearchFilterText);
+	ExternalSearchPtr = InArgs._ExternalSearch;
+	TitleContent = SNew(SHorizontalBox);
+	SetAllowExternalSearch(true);
 
-	TSharedRef<SHorizontalBox> TitleContent = SNew(SHorizontalBox);
-	if (InArgs._ExternalSearch)
-	{
-		// If using an external search then just show the title and don't hide it when collapsing
-		TitleContent->AddSlot()
-		[
-			SNew(STextBlock)
-			.Font(FEditorStyle::GetFontStyle("ContentBrowser.SourceTitleFont"))
-			.Text(LOCTEXT("CollectionsListTitle", "Collections"))
-		];
-	}
-	else
-	{
-		// If using an internal search then show the title or search box depending on whether we're collapsed or not
-		TitleContent->AddSlot()
-		[
-			SNew(STextBlock)
-			.Font(FEditorStyle::GetFontStyle("ContentBrowser.SourceTitleFont"))
-			.Text(LOCTEXT("CollectionsListTitle", "Collections"))
-			.Visibility(this, &SCollectionView::GetCollectionsTitleTextVisibility)
-		];
-
-		TitleContent->AddSlot()
-		[
-			SNew(SBox)
-			.Visibility(this, &SCollectionView::GetCollectionsSearchBoxVisibility)
-			[
-				SearchPtr->GetWidget()
-			]
-		];
-	}
-
 	PreventSelectionChangedDelegateCount = 0;
 
@@ -169,7 +133,7 @@
 			.FillWidth(1.0f)
 			.Padding(0.0f)
 			[
-				TitleContent
+				TitleContent.ToSharedRef()
 			]
 
 			+SHorizontalBox::Slot()
@@ -297,6 +261,62 @@
 	UpdateCollectionItems();
 }
 
+void SCollectionView::SetAllowExternalSearch(const bool InAllowExternalSearch)
+{
+	if (bAllowExternalSearch == InAllowExternalSearch)
+	{
+		return;
+	}
+	bAllowExternalSearch = InAllowExternalSearch;
+
+	if (SearchPtr)
+	{
+		SearchPtr->OnSearchChanged().RemoveAll(this);
+	}
+	TitleContent->ClearChildren();
+
+	SearchPtr = bAllowExternalSearch ? ExternalSearchPtr : nullptr;
+	if (!SearchPtr)
+	{
+		SearchPtr = MakeShared<FSourcesSearch>();
+		SearchPtr->Initialize();
+		SearchPtr->SetHintText(LOCTEXT("CollectionsViewSearchBoxHint", "Search Collections"));
+	}
+	SearchPtr->OnSearchChanged().AddSP(this, &SCollectionView::SetCollectionsSearchFilterText);
+
+	if (SearchPtr == ExternalSearchPtr)
+	{
+		// If using an external search then just show the title and don't hide it when collapsing
+		TitleContent->AddSlot()
+		[
+			SNew(STextBlock)
+			.Font(FEditorStyle::GetFontStyle("ContentBrowser.SourceTitleFont"))
+			.Text(LOCTEXT("CollectionsListTitle", "Collections"))
+		];
+	}
+	else
+	{
+		// If using an internal search then show the title or search box depending on whether we're collapsed or not
+		TitleContent->AddSlot()
+		[
+			SNew(STextBlock)
+			.Font(FEditorStyle::GetFontStyle("ContentBrowser.SourceTitleFont"))
+			.Text(LOCTEXT("CollectionsListTitle", "Collections"))
+			.Visibility(this, &SCollectionView::GetCollectionsTitleTextVisibility)
+		];
+
+		TitleContent->AddSlot()
+		[
+			SNew(SBox)
+			.VAlign(VAlign_Center)
+			.Visibility(this, &SCollectionView::GetCollectionsSearchBoxVisibility)
+			[
+				SearchPtr->GetWidget()
+			]
+		];
+	}
+}
+
 void SCollectionView::HandleCollectionCreated( const FCollectionNameType& Collection )
 {
 	bQueueCollectionItemsUpdate = true;
diff --git a/Engine/Source/Editor/ContentBrowser/Private/SCollectionView.h b/Engine/Source/Editor/ContentBrowser/Private/SCollectionView.h
--- a/Engine/Source/Editor/ContentBrowser/Private/SCollectionView.h
+++ b/Engine/Source/Editor/ContentBrowser/Private/SCollectionView.h
@@ -22,6 +22,7 @@ class FCollectionContextMenu;
 class FCollectionDragDropOp;
 class FSourcesSearch;
 class FUICommandList;
+class SHorizontalBox;
 struct FHistoryData;
 
 /**
@@ -69,6 +70,9 @@ public:
 	/** Constructs this widget with InArgs */
 	void Construct( const FArguments& InArgs );
 
+	/** Set whether we're using an external search or not */
+	void SetAllowExternalSearch(const bool InAllowExternalSearch);
+
 	/** Selects the specified collections */
 	void SetSelectedCollections(const TArray<FCollectionNameType>& CollectionsToSelect, const bool bEnsureVisible = true);
 
@@ -296,6 +300,12 @@ private:
 	/** The collection search interface */
 	TSharedPtr<FSourcesSearch> SearchPtr;
 
+	/** The collection external search interface (if any) */
+	TSharedPtr<FSourcesSearch> ExternalSearchPtr;
+
+	/** The collection title content */
+	TSharedPtr<SHorizontalBox> TitleContent;
+
 	/** The collection tree widget */
 	TSharedPtr< STreeView< TSharedPtr<FCollectionItem> > > CollectionTreePtr;
 
@@ -331,6 +341,9 @@ private:
 	/** If true, the user will be able to drag collections from this view */
 	bool bAllowCollectionDrag;
 
+	/** If true, we will use the external search if provided during construction */
+	bool bAllowExternalSearch;
+
 	/** True when a drag is over this view with a valid operation for drop */
 	bool bDraggedOver;
 
diff --git a/Engine/Source/Editor/ContentBrowser/Private/SContentBrowser.cpp b/Engine/Source/Editor/ContentBrowser/Private/SContentBrowser.cpp
--- a/Engine/Source/Editor/ContentBrowser/Private/SContentBrowser.cpp
+++ b/Engine/Source/Editor/ContentBrowser/Private/SContentBrowser.cpp
@@ -129,9 +129,18 @@ void SContentBrowser::Construct( const FArguments& InArgs, const FName& InInstan
 	SourcesSearch->Initialize();
 	SourcesSearch->SetHintText(MakeAttributeSP(this, &SContentBrowser::GetSourcesSearchHintText));
 
+	CollectionViewPtr = SNew(SCollectionView)
+		.OnCollectionSelected(this, &SContentBrowser::CollectionSelected)
+		.AddMetaData<FTagMetaData>(FTagMetaData(TEXT("ContentBrowserCollections")))
+		.AllowCollapsing(false)
+		.AllowCollectionDrag(true)
+		.AllowQuickAssetManagement(true)
+		.ExternalSearch(SourcesSearch);
+
 	static const FName DefaultForegroundName("DefaultForeground");
 
 	BindCommands();
+	UContentBrowserSettings::OnSettingChanged().AddSP(this, &SContentBrowser::OnContentBrowserSettingsChanged);
 
 	ChildSlot
 	[
@@ -510,6 +519,7 @@ void SContentBrowser::Construct( const FArguments& InArgs, const FName& InInstan
 						.Padding(2, 0, 0, 0)
 						[
 							SNew(SButton)
+							.Visibility(this, &SContentBrowser::GetSourcesSwitcherVisibility)
 							.VAlign(EVerticalAlignment::VAlign_Center)
 							.ButtonStyle(FEditorStyle::Get(), "ToggleButton")
 							.ToolTipText(this, &SContentBrowser::GetSourcesSwitcherToolTipText)
@@ -536,8 +546,9 @@ void SContentBrowser::Construct( const FArguments& InArgs, const FName& InInstan
 							.Orientation(EOrientation::Orient_Vertical)
 							.MinimumSlotHeight(70.0f)
 							.Visibility( this, &SContentBrowser::GetSourcesViewVisibility )
-							+ SSplitter::Slot()
-							.Value(.2f)
+							
+							+SSplitter::Slot()
+							.Value(0.2f)
 							[
 								SNew(SBox)
 								.Visibility(this, &SContentBrowser::GetFavoriteFolderVisibility)
@@ -586,7 +597,8 @@ void SContentBrowser::Construct( const FArguments& InArgs, const FName& InInstan
 									]
 								]
 							]
-							+ SSplitter::Slot()
+							
+							+SSplitter::Slot()
 							.Value(0.8f)
 							[
 								SNew(SBox)
@@ -604,6 +616,16 @@ void SContentBrowser::Construct( const FArguments& InArgs, const FName& InInstan
 									.ExternalSearch(SourcesSearch)
 								]
 							]
+
+							+SSplitter::Slot()
+							.Value(0.4f)
+							[
+								SNew(SBox)
+								.Visibility(this, &SContentBrowser::GetDockedCollectionsVisibility)
+								[
+									CollectionViewPtr.ToSharedRef()
+								]
+							]
 						]
 
 						// Collections View
@@ -612,13 +634,7 @@ void SContentBrowser::Construct( const FArguments& InArgs, const FName& InInstan
 							SNew(SBox)
 							.Padding(FMargin(0.0f, 4.0f, 0.0f, 0.0f))
 							[
-								SAssignNew(CollectionViewPtr, SCollectionView)
-								.OnCollectionSelected(this, &SContentBrowser::CollectionSelected)
-								.AddMetaData<FTagMetaData>(FTagMetaData(TEXT("ContentBrowserCollections")))
-								.AllowCollapsing(false)
-								.AllowCollectionDrag(true)
-								.AllowQuickAssetManagement(true)
-								.ExternalSearch(SourcesSearch)
+								CollectionViewPtr.ToSharedRef()
 							]
 						]
 					]
@@ -784,6 +800,7 @@ void SContentBrowser::Construct( const FArguments& InArgs, const FName& InInstan
 						.CanShowRealTimeThumbnails( Config != nullptr ? Config->bCanShowRealTimeThumbnails : true)
 						.CanShowDevelopersFolder( Config != nullptr ? Config->bCanShowDevelopersFolder : true)
 						.CanShowFavorites(true)
+						.CanDockCollections(true)
 						.AddMetaData<FTagMetaData>(FTagMetaData(TEXT("ContentBrowserAssets")))
 						.OnSearchOptionsChanged(this, &SContentBrowser::HandleAssetViewSearchOptionsChanged)
 					]
@@ -861,7 +878,7 @@ void SContentBrowser::Construct( const FArguments& InArgs, const FName& InInstan
 	IAddContentDialogModule& AddContentDialogModule = FModuleManager::LoadModuleChecked<IAddContentDialogModule>("AddContentDialog");
 
 	// Update the breadcrumb trail path
-	UpdatePath();
+	OnContentBrowserSettingsChanged(NAME_None);
 }
 END_SLATE_FUNCTION_BUILD_OPTIMIZATION
 
@@ -926,6 +943,11 @@ EVisibility SContentBrowser::GetFavoriteFolderVisibility() const
 	return GetDefault<UContentBrowserSettings>()->GetDisplayFavorites() ? EVisibility::Visible : EVisibility::Collapsed;
 }
 
+EVisibility SContentBrowser::GetDockedCollectionsVisibility() const
+{
+	return GetDefault<UContentBrowserSettings>()->GetDockCollections() ? EVisibility::Visible : EVisibility::Collapsed;
+}
+
 void SContentBrowser::ToggleFolderFavorite(const TArray<FString>& FolderPaths)
 {
 	bool bAddedFavorite = false;
@@ -1729,7 +1751,7 @@ FReply SContentBrowser::OnSaveSearchButtonClicked()
 	{
 		SourcesViewExpandClicked();
 	}
-	if (ActiveSourcesWidgetIndex != ContentBrowserSourcesWidgetSwitcherIndex::CollectionsView)
+	if (!GetDefault<UContentBrowserSettings>()->GetDockCollections() && ActiveSourcesWidgetIndex != ContentBrowserSourcesWidgetSwitcherIndex::CollectionsView)
 	{
 		ActiveSourcesWidgetIndex = ContentBrowserSourcesWidgetSwitcherIndex::CollectionsView;
 		SourcesWidgetSwitcher->SetActiveWidgetIndex(ActiveSourcesWidgetIndex);
@@ -2336,6 +2358,11 @@ EVisibility SContentBrowser::GetPathExpanderVisibility() const
 	return bSourcesViewExpanded ? EVisibility::Collapsed : EVisibility::Visible;
 }
 
+EVisibility SContentBrowser::GetSourcesSwitcherVisibility() const
+{
+	return GetDefault<UContentBrowserSettings>()->GetDockCollections() ? EVisibility::Collapsed : EVisibility::Visible;
+}
+
 const FSlateBrush* SContentBrowser::GetSourcesSwitcherIcon() const
 {
 	switch (ActiveSourcesWidgetIndex)
@@ -2390,6 +2417,19 @@ FText SContentBrowser::GetSourcesSearchHintText() const
 	return FText();
 }
 
+void SContentBrowser::OnContentBrowserSettingsChanged(FName PropertyName)
+{
+	const FName NAME_DockCollections = "DockCollections";//GET_MEMBER_NAME_CHECKED(UContentBrowserSettings, DockCollections); // Doesn't work as DockCollections is private :(
+	if (PropertyName.IsNone() || PropertyName == NAME_DockCollections)
+	{
+		// Ensure the omni-search is enabled correctly
+		CollectionViewPtr->SetAllowExternalSearch(!GetDefault<UContentBrowserSettings>()->GetDockCollections());
+
+		// Ensure the path is set to the correct view mode
+		UpdatePath();
+	}
+}
+
 FReply SContentBrowser::BackClicked()
 {
 	HistoryManager.GoBack();
@@ -2693,7 +2733,7 @@ void SContentBrowser::UpdatePath()
 	}
 	else if ( SourcesData.HasCollections() )
 	{
-		NewSourcesWidgetIndex = ContentBrowserSourcesWidgetSwitcherIndex::CollectionsView;
+		NewSourcesWidgetIndex = GetDefault<UContentBrowserSettings>()->GetDockCollections() ? ContentBrowserSourcesWidgetSwitcherIndex::PathView : ContentBrowserSourcesWidgetSwitcherIndex::CollectionsView;
 
 		FCollectionManagerModule& CollectionManagerModule = FCollectionManagerModule::GetModule();
 		TArray<FCollectionNameType> CollectionPathItems;
diff --git a/Engine/Source/Editor/ContentBrowser/Private/SContentBrowser.h b/Engine/Source/Editor/ContentBrowser/Private/SContentBrowser.h
--- a/Engine/Source/Editor/ContentBrowser/Private/SContentBrowser.h
+++ b/Engine/Source/Editor/ContentBrowser/Private/SContentBrowser.h
@@ -292,6 +292,9 @@ private:
 	/** Gets the visibility of the path expander button */
 	EVisibility GetPathExpanderVisibility() const;
 
+	/** Gets the visibility of the source switch button */
+	EVisibility GetSourcesSwitcherVisibility() const;
+
 	/** Gets the icon used on the source switch button */
 	const FSlateBrush* GetSourcesSwitcherIcon() const;
 
@@ -304,6 +307,9 @@ private:
 	/** Gets the source search hint text */
 	FText GetSourcesSearchHintText() const;
 
+	/** Called to handle the Content Browser settings changing */
+	void OnContentBrowserSettingsChanged(FName PropertyName);
+
 	/** Handler for clicking the history back button */
 	FReply BackClicked();
 
@@ -442,6 +448,9 @@ private:
 	/** Gets the visibility of the favorites view */
 	EVisibility GetFavoriteFolderVisibility() const;
 
+	/** Get the visibility of the docked collections view */
+	EVisibility GetDockedCollectionsVisibility() const;
+
 	/** Toggles the favorite status of an array of folders*/
 	void ToggleFolderFavorite(const TArray<FString>& FolderPaths);
 
diff --git a/Engine/Source/Editor/UnrealEd/Classes/Settings/ContentBrowserSettings.h b/Engine/Source/Editor/UnrealEd/Classes/Settings/ContentBrowserSettings.h
--- a/Engine/Source/Editor/UnrealEd/Classes/Settings/ContentBrowserSettings.h
+++ b/Engine/Source/Editor/UnrealEd/Classes/Settings/ContentBrowserSettings.h
@@ -116,6 +116,18 @@ public:
 		return DisplayFavorites;
 	}
 
+	/** Sets whether we should dock the collections view under the paths view */
+	void SetDockCollections(bool bInDockCollections)
+	{
+		DockCollections = bInDockCollections;
+	}
+
+	/** Gets whether we should dock the collections view under the paths view */
+	bool GetDockCollections() const
+	{
+		return DockCollections;
+	}
+
 	/** Sets whether we are allowed to display C++ folders or not */
 	void SetDisplayCppFolders(bool bDisplay)
 	{
@@ -207,6 +219,9 @@ private:
 	UPROPERTY(config)
 	bool DisplayFavorites;
 
+	UPROPERTY(config)
+	bool DockCollections;
+
 	UPROPERTY(config)
 	bool DisplayCppFolders;
 
