Content Browser: Fix path context menu to use argument passed to function instead of member variable

diff --git a/Engine/Source/Editor/ContentBrowser/Private/PathContextMenu.cpp b/Engine/Source/Editor/ContentBrowser/Private/PathContextMenu.cpp
--- a/Engine/Source/Editor/ContentBrowser/Private/PathContextMenu.cpp
+++ b/Engine/Source/Editor/ContentBrowser/Private/PathContextMenu.cpp
@@ -107,7 +107,7 @@ TSharedRef<FExtender> FPathContextMenu::MakePathViewContextMenuExtender(const TA
 	{
 		if (MenuExtenderDelegates[i].IsBound())
 		{
-			Extenders.Add(MenuExtenderDelegates[i].Execute( SelectedPaths ));
+			Extenders.Add(MenuExtenderDelegates[i].Execute( InSelectedPaths ));
 		}
 	}
 	TSharedPtr<FExtender> MenuExtender = FExtender::Combine(Extenders);
