PR #4843: Enable all filters in content browser + shortcuts (Contributed by [REDACTED_PERSON])

diff --git a/Engine/Source/Editor/ContentBrowser/Private/SFilterList.cpp b/Engine/Source/Editor/ContentBrowser/Private/SFilterList.cpp
--- a/Engine/Source/Editor/ContentBrowser/Private/SFilterList.cpp
+++ b/Engine/Source/Editor/ContentBrowser/Private/SFilterList.cpp
@@ -30,6 +30,16 @@
 class SFilterCheckBox : public SCheckBox
 {
 public:
+	void SetOnFilterCtrlClicked(const FOnClicked& NewFilterCtrlClicked)
+	{
+		OnFilterCtrlClicked = NewFilterCtrlClicked;
+	}
+
+	void SetOnFilterAltClicked(const FOnClicked& NewFilteAltClicked)
+	{
+		OnFilterAltClicked = NewFilteAltClicked;
+	}
+
 	void SetOnFilterDoubleClicked( const FOnClicked& NewFilterDoubleClicked )
 	{
 		OnFilterDoubleClicked = NewFilterDoubleClicked;
@@ -54,7 +64,15 @@ public:
 
 	virtual FReply OnMouseButtonUp( const FGeometry& InMyGeometry, const FPointerEvent& InMouseEvent ) override
 	{
-		if( InMouseEvent.GetEffectingButton() == EKeys::MiddleMouseButton && OnFilterMiddleButtonClicked.IsBound() )
+		if (InMouseEvent.IsControlDown() && OnFilterCtrlClicked.IsBound())
+		{
+			return OnFilterCtrlClicked.Execute();
+		}
+		else if (InMouseEvent.IsAltDown() && OnFilterAltClicked.IsBound())
+		{
+			return OnFilterAltClicked.Execute();
+		}
+		else if( InMouseEvent.GetEffectingButton() == EKeys::MiddleMouseButton && OnFilterMiddleButtonClicked.IsBound() )
 		{
 			return OnFilterMiddleButtonClicked.Execute();
 		}
@@ -66,6 +84,8 @@ public:
 	}
 
 private:
+	FOnClicked OnFilterCtrlClicked;
+	FOnClicked OnFilterAltClicked;
 	FOnClicked OnFilterDoubleClicked;
 	FOnClicked OnFilterMiddleButtonClicked;
 };
@@ -78,6 +98,7 @@ class SFilter : public SCompoundWidget
 public:
 	DECLARE_DELEGATE_OneParam( FOnRequestRemove, const TSharedRef<SFilter>& /*FilterToRemove*/ );
 	DECLARE_DELEGATE_OneParam( FOnRequestEnableOnly, const TSharedRef<SFilter>& /*FilterToEnable*/ );
+	DECLARE_DELEGATE( FOnRequestEnableAll );
 	DECLARE_DELEGATE( FOnRequestDisableAll );
 	DECLARE_DELEGATE( FOnRequestRemoveAll );
 
@@ -98,6 +119,9 @@ public:
 		/** Invoked when a request to enable only this filter originated from within this filter */
 		SLATE_EVENT( FOnRequestEnableOnly, OnRequestEnableOnly )
 
+		/** Invoked when a request to enable all filters originated from within this filter */
+		SLATE_EVENT(FOnRequestEnableAll, OnRequestEnableAll)
+
 		/** Invoked when a request to disable all filters originated from within this filter */
 		SLATE_EVENT( FOnRequestDisableAll, OnRequestDisableAll )
 
@@ -114,6 +138,7 @@ public:
 		AssetTypeActions = InArgs._AssetTypeActions;
 		OnRequestRemove = InArgs._OnRequestRemove;
 		OnRequestEnableOnly = InArgs._OnRequestEnableOnly;
+		OnRequestEnableAll = InArgs._OnRequestEnableAll;
 		OnRequestDisableAll = InArgs._OnRequestDisableAll;
 		OnRequestRemoveAll = InArgs._OnRequestRemoveAll;
 		FrontendFilter = InArgs._FrontendFilter;
@@ -159,6 +184,8 @@ public:
 			]
 		];
 
+		ToggleButtonPtr->SetOnFilterCtrlClicked(FOnClicked::CreateSP(this, &SFilter::FilterCtrlClicked));
+		ToggleButtonPtr->SetOnFilterAltClicked(FOnClicked::CreateSP(this, &SFilter::FilterAltClicked));
 		ToggleButtonPtr->SetOnFilterDoubleClicked( FOnClicked::CreateSP(this, &SFilter::FilterDoubleClicked) );
 		ToggleButtonPtr->SetOnFilterMiddleButtonClicked( FOnClicked::CreateSP(this, &SFilter::FilterMiddleButtonClicked) );
 	}
@@ -218,6 +245,20 @@ private:
 		OnFilterChanged.ExecuteIfBound();
 	}
 
+	/** Handler for when the filter checkbox is clicked and a control key is pressed */
+	FReply FilterCtrlClicked()
+	{
+		OnRequestEnableAll.ExecuteIfBound();
+		return FReply::Handled();
+	}
+
+	/** Handler for when the filter checkbox is clicked and an alt key is pressed */
+	FReply FilterAltClicked()
+	{
+		OnRequestDisableAll.ExecuteIfBound();
+		return FReply::Handled();
+	}
+
 	/** Handler for when the filter checkbox is double clicked */
 	FReply FilterDoubleClicked()
 	{
@@ -262,6 +303,13 @@ private:
 
 		MenuBuilder.BeginSection("FilterBulkOptions", LOCTEXT("BulkFilterContextHeading", "Bulk Filter Options"));
 		{
+			MenuBuilder.AddMenuEntry(
+				LOCTEXT("EnableAllFilters", "Enable All Filters"),
+				LOCTEXT("EnableAllFiltersTooltip", "Enables all filters."),
+				FSlateIcon(),
+				FUIAction(FExecuteAction::CreateSP(this, &SFilter::EnableAllFilters))
+			);
+
 			MenuBuilder.AddMenuEntry(
 				LOCTEXT("DisableAllFilters", "Disable All Filters"),
 				LOCTEXT("DisableAllFiltersTooltip", "Disables all active filters."),
@@ -300,6 +348,12 @@ private:
 		OnRequestEnableOnly.ExecuteIfBound( Self );
 	}
 
+	/** Enables all filters in the list */
+	void EnableAllFilters()
+	{
+		OnRequestEnableAll.ExecuteIfBound();
+	}
+
 	/** Disables all active filters in the list */
 	void DisableAllFilters()
 	{
@@ -369,6 +423,9 @@ private:
 	/** Invoked when a request to enable only this filter originated from within this filter */
 	FOnRequestEnableOnly OnRequestEnableOnly;
 
+	/** Invoked when a request to enable all filters originated from within this filter */
+	FOnRequestEnableAll OnRequestEnableAll;
+
 	/** Invoked when a request to disable all filters originated from within this filter */
 	FOnRequestDisableAll OnRequestDisableAll;
 
@@ -563,6 +620,14 @@ TSharedRef<SWidget> SFilterList::ExternalMakeAddFilterMenu(EAssetTypeCategories:
 	return MakeAddFilterMenu(MenuExpansion);
 }
 
+void SFilterList::EnableAllFilters()
+{
+	for (auto FilterIt = Filters.CreateConstIterator(); FilterIt; ++FilterIt)
+	{
+		(*FilterIt)->SetEnabled(true);
+	}
+}
+
 void SFilterList::DisableAllFilters()
 {
 	for ( auto FilterIt = Filters.CreateConstIterator(); FilterIt; ++FilterIt )
@@ -822,6 +887,7 @@ TSharedRef<SFilter> SFilterList::AddFilter(const TWeakPtr<IAssetTypeActions>& As
 		.OnFilterChanged(OnFilterChanged)
 		.OnRequestRemove(this, &SFilterList::RemoveFilterAndUpdate)
 		.OnRequestEnableOnly(this, &SFilterList::EnableOnlyThisFilter)
+		.OnRequestEnableAll(this, &SFilterList::EnableAllFilters)
 		.OnRequestDisableAll(this, &SFilterList::DisableAllFilters)
 		.OnRequestRemoveAll(this, &SFilterList::RemoveAllFilters);
 
@@ -837,6 +903,7 @@ TSharedRef<SFilter> SFilterList::AddFilter(const TSharedRef<FFrontendFilter>& Fr
 		.FrontendFilter(FrontendFilter)
 		.OnFilterChanged( this, &SFilterList::FrontendFilterChanged, FrontendFilter )
 		.OnRequestRemove(this, &SFilterList::RemoveFilterAndUpdate)
+		.OnRequestEnableAll(this, &SFilterList::EnableAllFilters)
 		.OnRequestDisableAll(this, &SFilterList::DisableAllFilters)
 		.OnRequestRemoveAll(this, &SFilterList::RemoveAllFilters);
 
diff --git a/Engine/Source/Editor/ContentBrowser/Private/SFilterList.h b/Engine/Source/Editor/ContentBrowser/Private/SFilterList.h
--- a/Engine/Source/Editor/ContentBrowser/Private/SFilterList.h
+++ b/Engine/Source/Editor/ContentBrowser/Private/SFilterList.h
@@ -57,6 +57,9 @@ public:
 	/** Handler for when the floating add filter button was clicked */
 	TSharedRef<SWidget> ExternalMakeAddFilterMenu(EAssetTypeCategories::Type MenuExpansion = EAssetTypeCategories::Basic);
 
+	/** Enables all filters */
+	void EnableAllFilters();
+
 	/** Disables any applied filters */
 	void DisableAllFilters();
 
