PR #4562: UE-21733: Select parent of child folder when collapsing parent (Contributed by [REDACTED_PERSON])

diff --git a/Engine/Source/Editor/ContentBrowser/Private/SPathView.cpp b/Engine/Source/Editor/ContentBrowser/Private/SPathView.cpp
--- a/Engine/Source/Editor/ContentBrowser/Private/SPathView.cpp
+++ b/Engine/Source/Editor/ContentBrowser/Private/SPathView.cpp
@@ -967,6 +967,28 @@ void SPathView::TreeExpansionChanged( TSharedPtr< FTreeItem > TreeItem, bool bIs
 			// Keep track of the last paths that we broadcasted for expansion reasons when filtering
 			LastExpandedPaths.Add(Item->FolderPath);
 		}
+
+		if (!bIsExpanded)
+		{
+			const TArray<TSharedPtr<FTreeItem>> SelectedItems = TreeViewPtr->GetSelectedItems();
+			bool bSelectTreeItem = false;
+
+			// If any selected item was a child of the collapsed node, then add the collapsed node to the current selection
+			// This avoids the selection ever becoming empty, as this causes the Content Browser to show everything
+			for (const TSharedPtr<FTreeItem>& SelectedItem : SelectedItems)
+			{
+				if (SelectedItem->IsChildOf(*TreeItem.Get()))
+				{
+					bSelectTreeItem = true;
+					break;
+				}
+			}
+
+			if (bSelectTreeItem)
+			{
+				TreeViewPtr->SetItemSelection(TreeItem, true);
+			}
+		}
 	}
 }
 
