It is now possible to disable recursive filtering in the asset view, so that when an asset type filter is enabled, you don't see all assets of that type at once. View Options > Filter Recursively toggles this behaviour.

diff --git a/Engine/Source/Editor/ContentBrowser/Private/SAssetView.cpp b/Engine/Source/Editor/ContentBrowser/Private/SAssetView.cpp
--- a/Engine/Source/Editor/ContentBrowser/Private/SAssetView.cpp
+++ b/Engine/Source/Editor/ContentBrowser/Private/SAssetView.cpp
@@ -2041,6 +2041,25 @@ void SAssetView::RefreshSourceItems()
 	}
 }
 
+bool SAssetView::IsFilteringRecursively() const
+{
+	// In some cases we want to not filter recursively even if we have a backend filter (e.g. the open level window)
+	// Most of the time, bFilterRecursivelyWithBackendFilter is true
+	return bFilterRecursivelyWithBackendFilter && GetDefault<UContentBrowserSettings>()->FilterRecursively;
+}
+
+bool SAssetView::IsToggleFilteringRecursivelyAllowed() const
+{
+	return bFilterRecursivelyWithBackendFilter;
+}
+
+void SAssetView::ToggleFilteringRecursively()
+{
+	check(IsToggleFilteringRecursivelyAllowed());
+	GetMutableDefault<UContentBrowserSettings>()->FilterRecursively = !GetDefault<UContentBrowserSettings>()->FilterRecursively;
+	GetMutableDefault<UContentBrowserSettings>()->PostEditChange();
+}
+
 bool SAssetView::ShouldFilterRecursively() const
 {
 	// Quick check for conditions which force recursive filtering
@@ -2049,9 +2068,7 @@ bool SAssetView::ShouldFilterRecursively() const
 		return true;
 	}
 
-	// In some cases we want to not filter recursively even if we have a backend filter (e.g. the open level window)
-	// Most of the time, bFilterRecursivelyWithBackendFilter is true
-	if ( bFilterRecursivelyWithBackendFilter && !BackendFilter.IsEmpty() )
+	if (IsFilteringRecursively() && !BackendFilter.IsEmpty() )
 	{
 		return true;
 	}
@@ -3131,6 +3148,19 @@ TSharedRef<SWidget> SAssetView::GetViewButtonContent()
 			NAME_None,
 			EUserInterfaceActionType::ToggleButton
 		);
+
+		MenuBuilder.AddMenuEntry(
+			LOCTEXT("FilterRecursivelyOption", "Filter Recursively"),
+			LOCTEXT("FilterRecursivelyOptionToolTip", "Should filters apply recursively in the view?"),
+			FSlateIcon(),
+			FUIAction(
+				FExecuteAction::CreateSP(this, &SAssetView::ToggleFilteringRecursively),
+				FCanExecuteAction::CreateSP(this, &SAssetView::IsToggleFilteringRecursivelyAllowed),
+				FIsActionChecked::CreateSP(this, &SAssetView::IsFilteringRecursively)
+			),
+			NAME_None,
+			EUserInterfaceActionType::ToggleButton
+		);
 	}
 	MenuBuilder.EndSection();
 
diff --git a/Engine/Source/Editor/ContentBrowser/Private/SAssetView.h b/Engine/Source/Editor/ContentBrowser/Private/SAssetView.h
--- a/Engine/Source/Editor/ContentBrowser/Private/SAssetView.h
+++ b/Engine/Source/Editor/ContentBrowser/Private/SAssetView.h
@@ -557,6 +557,15 @@ private:
 	/** Whether or not it's possible to include collection names in search criteria */
 	bool IsToggleIncludeCollectionNamesAllowed() const;
 
+	/** @return true when we are filtering recursively when we have an asset path */
+	bool IsFilteringRecursively() const;
+
+	/** Whether or not it's possible to toggle how to filtering recursively */
+	bool IsToggleFilteringRecursivelyAllowed() const;
+
+	/** Toggle whether we're filtering recursively */
+	void ToggleFilteringRecursively();
+
 	/** Sets the view type and updates lists accordingly */
 	void SetCurrentViewType(EAssetViewType::Type NewType);
 
diff --git a/Engine/Source/Editor/UnrealEd/Classes/Settings/ContentBrowserSettings.h b/Engine/Source/Editor/UnrealEd/Classes/Settings/ContentBrowserSettings.h
--- a/Engine/Source/Editor/UnrealEd/Classes/Settings/ContentBrowserSettings.h
+++ b/Engine/Source/Editor/UnrealEd/Classes/Settings/ContentBrowserSettings.h
@@ -42,6 +42,10 @@ public:
 	UPROPERTY(config)
 	bool DisplayEmptyFolders;
 
+	/** Whether to filter recursively when a filter is applied in the asset view of the content browser. */
+	UPROPERTY(config)
+	bool FilterRecursively = true;
+
 	/** The number of objects to keep in the Content Browser Recently Opened filter */
 	UPROPERTY(EditAnywhere, config, Category = ContentBrowser, meta = (DisplayName = "Number of Assets to Keep in the Recently Opened Filter", ClampMin = "1", ClampMax = "30"))
 	int32 NumObjectsInRecentList;
