Fix show in content browser not working for engine plugin content.

diff --git a/Engine/Source/Editor/ContentBrowser/Private/SContentBrowser.cpp b/Engine/Source/Editor/ContentBrowser/Private/SContentBrowser.cpp
--- a/Engine/Source/Editor/ContentBrowser/Private/SContentBrowser.cpp
+++ b/Engine/Source/Editor/ContentBrowser/Private/SContentBrowser.cpp
@@ -997,6 +997,17 @@ void SContentBrowser::PrepareToSync( const TArray<FAssetData>& AssetDataList, co
 				bDisplayEngine = true;
 				GetMutableDefault<UContentBrowserSettings>()->SetDisplayEngineFolder(true, true);
 				bRepopulate = true;
+
+				// Handle being a plugin as well
+				if (!bDisplayPlugins && (FolderCategory == ContentBrowserUtils::ECBFolderCategory::EngineContent))
+				{
+					EPluginLoadedFrom PluginSource;
+					if (ContentBrowserUtils::IsPluginFolder(PackagePath, &PluginSource))
+					{
+						bDisplayPlugins = true;
+						GetMutableDefault<UContentBrowserSettings>()->SetDisplayPluginFolders(true, true);
+					}
+				}
 			}
 			else if ( !bDisplayPlugins && (FolderCategory == ContentBrowserUtils::ECBFolderCategory::PluginContent || FolderCategory == ContentBrowserUtils::ECBFolderCategory::PluginClasses) )
 			{
