Content Browser Context Menu refactored to use new menu system
#rnx

diff --git a/Engine/Plugins/Media/MediaPlayerEditor/Source/MediaPlayerEditor/Private/AssetTools/MediaTextureActions.cpp b/Engine/Plugins/Media/MediaPlayerEditor/Source/MediaPlayerEditor/Private/AssetTools/MediaTextureActions.cpp
--- a/Engine/Plugins/Media/MediaPlayerEditor/Source/MediaPlayerEditor/Private/AssetTools/MediaTextureActions.cpp
+++ b/Engine/Plugins/Media/MediaPlayerEditor/Source/MediaPlayerEditor/Private/AssetTools/MediaTextureActions.cpp
@@ -2,7 +2,7 @@
 
 #include "AssetTools/MediaTextureActions.h"
 #include "Engine/Texture.h"
-#include "Framework/MultiBox/MultiBoxBuilder.h"
+#include "EditorMenuSubsystem.h"
 #include "Misc/PackageName.h"
 #include "IContentBrowserSingleton.h"
 #include "ContentBrowserModule.h"
@@ -23,14 +23,12 @@ bool FMediaTextureActions::CanFilter()
 	return true;
 }
 
-
-void FMediaTextureActions::GetActions(const TArray<UObject*>& InObjects, FMenuBuilder& MenuBuilder)
+void FMediaTextureActions::GetActions(const TArray<UObject*>& InObjects, FEditorMenuSection& Section)
 {
-	FAssetTypeActions_Base::GetActions(InObjects, MenuBuilder);
-
 	auto Textures = GetTypedWeakObjectPtrs<UTexture>(InObjects);
 
-	MenuBuilder.AddMenuEntry(
+	Section.AddMenuEntry(
+		"MediaTexture_CreateMaterial",
 		LOCTEXT("MediaTexture_CreateMaterial", "Create Material"),
 		LOCTEXT("MediaTexture_CreateMaterialTooltip", "Creates a new material using this texture."),
 		FSlateIcon(FEditorStyle::GetStyleSetName(), "ClassIcon.Material"),
diff --git a/Engine/Plugins/Media/MediaPlayerEditor/Source/MediaPlayerEditor/Private/AssetTools/MediaTextureActions.h b/Engine/Plugins/Media/MediaPlayerEditor/Source/MediaPlayerEditor/Private/AssetTools/MediaTextureActions.h
--- a/Engine/Plugins/Media/MediaPlayerEditor/Source/MediaPlayerEditor/Private/AssetTools/MediaTextureActions.h
+++ b/Engine/Plugins/Media/MediaPlayerEditor/Source/MediaPlayerEditor/Private/AssetTools/MediaTextureActions.h
@@ -7,7 +7,7 @@
 #include "Toolkits/IToolkitHost.h"
 #include "AssetTypeActions_Base.h"
 
-class FMenuBuilder;
+struct FEditorMenuSection;
 class UTexture;
 
 /**
@@ -21,7 +21,7 @@ public:
 	//~ FAssetTypeActions_Base overrides
 
 	virtual bool CanFilter() override;
-	virtual void GetActions(const TArray<UObject*>& InObjects, FMenuBuilder& MenuBuilder) override;
+	virtual void GetActions(const TArray<UObject*>& InObjects, FEditorMenuSection& Section) override;
 	virtual uint32 GetCategories() override;
 	virtual FText GetName() const override;
 	virtual UClass* GetSupportedClass() const override;
diff --git a/Engine/Source/Developer/AssetTools/AssetTools.Build.cs b/Engine/Source/Developer/AssetTools/AssetTools.Build.cs
--- a/Engine/Source/Developer/AssetTools/AssetTools.Build.cs
+++ b/Engine/Source/Developer/AssetTools/AssetTools.Build.cs
@@ -33,7 +33,9 @@ public class AssetTools : ModuleRules
                 "Foliage",
                 "Projects",
 				"RHI",
-				"MaterialEditor"
+				"MaterialEditor",
+				"EditorMenus",
+				"ContentBrowser",
             }
 		);
 
@@ -69,7 +71,6 @@ public class AssetTools : ModuleRules
 		DynamicallyLoadedModuleNames.AddRange(
 			new string[] {
 				"AssetRegistry",
-				"ContentBrowser",
 				"CollectionManager",
 				"CurveTableEditor",
 				"DataTableEditor",
diff --git a/Engine/Source/Developer/AssetTools/Private/AssetTools.cpp b/Engine/Source/Developer/AssetTools/Private/AssetTools.cpp
--- a/Engine/Source/Developer/AssetTools/Private/AssetTools.cpp
+++ b/Engine/Source/Developer/AssetTools/Private/AssetTools.cpp
@@ -388,48 +388,6 @@ TWeakPtr<IClassTypeActions> UAssetToolsImpl::GetClassTypeActionsForClass( UClass
 	return MostDerivedClassTypeActions;
 }
 
-bool UAssetToolsImpl::GetAssetActions( const TArray<UObject*>& InObjects, FMenuBuilder& MenuBuilder, bool bIncludeHeading )
-{
-	bool bAddedActions = false;
-
-	if ( InObjects.Num() )
-	{
-		// Find the most derived common class for all passed in Objects
-		UClass* CommonClass = InObjects[0]->GetClass();
-		for (int32 ObjIdx = 1; ObjIdx < InObjects.Num(); ++ObjIdx)
-		{
-			while (!InObjects[ObjIdx]->IsA(CommonClass))
-			{
-				CommonClass = CommonClass->GetSuperClass();
-			}
-		}
-
-		// Get the nearest common asset type for all the selected objects
-		TSharedPtr<IAssetTypeActions> CommonAssetTypeActions = GetAssetTypeActionsForClass(CommonClass).Pin();
-
-		// If we found a common type actions object, get actions from it
-		if ( CommonAssetTypeActions.IsValid() && CommonAssetTypeActions->HasActions(InObjects) )
-		{
-			if ( bIncludeHeading )
-			{
-				MenuBuilder.BeginSection("GetAssetActions", FText::Format( LOCTEXT("AssetSpecificOptionsMenuHeading", "{0} Actions"), CommonAssetTypeActions->GetName() ) );
-			}
-
-			// Get the actions
-			CommonAssetTypeActions->GetActions(InObjects, MenuBuilder);
-
-			if ( bIncludeHeading )
-			{
-				MenuBuilder.EndSection();
-			}
-
-			bAddedActions = true;
-		}
-	}
-
-	return bAddedActions;
-}
-
 struct FRootedOnScope
 {
 	UObject* Obj;
diff --git a/Engine/Source/Developer/AssetTools/Private/AssetTools.h b/Engine/Source/Developer/AssetTools/Private/AssetTools.h
--- a/Engine/Source/Developer/AssetTools/Private/AssetTools.h
+++ b/Engine/Source/Developer/AssetTools/Private/AssetTools.h
@@ -10,7 +10,7 @@
 #include "AssetTools.generated.h"
 
 class FAssetFixUpRedirectors;
-class FMenuBuilder;
+class UEditorMenu;
 class IClassTypeActions;
 class UAutomatedAssetImportData;
 class UFactory;
@@ -71,7 +71,6 @@ public:
 	virtual void UnregisterClassTypeActions(const TSharedRef<IClassTypeActions>& ActionsToRemove) override;
 	virtual void GetClassTypeActionsList( TArray<TWeakPtr<IClassTypeActions>>& OutClassTypeActionsList ) const override;
 	virtual TWeakPtr<IClassTypeActions> GetClassTypeActionsForClass( UClass* Class ) const override;
-	virtual bool GetAssetActions( const TArray<UObject*>& InObjects, FMenuBuilder& MenuBuilder, bool bIncludeHeading = true ) override;
 	virtual UObject* CreateAsset(const FString& AssetName, const FString& PackagePath, UClass* AssetClass, UFactory* Factory, FName CallingContext = NAME_None) override;
 	virtual UObject* CreateAsset(UClass* AssetClass, UFactory* Factory, FName CallingContext = NAME_None) override;
 	virtual UObject* CreateAssetWithDialog(UClass* AssetClass, UFactory* Factory, FName CallingContext = NAME_None) override;
diff --git a/Engine/Source/Developer/AssetTools/Private/AssetTypeActions/AssetTypeActions_StaticMesh.cpp b/Engine/Source/Developer/AssetTools/Private/AssetTypeActions/AssetTypeActions_StaticMesh.cpp
--- a/Engine/Source/Developer/AssetTools/Private/AssetTypeActions/AssetTypeActions_StaticMesh.cpp
+++ b/Engine/Source/Developer/AssetTools/Private/AssetTypeActions/AssetTypeActions_StaticMesh.cpp
@@ -2,6 +2,7 @@
 
 #include "AssetTypeActions/AssetTypeActions_StaticMesh.h"
 #include "Framework/MultiBox/MultiBoxBuilder.h"
+#include "EditorMenuSubsystem.h"
 #include "EditorStyleSet.h"
 #include "EditorFramework/AssetImportData.h"
 #include "ThumbnailRendering/SceneThumbnailInfo.h"
@@ -23,13 +24,14 @@ static TAutoConsoleVariable<int32> CVarEnableSaveGeneratedLODsInPackage(
 	TEXT("1 - Enable this option and save the LODs in the Package.\n"),
 	ECVF_Default);
 
-void FAssetTypeActions_StaticMesh::GetActions( const TArray<UObject*>& InObjects, FMenuBuilder& MenuBuilder )
+void FAssetTypeActions_StaticMesh::GetActions(const TArray<UObject*>& InObjects, FEditorMenuSection& Section)
 {
 	auto Meshes = GetTypedWeakObjectPtrs<UStaticMesh>(InObjects);
 
 	if (CVarEnableSaveGeneratedLODsInPackage.GetValueOnGameThread() != 0)
 	{
-		MenuBuilder.AddMenuEntry(
+		Section.AddMenuEntry(
+			"ObjectContext_SaveGeneratedLODsInPackage",
 			NSLOCTEXT("AssetTypeActions_StaticMesh", "ObjectContext_SaveGeneratedLODsInPackage", "Save Generated LODs"),
 			NSLOCTEXT("AssetTypeActions_StaticMesh", "ObjectContext_SaveGeneratedLODsInPackageTooltip", "Run the mesh reduce and save the generated LODs as part of the package."),
 			FSlateIcon(),
@@ -40,15 +42,18 @@ void FAssetTypeActions_StaticMesh::GetActions( const TArray<UObject*>& InObjects
 			);
 	}
 
-	MenuBuilder.AddSubMenu(
+	Section.AddEntry(FEditorMenuEntry::InitSubMenu(
+		NAME_None,
+		"StaticMesh_LODMenu",
 		NSLOCTEXT("AssetTypeActions_StaticMesh", "StaticMesh_LODMenu", "Level Of Detail"),
 		NSLOCTEXT("AssetTypeActions_StaticMesh", "StaticMesh_LODTooltip", "LOD Options and Tools"),
 		FNewMenuDelegate::CreateSP(this, &FAssetTypeActions_StaticMesh::GetLODMenu, Meshes),
 		false,
 		FSlateIcon(FEditorStyle::GetStyleSetName(), "ContentBrowser.AssetActions")
-		);
+		));
 
-	MenuBuilder.AddMenuEntry(
+	Section.AddMenuEntry(
+		"ObjectContext_ClearVertexColors",
 		NSLOCTEXT("AssetTypeActions_StaticMesh", "ObjectContext_ClearVertexColors", "Remove Vertex Colors"),
 		NSLOCTEXT("AssetTypeActions_StaticMesh", "ObjectContext_ClearVertexColorsTooltip", "Removes vertex colors from all LODS in all selected meshes."),
 		FSlateIcon(),
diff --git a/Engine/Source/Developer/AssetTools/Private/AssetTypeActions/AssetTypeActions_StaticMesh.h b/Engine/Source/Developer/AssetTools/Private/AssetTypeActions/AssetTypeActions_StaticMesh.h
--- a/Engine/Source/Developer/AssetTools/Private/AssetTypeActions/AssetTypeActions_StaticMesh.h
+++ b/Engine/Source/Developer/AssetTools/Private/AssetTypeActions/AssetTypeActions_StaticMesh.h
@@ -7,6 +7,7 @@
 #include "Toolkits/IToolkitHost.h"
 #include "AssetTypeActions_Base.h"
 
+struct FEditorMenuSection;
 class FMenuBuilder;
 
 class FAssetTypeActions_StaticMesh : public FAssetTypeActions_Base
@@ -17,7 +18,7 @@ public:
 	virtual FColor GetTypeColor() const override { return FColor(0, 255, 255); }
 	virtual UClass* GetSupportedClass() const override { return UStaticMesh::StaticClass(); }
 	virtual bool HasActions( const TArray<UObject*>& InObjects ) const override { return true; }
-	virtual void GetActions( const TArray<UObject*>& InObjects, FMenuBuilder& MenuBuilder ) override;
+	virtual void GetActions(const TArray<UObject*>& InObjects, FEditorMenuSection& Section) override;
 	virtual void OpenAssetEditor( const TArray<UObject*>& InObjects, TSharedPtr<class IToolkitHost> EditWithinLevelEditor = TSharedPtr<IToolkitHost>() ) override;
 	virtual uint32 GetCategories() override { return EAssetTypeCategories::Basic; }
 	virtual class UThumbnailInfo* GetThumbnailInfo(UObject* Asset) const override;
diff --git a/Engine/Source/Developer/AssetTools/Public/AssetTypeActions_Base.h b/Engine/Source/Developer/AssetTools/Public/AssetTypeActions_Base.h
--- a/Engine/Source/Developer/AssetTools/Public/AssetTypeActions_Base.h
+++ b/Engine/Source/Developer/AssetTools/Public/AssetTypeActions_Base.h
@@ -18,6 +18,7 @@
 struct FAssetData;
 struct FARFilter;
 class FMenuBuilder;
+struct FEditorMenuSection;
 
 /** A base class for all AssetTypeActions. Provides helper functions useful for many types. Deriving from this class is optional. */
 class FAssetTypeActions_Base : public IAssetTypeActions
@@ -36,6 +37,11 @@ public:
 
 	}
 
+	virtual void GetActions(const TArray<UObject*>& InObjects, FEditorMenuSection& Section) override
+	{
+
+	}
+
 	virtual void OpenAssetEditor( const TArray<UObject*>& InObjects, TSharedPtr<class IToolkitHost> EditWithinLevelEditor = TSharedPtr<IToolkitHost>() ) override
 	{
 		FSimpleAssetEditor::CreateEditor(EToolkitMode::Standalone, EditWithinLevelEditor, InObjects);
@@ -186,4 +192,18 @@ protected:
 
 		return TypedObjects;
 	}
+
+	template <typename T>
+	static TArray<T> GetTypedObjectPtrs(const TArray<UObject*>& InObjects)
+	{
+		check(InObjects.Num() > 0);
+
+		TArray<T> TypedObjects;
+		for (UObject* ObjIt : InObjects)
+		{
+			TypedObjects.Add(CastChecked<T>(ObjIt));
+		}
+
+		return TypedObjects;
+	}
 };
diff --git a/Engine/Source/Developer/AssetTools/Public/IAssetTools.h b/Engine/Source/Developer/AssetTools/Public/IAssetTools.h
--- a/Engine/Source/Developer/AssetTools/Public/IAssetTools.h
+++ b/Engine/Source/Developer/AssetTools/Public/IAssetTools.h
@@ -192,16 +192,6 @@ public:
 	/** Gets the appropriate ClassTypeActions for the supplied class */
 	virtual TWeakPtr<IClassTypeActions> GetClassTypeActionsForClass(UClass* Class) const = 0;
 
-	/**
-	 * Fills out a menubuilder with a list of commands that can be applied to the specified objects.
-	 *
-	 * @param InObjects the objects for which to generate type-specific menu options
-	 * @param MenuBuilder the menu in which to build options
-	 * @param bIncludeHeader if true, will include a heading in the menu if any options were found
-	 * @return true if any options were added to the MenuBuilder
-	 */
-	virtual bool GetAssetActions(const TArray<UObject*>& InObjects, class FMenuBuilder& MenuBuilder, bool bIncludeHeading = true) = 0;
-
 	/**
 	 * Creates an asset with the specified name, path, and factory
 	 *
diff --git a/Engine/Source/Developer/AssetTools/Public/IAssetTypeActions.h b/Engine/Source/Developer/AssetTools/Public/IAssetTypeActions.h
--- a/Engine/Source/Developer/AssetTools/Public/IAssetTypeActions.h
+++ b/Engine/Source/Developer/AssetTools/Public/IAssetTypeActions.h
@@ -57,6 +57,9 @@ public:
 	/** Generates a menubuilder for the specified objects. */
 	virtual void GetActions( const TArray<UObject*>& InObjects, class FMenuBuilder& MenuBuilder ) = 0;
 
+	/** Generates a menu section for the specified objects. */
+	virtual void GetActions(const TArray<UObject*>& InObjects, struct FEditorMenuSection& Section) = 0;
+
 	/** Opens the asset editor for the specified objects. If EditWithinLevelEditor is valid, the world-centric editor will be used. */
 	virtual void OpenAssetEditor( const TArray<UObject*>& InObjects, TSharedPtr<IToolkitHost> EditWithinLevelEditor = TSharedPtr<IToolkitHost>() ) = 0;
 
diff --git a/Engine/Source/Editor/ContentBrowser/ContentBrowser.Build.cs b/Engine/Source/Editor/ContentBrowser/ContentBrowser.Build.cs
--- a/Engine/Source/Editor/ContentBrowser/ContentBrowser.Build.cs
+++ b/Engine/Source/Editor/ContentBrowser/ContentBrowser.Build.cs
@@ -39,7 +39,8 @@ public class ContentBrowser : ModuleRules
 				"Projects",
 				"AddContentDialog",
 				"DesktopPlatform",
-				"AssetRegistry"
+				"AssetRegistry",
+				"EditorMenus",
 			}
 		);
 
diff --git a/Engine/Source/Editor/ContentBrowser/Private/AssetContextMenu.cpp b/Engine/Source/Editor/ContentBrowser/Private/AssetContextMenu.cpp
--- a/Engine/Source/Editor/ContentBrowser/Private/AssetContextMenu.cpp
+++ b/Engine/Source/Editor/ContentBrowser/Private/AssetContextMenu.cpp
@@ -20,6 +20,8 @@
 #include "Widgets/Text/STextBlock.h"
 #include "Widgets/Input/SMultiLineEditableTextBox.h"
 #include "Framework/MultiBox/MultiBoxBuilder.h"
+#include "EditorMenuSubsystem.h"
+#include "ContentBrowserMenuContexts.h"
 #include "Widgets/Input/SButton.h"
 #include "EditorStyleSet.h"
 #include "EditorReimportHandler.h"
@@ -131,37 +133,140 @@ TSharedRef<SWidget> FAssetContextMenu::MakeContextMenu(const TArray<FAssetData>&
 	}
 	TSharedPtr<FExtender> MenuExtender = FExtender::Combine(Extenders);
 
-	FMenuBuilder MenuBuilder(/*bInShouldCloseWindowAfterMenuSelection=*/true, InCommandList, MenuExtender);
+	UContentBrowserAssetContextMenuContext* ContextObject = NewObject<UContentBrowserAssetContextMenuContext>();
+	ContextObject->AssetContextMenu = SharedThis(this);
+
+	UEditorMenuSubsystem* EditorMenus = UEditorMenuSubsystem::Get();
+
+	static const FName BaseMenuName("ContentBrowser.AssetContextMenu");
+	RegisterContextMenu(BaseMenuName);
+
+	// Create menu hierarchy based on class hierarchy
+	TArray<UEditorMenu*> MenuHierarchy;
+	FName MenuName = BaseMenuName;
+	{
+		// Objects must be loaded for this operation... for now
+		TArray<FString> ObjectPaths;
+		for (int32 AssetIdx = 0; AssetIdx < SelectedAssets.Num(); ++AssetIdx)
+		{
+			ObjectPaths.Add(SelectedAssets[AssetIdx].ObjectPath.ToString());
+		}
+
+		ContextObject->SelectedObjects.Reset();
+		if (ContentBrowserUtils::LoadAssetsIfNeeded(ObjectPaths, ContextObject->SelectedObjects) && ContextObject->SelectedObjects.Num() > 0)
+		{
+			// Find common class for selected objects
+			UClass* CommonClass = ContextObject->SelectedObjects[0]->GetClass();
+			for (int32 ObjIdx = 1; ObjIdx < ContextObject->SelectedObjects.Num(); ++ObjIdx)
+			{
+				while (!ContextObject->SelectedObjects[ObjIdx]->IsA(CommonClass))
+				{
+					CommonClass = CommonClass->GetSuperClass();
+				}
+			}
+			ContextObject->CommonClass = CommonClass;
+
+			// Find asset actions for common class
+			FAssetToolsModule& AssetToolsModule = FModuleManager::LoadModuleChecked<FAssetToolsModule>("AssetTools");
+			TSharedPtr<IAssetTypeActions> CommonAssetTypeActions = AssetToolsModule.Get().GetAssetTypeActionsForClass(ContextObject->CommonClass).Pin();
+			UClass* AssetTypeActionsClass = nullptr;
+			if (CommonAssetTypeActions.IsValid() && CommonAssetTypeActions->HasActions(ContextObject->SelectedObjects))
+			{
+				ContextObject->CommonAssetTypeActions = CommonAssetTypeActions;
+				AssetTypeActionsClass = CommonAssetTypeActions->GetSupportedClass();
+				MenuName = UEditorMenuSubsystem::JoinMenuPaths(BaseMenuName, AssetTypeActionsClass->GetFName());
+			}
+
+			// Find menus for common class
+			UClass* ClassWithMenu = CommonClass;
+			while (ClassWithMenu)
+			{
+				FName CheckName = UEditorMenuSubsystem::JoinMenuPaths(BaseMenuName, ClassWithMenu->GetFName());
+				UEditorMenu* FoundMenu = EditorMenus->FindMenu(CheckName);
+				if (FoundMenu && FoundMenu->Sections.Num() > 0)
+				{
+					// Change generated menu's name if necessary
+					if (MenuHierarchy.Num() == 0)
+					{
+						if (!AssetTypeActionsClass || ClassWithMenu->IsChildOf(AssetTypeActionsClass))
+						{
+							MenuName = CheckName;
+						}
+					}
+
+					MenuHierarchy.Add(FoundMenu);
+				}
+
+				ClassWithMenu = ClassWithMenu->GetSuperClass();
+			}
+		}
+
+		MenuHierarchy.Add(EditorMenus->FindMenu(BaseMenuName));
+
+		Algo::Reverse(MenuHierarchy);
+	}
+
+	FEditorMenuContext MenuContext(InCommandList, MenuExtender, ContextObject);
+	UEditorMenu* Menu = EditorMenus->GenerateMenu(MenuHierarchy, MenuContext);
+	Menu->MenuName = MenuName;
 
 	// Only add something if at least one asset is selected
 	if ( SelectedAssets.Num() )
 	{
 		// Add any type-specific context menu options
-		AddAssetTypeMenuOptions(MenuBuilder);
+		AddAssetTypeMenuOptions(Menu, ContextObject->SelectedObjects);
 
 		// Add imported asset context menu options
-		AddImportedAssetMenuOptions(MenuBuilder);
+		AddImportedAssetMenuOptions(Menu);
 
 		// Add quick access to common commands.
-		AddCommonMenuOptions(MenuBuilder);
+		AddCommonMenuOptions(Menu);
 
 		// Add quick access to view commands
-		AddExploreMenuOptions(MenuBuilder);
+		AddExploreMenuOptions(Menu);
 
 		// Add reference options
-		AddReferenceMenuOptions(MenuBuilder);
+		AddReferenceMenuOptions(Menu);
 
 		// Add collection options
-		AddCollectionMenuOptions(MenuBuilder);
+		AddCollectionMenuOptions(Menu);
 
 		// Add documentation options
-		AddDocumentationMenuOptions(MenuBuilder);
+		AddDocumentationMenuOptions(Menu);
 
 		// Add source control options
-		AddSourceControlMenuOptions(MenuBuilder);
+		AddSourceControlMenuOptions(Menu);
 	}
 	
-	return MenuBuilder.MakeWidget();
+	return EditorMenus->GenerateWidget(Menu);
+}
+
+void FAssetContextMenu::RegisterContextMenu(const FName MenuName)
+{
+	UEditorMenuSubsystem* EditorMenus = UEditorMenuSubsystem::Get();
+	if (!EditorMenus->IsMenuRegistered(MenuName))
+	{
+		UEditorMenu* Menu = EditorMenus->RegisterMenu(MenuName);
+		FEditorMenuSection& Section = Menu->FindOrAddSection("GetAssetActions");
+
+		Section.AddDynamicEntry("GetActions", FNewEditorMenuSectionDelegate::CreateLambda([](FEditorMenuSection& InSection)
+		{
+			UContentBrowserAssetContextMenuContext* Context = InSection.FindContext<UContentBrowserAssetContextMenuContext>();
+			if (Context && Context->CommonAssetTypeActions.IsValid())
+			{
+				Context->CommonAssetTypeActions.Pin()->GetActions(Context->SelectedObjects, InSection);
+			}
+		}));
+
+		Section.AddDynamicEntry("GetActionsLegacy", FNewEditorMenuDelegateLegacy::CreateLambda([](FMenuBuilder& MenuBuilder, UEditorMenu* InMenu)
+		{
+			UContentBrowserAssetContextMenuContext* Context = InMenu->FindContext<UContentBrowserAssetContextMenuContext>();
+			if (Context && Context->CommonAssetTypeActions.IsValid())
+			{
+				Context->CommonAssetTypeActions.Pin()->GetActions(Context->SelectedObjects, MenuBuilder);
+			}
+		}));
+	}
 }
 
 void FAssetContextMenu::SetSelectedAssets(const TArray<FAssetData>& InSelectedAssets)
@@ -194,7 +299,7 @@ void FAssetContextMenu::SetOnAssetViewRefreshRequested(const FOnAssetViewRefresh
 	OnAssetViewRefreshRequested = InOnAssetViewRefreshRequested;
 }
 
-bool FAssetContextMenu::AddImportedAssetMenuOptions(FMenuBuilder& MenuBuilder)
+bool FAssetContextMenu::AddImportedAssetMenuOptions(UEditorMenu* Menu)
 {
 	if (AreImportedAssetActionsVisible())
 	{
@@ -203,9 +308,9 @@ bool FAssetContextMenu::AddImportedAssetMenuOptions(FMenuBuilder& MenuBuilder)
 		int32 ValidSelectedAssetCount = 0;
 		GetSelectedAssetSourceFilePaths(ResolvedFilePaths, SourceFileLabels, ValidSelectedAssetCount);
 
-		MenuBuilder.BeginSection("ImportedAssetActions", LOCTEXT("ImportedAssetActionsMenuHeading", "Imported Asset"));
+		FEditorMenuSection& Section = Menu->AddSection("ImportedAssetActions", LOCTEXT("ImportedAssetActionsMenuHeading", "Imported Asset"));
 		{
-			auto CreateSubMenu = [this](FMenuBuilder& SubMenuBuilder, bool bReimportWithNewFile)
+			auto CreateSubMenu = [this](UEditorMenu* SubMenu, bool bReimportWithNewFile)
 			{
 				//Get the data, we cannot use the closure since the lambda will be call when the function scope will be gone
 				TArray<FString> ResolvedFilePaths;
@@ -214,6 +319,7 @@ bool FAssetContextMenu::AddImportedAssetMenuOptions(FMenuBuilder& MenuBuilder)
 				GetSelectedAssetSourceFilePaths(ResolvedFilePaths, SourceFileLabels, ValidSelectedAssetCount);
 				if (SourceFileLabels.Num() > 0 )
 				{
+					FEditorMenuSection& SubSection = SubMenu->AddSection("Section");
 					for (int32 SourceFileIndex = 0; SourceFileIndex < SourceFileLabels.Num(); ++SourceFileIndex)
 					{
 						FText ReimportLabel = FText::Format(LOCTEXT("ReimportNoLabel", "SourceFile {0}"), SourceFileIndex);
@@ -232,7 +338,8 @@ bool FAssetContextMenu::AddImportedAssetMenuOptions(FMenuBuilder& MenuBuilder)
 						}
 						if (bReimportWithNewFile)
 						{
-							SubMenuBuilder.AddMenuEntry(
+							SubSection.AddMenuEntry(
+								NAME_None,
 								ReimportLabel,
 								ReimportLabelTooltip,
 								FSlateIcon(FEditorStyle::GetStyleSetName(), "ContentBrowser.AssetActions.ReimportAsset"),
@@ -244,7 +351,8 @@ bool FAssetContextMenu::AddImportedAssetMenuOptions(FMenuBuilder& MenuBuilder)
 						}
 						else
 						{
-							SubMenuBuilder.AddMenuEntry(
+							SubSection.AddMenuEntry(
+								NAME_None,
 								ReimportLabel,
 								ReimportLabelTooltip,
 								FSlateIcon(FEditorStyle::GetStyleSetName(), "ContentBrowser.AssetActions.ReimportAsset"),
@@ -262,19 +370,24 @@ bool FAssetContextMenu::AddImportedAssetMenuOptions(FMenuBuilder& MenuBuilder)
 			//Reimport Menu
 			if (ValidSelectedAssetCount == 1 && SourceFileLabels.Num() > 1)
 			{
-				MenuBuilder.AddSubMenu(
+				Section.AddEntry(FEditorMenuEntry::InitSubMenu(
+					Menu->GetMenuName(),
+					"Reimport",
 					LOCTEXT("Reimport", "Reimport"),
 					LOCTEXT("ReimportEmptyTooltip", ""),
-					FNewMenuDelegate::CreateLambda(CreateSubMenu, false) );
+					FNewEditorMenuDelegate::CreateLambda(CreateSubMenu, false) ));
 				//With new file
-				MenuBuilder.AddSubMenu(
+				Section.AddEntry(FEditorMenuEntry::InitSubMenu(
+					Menu->GetMenuName(),
+					"ReimportWithNewFile",
 					LOCTEXT("ReimportWithNewFile", "Reimport With New File"),
 					LOCTEXT("ReimportEmptyTooltip", ""),
-					FNewMenuDelegate::CreateLambda(CreateSubMenu, true));
+					FNewEditorMenuDelegate::CreateLambda(CreateSubMenu, true)));
 			}
 			else
 			{
-				MenuBuilder.AddMenuEntry(
+				Section.AddMenuEntry(
+					"Reimport",
 					LOCTEXT("Reimport", "Reimport"),
 					LOCTEXT("ReimportTooltip", "Reimport the selected asset(s) from the source file on disk."),
 					FSlateIcon(FEditorStyle::GetStyleSetName(), "ContentBrowser.AssetActions.ReimportAsset"),
@@ -286,7 +399,8 @@ bool FAssetContextMenu::AddImportedAssetMenuOptions(FMenuBuilder& MenuBuilder)
 				if (ValidSelectedAssetCount == 1)
 				{
 					//With new file
-					MenuBuilder.AddMenuEntry(
+					Section.AddMenuEntry(
+						"ReimportWithNewFile",
 						LOCTEXT("ReimportWithNewFile", "Reimport With New File"),
 						LOCTEXT("ReimportWithNewFileTooltip", "Reimport the selected asset from a new source file on disk."),
 						FSlateIcon(FEditorStyle::GetStyleSetName(), "ContentBrowser.AssetActions.ReimportAsset"),
@@ -299,7 +413,8 @@ bool FAssetContextMenu::AddImportedAssetMenuOptions(FMenuBuilder& MenuBuilder)
 			}
 
 			// Show Source In Explorer
-			MenuBuilder.AddMenuEntry(
+			Section.AddMenuEntry(
+				"FindSourceFile",
 				LOCTEXT("FindSourceFile", "Open Source Location"),
 				LOCTEXT("FindSourceFileTooltip", "Opens the folder containing the source of the selected asset(s)."),
 				FSlateIcon(FEditorStyle::GetStyleSetName(), "ContentBrowser.AssetActions.OpenSourceLocation"),
@@ -310,7 +425,8 @@ bool FAssetContextMenu::AddImportedAssetMenuOptions(FMenuBuilder& MenuBuilder)
 				);
 
 			// Open In External Editor
-			MenuBuilder.AddMenuEntry(
+			Section.AddMenuEntry(
+				"OpenInExternalEditor",
 				LOCTEXT("OpenInExternalEditor", "Open In External Editor"),
 				LOCTEXT("OpenInExternalEditorTooltip", "Open the selected asset(s) in the default external editor."),
 				FSlateIcon(FEditorStyle::GetStyleSetName(), "ContentBrowser.AssetActions.OpenInExternalEditor"),
@@ -320,7 +436,6 @@ bool FAssetContextMenu::AddImportedAssetMenuOptions(FMenuBuilder& MenuBuilder)
 					)
 				);
 		}
-		MenuBuilder.EndSection();
 
 		return true;
 	}
@@ -329,15 +444,17 @@ bool FAssetContextMenu::AddImportedAssetMenuOptions(FMenuBuilder& MenuBuilder)
 	return false;
 }
 
-bool FAssetContextMenu::AddCommonMenuOptions(FMenuBuilder& MenuBuilder)
+bool FAssetContextMenu::AddCommonMenuOptions(UEditorMenu* Menu)
 {
 	int32 NumAssetItems, NumClassItems;
 	ContentBrowserUtils::CountItemTypes(SelectedAssets, NumAssetItems, NumClassItems);
 
-	MenuBuilder.BeginSection("CommonAssetActions", LOCTEXT("CommonAssetActionsMenuHeading", "Common"));
 	{
+		FEditorMenuSection& Section = Menu->AddSection("CommonAssetActions", LOCTEXT("CommonAssetActionsMenuHeading", "Common"));
+
 		// Edit
-		MenuBuilder.AddMenuEntry(
+		Section.AddMenuEntry(
+			"EditAsset",
 			LOCTEXT("EditAsset", "Edit..."),
 			LOCTEXT("EditAssetTooltip", "Opens the selected asset(s) for edit."),
 			FSlateIcon(FEditorStyle::GetStyleSetName(), "ContentBrowser.AssetActions.Edit"),
@@ -348,83 +465,84 @@ bool FAssetContextMenu::AddCommonMenuOptions(FMenuBuilder& MenuBuilder)
 		if (NumAssetItems > 0)
 		{
 			// Rename
-			MenuBuilder.AddMenuEntry(FGenericCommands::Get().Rename, NAME_None,
+			Section.AddMenuEntry(FGenericCommands::Get().Rename,
 				LOCTEXT("Rename", "Rename"),
 				LOCTEXT("RenameTooltip", "Rename the selected asset."),
 				FSlateIcon(FEditorStyle::GetStyleSetName(), "ContentBrowser.AssetActions.Rename")
 				);
 
 			// Duplicate
-			MenuBuilder.AddMenuEntry(FGenericCommands::Get().Duplicate, NAME_None,
+			Section.AddMenuEntry(FGenericCommands::Get().Duplicate,
 				LOCTEXT("Duplicate", "Duplicate"),
 				LOCTEXT("DuplicateTooltip", "Create a copy of the selected asset(s)."),
 				FSlateIcon(FEditorStyle::GetStyleSetName(), "ContentBrowser.AssetActions.Duplicate")
 				);
 
 			// Save
-			MenuBuilder.AddMenuEntry(FContentBrowserCommands::Get().SaveSelectedAsset, NAME_None,
+			Section.AddMenuEntry(FContentBrowserCommands::Get().SaveSelectedAsset,
 				LOCTEXT("SaveAsset", "Save"),
 				LOCTEXT("SaveAssetTooltip", "Saves the asset to file."),
 				FSlateIcon(FEditorStyle::GetStyleSetName(), "Level.SaveIcon16x")
 				);
 
 			// Delete
-			MenuBuilder.AddMenuEntry(FGenericCommands::Get().Delete, NAME_None,
+			Section.AddMenuEntry(FGenericCommands::Get().Delete,
 				LOCTEXT("Delete", "Delete"),
 				LOCTEXT("DeleteTooltip", "Delete the selected assets."),
 				FSlateIcon(FEditorStyle::GetStyleSetName(), "ContentBrowser.AssetActions.Delete")
 				);
 
 			// Asset Actions sub-menu
-			MenuBuilder.AddSubMenu(
+			Section.AddEntry(FEditorMenuEntry::InitSubMenu(
+				Menu->GetMenuName(),
+				"AssetActionsSubMenu",
 				LOCTEXT("AssetActionsSubMenuLabel", "Asset Actions"),
 				LOCTEXT("AssetActionsSubMenuToolTip", "Other asset actions"),
-				FNewMenuDelegate::CreateSP(this, &FAssetContextMenu::MakeAssetActionsSubMenu),
+				FNewEditorMenuDelegate::CreateSP(this, &FAssetContextMenu::MakeAssetActionsSubMenu),
 				FUIAction(
 					FExecuteAction(),
 					FCanExecuteAction::CreateSP( this, &FAssetContextMenu::CanExecuteAssetActions )
 					),
-				NAME_None,
 				EUserInterfaceActionType::Button,
 				false, 
 				FSlateIcon(FEditorStyle::GetStyleSetName(), "ContentBrowser.AssetActions")
-				);
+				));
 
 			if (NumClassItems == 0)
 			{
 				// Asset Localization sub-menu
-				MenuBuilder.AddSubMenu(
+				Section.AddEntry(FEditorMenuEntry::InitSubMenu(
+					Menu->GetMenuName(),
+					"LocalizationSubMenu",
 					LOCTEXT("LocalizationSubMenuLabel", "Asset Localization"),
 					LOCTEXT("LocalizationSubMenuToolTip", "Manage the localization of this asset"),
-					FNewMenuDelegate::CreateSP(this, &FAssetContextMenu::MakeAssetLocalizationSubMenu),
+					FNewEditorMenuDelegate::CreateSP(this, &FAssetContextMenu::MakeAssetLocalizationSubMenu),
 					FUIAction(),
-					NAME_None,
 					EUserInterfaceActionType::Button,
 					false,
 					FSlateIcon(FEditorStyle::GetStyleSetName(), "ContentBrowser.AssetLocalization")
-					);
+					));
 			}
 		}
 	}
-	MenuBuilder.EndSection();
 
 	return true;
 }
 
-void FAssetContextMenu::AddExploreMenuOptions(FMenuBuilder& MenuBuilder)
+void FAssetContextMenu::AddExploreMenuOptions(UEditorMenu* Menu)
 {
-	MenuBuilder.BeginSection("AssetContextExploreMenuOptions", LOCTEXT("AssetContextExploreMenuOptionsHeading", "Explore"));
+	FEditorMenuSection& Section = Menu->AddSection("AssetContextExploreMenuOptions", LOCTEXT("AssetContextExploreMenuOptionsHeading", "Explore"));
 	{
 		// Find in Content Browser
-		MenuBuilder.AddMenuEntry(
+		Section.AddMenuEntry(
 			FGlobalEditorCommonCommands::Get().FindInContentBrowser, 
-			NAME_None, 
 			LOCTEXT("ShowInFolderView", "Show in Folder View"),
 			LOCTEXT("ShowInFolderViewTooltip", "Selects the folder that contains this asset in the Content Browser Sources Panel.")
 			);
 
 		// Find in Explorer
-		MenuBuilder.AddMenuEntry(
+		Section.AddMenuEntry(
+			"FindInExplorer",
 			ContentBrowserUtils::GetExploreFolderText(),
 			LOCTEXT("FindInExplorerTooltip", "Finds this asset on disk"),
 			FSlateIcon(FEditorStyle::GetStyleSetName(), "SystemWideCommands.FindInContentBrowser"),
@@ -434,53 +552,60 @@ void FAssetContextMenu::AddExploreMenuOptions(FMenuBuilder& MenuBuilder)
 				)
 			);
 	}
-	MenuBuilder.EndSection();
 }
 
-void FAssetContextMenu::MakeAssetActionsSubMenu(FMenuBuilder& MenuBuilder)
+void FAssetContextMenu::MakeAssetActionsSubMenu(UEditorMenu* Menu)
 {	
-	// Create BP Using This
-	MenuBuilder.AddMenuEntry(
-		LOCTEXT("CreateBlueprintUsing", "Create Blueprint Using This..."),
-		LOCTEXT("CreateBlueprintUsingTooltip", "Create a new Blueprint and add this asset to it"),
-		FSlateIcon(FEditorStyle::GetStyleSetName(), "LevelEditor.CreateClassBlueprint"),
-		FUIAction(
-			FExecuteAction::CreateSP(this, &FAssetContextMenu::ExecuteCreateBlueprintUsing),
-			FCanExecuteAction::CreateSP(this, &FAssetContextMenu::CanExecuteCreateBlueprintUsing)
-		)
-	);
-
-	// Capture Thumbnail
-	FAssetToolsModule& AssetToolsModule = FModuleManager::LoadModuleChecked<FAssetToolsModule>("AssetTools");
-	if (SelectedAssets.Num() == 1 && AssetToolsModule.Get().AssetUsesGenericThumbnail(SelectedAssets[0]))
 	{
-		MenuBuilder.AddMenuEntry(
-			LOCTEXT("CaptureThumbnail", "Capture Thumbnail"),
-			LOCTEXT("CaptureThumbnailTooltip", "Captures a thumbnail from the active viewport."),
-			FSlateIcon(FEditorStyle::GetStyleSetName(), "ContentBrowser.AssetActions.CreateThumbnail"),
+		FEditorMenuSection& Section = Menu->AddSection("AssetActionsSection");
+
+		// Create BP Using This
+		Section.AddMenuEntry(
+			"CreateBlueprintUsing",
+			LOCTEXT("CreateBlueprintUsing", "Create Blueprint Using This..."),
+			LOCTEXT("CreateBlueprintUsingTooltip", "Create a new Blueprint and add this asset to it"),
+			FSlateIcon(FEditorStyle::GetStyleSetName(), "LevelEditor.CreateClassBlueprint"),
 			FUIAction(
-				FExecuteAction::CreateSP(this, &FAssetContextMenu::ExecuteCaptureThumbnail),
-				FCanExecuteAction::CreateSP(this, &FAssetContextMenu::CanExecuteCaptureThumbnail)
+				FExecuteAction::CreateSP(this, &FAssetContextMenu::ExecuteCreateBlueprintUsing),
+				FCanExecuteAction::CreateSP(this, &FAssetContextMenu::CanExecuteCreateBlueprintUsing)
 			)
 		);
-	}
 
-	// Clear Thumbnail
-	if (CanClearCustomThumbnails())
-	{
-		MenuBuilder.AddMenuEntry(
-			LOCTEXT("ClearCustomThumbnail", "Clear Thumbnail"),
-			LOCTEXT("ClearCustomThumbnailTooltip", "Clears all custom thumbnails for selected assets."),
-			FSlateIcon(FEditorStyle::GetStyleSetName(), "ContentBrowser.AssetActions.DeleteThumbnail"),
-			FUIAction( FExecuteAction::CreateSP(this, &FAssetContextMenu::ExecuteClearThumbnail) ) 
+		// Capture Thumbnail
+		FAssetToolsModule& AssetToolsModule = FModuleManager::LoadModuleChecked<FAssetToolsModule>("AssetTools");
+		if (SelectedAssets.Num() == 1 && AssetToolsModule.Get().AssetUsesGenericThumbnail(SelectedAssets[0]))
+		{
+			Section.AddMenuEntry(
+				"CaptureThumbnail",
+				LOCTEXT("CaptureThumbnail", "Capture Thumbnail"),
+				LOCTEXT("CaptureThumbnailTooltip", "Captures a thumbnail from the active viewport."),
+				FSlateIcon(FEditorStyle::GetStyleSetName(), "ContentBrowser.AssetActions.CreateThumbnail"),
+				FUIAction(
+					FExecuteAction::CreateSP(this, &FAssetContextMenu::ExecuteCaptureThumbnail),
+					FCanExecuteAction::CreateSP(this, &FAssetContextMenu::CanExecuteCaptureThumbnail)
+				)
 			);
+		}
+
+		// Clear Thumbnail
+		if (CanClearCustomThumbnails())
+		{
+			Section.AddMenuEntry(
+				"ClearCustomThumbnail",
+				LOCTEXT("ClearCustomThumbnail", "Clear Thumbnail"),
+				LOCTEXT("ClearCustomThumbnailTooltip", "Clears all custom thumbnails for selected assets."),
+				FSlateIcon(FEditorStyle::GetStyleSetName(), "ContentBrowser.AssetActions.DeleteThumbnail"),
+				FUIAction(FExecuteAction::CreateSP(this, &FAssetContextMenu::ExecuteClearThumbnail))
+			);
+		}
 	}
 
 	// FIND ACTIONS
-	MenuBuilder.BeginSection("AssetContextFindActions", LOCTEXT("AssetContextFindActionsMenuHeading", "Find"));
 	{
+		FEditorMenuSection& Section = Menu->AddSection("AssetContextFindActions", LOCTEXT("AssetContextFindActionsMenuHeading", "Find"));
 		// Select Actors Using This Asset
-		MenuBuilder.AddMenuEntry(
+		Section.AddMenuEntry(
+			"FindAssetInWorld",
 			LOCTEXT("FindAssetInWorld", "Select Actors Using This Asset"),
 			LOCTEXT("FindAssetInWorldTooltip", "Selects all actors referencing this asset."),
 			FSlateIcon(),
@@ -490,11 +615,10 @@ void FAssetContextMenu::MakeAssetActionsSubMenu(FMenuBuilder& MenuBuilder)
 				)
 			);
 	}
-	MenuBuilder.EndSection();
 
 	// MOVE ACTIONS
-	MenuBuilder.BeginSection("AssetContextMoveActions", LOCTEXT("AssetContextMoveActionsMenuHeading", "Move"));
 	{
+		FEditorMenuSection& Section = Menu->AddSection("AssetContextMoveActions", LOCTEXT("AssetContextMoveActionsMenuHeading", "Move"));
 		bool bHasExportableAssets = false;
 		for (const FAssetData& AssetData : SelectedAssets)
 		{
@@ -513,7 +637,8 @@ void FAssetContextMenu::MakeAssetActionsSubMenu(FMenuBuilder& MenuBuilder)
 		if (bHasExportableAssets)
 		{
 			// Export
-			MenuBuilder.AddMenuEntry(
+			Section.AddMenuEntry(
+				"Export",
 				LOCTEXT("Export", "Export..."),
 				LOCTEXT("ExportTooltip", "Export the selected assets to file."),
 				FSlateIcon(),
@@ -523,7 +648,8 @@ void FAssetContextMenu::MakeAssetActionsSubMenu(FMenuBuilder& MenuBuilder)
 			// Bulk Export
 			if (SelectedAssets.Num() > 1)
 			{
-				MenuBuilder.AddMenuEntry(
+				Section.AddMenuEntry(
+					"BulkExport",
 					LOCTEXT("BulkExport", "Bulk Export..."),
 					LOCTEXT("BulkExportTooltip", "Export the selected assets to file in the selected directory"),
 					FSlateIcon(),
@@ -533,20 +659,22 @@ void FAssetContextMenu::MakeAssetActionsSubMenu(FMenuBuilder& MenuBuilder)
 		}
 
 		// Migrate
-		MenuBuilder.AddMenuEntry(
+		Section.AddMenuEntry(
+			"MigrateAsset",
 			LOCTEXT("MigrateAsset", "Migrate..."),
 			LOCTEXT("MigrateAssetTooltip", "Copies all selected assets and their dependencies to another project"),
 			FSlateIcon(),
 			FUIAction( FExecuteAction::CreateSP( this, &FAssetContextMenu::ExecuteMigrateAsset ) )
 			);
 	}
-	MenuBuilder.EndSection();
 
 	// ADVANCED ACTIONS
-	MenuBuilder.BeginSection("AssetContextAdvancedActions", LOCTEXT("AssetContextAdvancedActionsMenuHeading", "Advanced"));
 	{
+		FEditorMenuSection& Section = Menu->AddSection("AssetContextAdvancedActions", LOCTEXT("AssetContextAdvancedActionsMenuHeading", "Advanced"));
+
 		// Reload
-		MenuBuilder.AddMenuEntry(
+		Section.AddMenuEntry(
+			"Reload",
 			LOCTEXT("Reload", "Reload"),
 			LOCTEXT("ReloadTooltip", "Reload the selected assets from their file on disk."),
 			FSlateIcon(),
@@ -559,7 +687,8 @@ void FAssetContextMenu::MakeAssetActionsSubMenu(FMenuBuilder& MenuBuilder)
 		// Replace References
 		if (CanExecuteConsolidate())
 		{
-			MenuBuilder.AddMenuEntry(
+			Section.AddMenuEntry(
+				"ReplaceReferences",
 				LOCTEXT("ReplaceReferences", "Replace References"),
 				LOCTEXT("ConsolidateTooltip", "Replace references to the selected assets."),
 				FSlateIcon(),
@@ -587,7 +716,8 @@ void FAssetContextMenu::MakeAssetActionsSubMenu(FMenuBuilder& MenuBuilder)
 			DynamicTooltipGetter.BindSP(this, &FAssetContextMenu::GetExecutePropertyMatrixTooltip);
 			TAttribute<FText> DynamicTooltipAttribute = TAttribute<FText>::Create(DynamicTooltipGetter);
 
-			MenuBuilder.AddMenuEntry(
+			Section.AddMenuEntry(
+				"PropertyMatrix",
 				LOCTEXT("PropertyMatrix", "Bulk Edit via Property Matrix..."),
 				DynamicTooltipAttribute,
 				FSlateIcon(),
@@ -599,7 +729,8 @@ void FAssetContextMenu::MakeAssetActionsSubMenu(FMenuBuilder& MenuBuilder)
 		}
 
 		// Create Metadata menu
-		MenuBuilder.AddMenuEntry(
+		Section.AddMenuEntry(
+			"ShowAssetMetaData",
 			LOCTEXT("ShowAssetMetaData", "Show Metadata"),
 			LOCTEXT("ShowAssetMetaDataTooltip", "Show the asset metadata dialog."),
 			FSlateIcon(),
@@ -612,20 +743,24 @@ void FAssetContextMenu::MakeAssetActionsSubMenu(FMenuBuilder& MenuBuilder)
 		// Chunk actions
 		if (GetDefault<UEditorExperimentalSettings>()->bContextMenuChunkAssignments)
 		{
-			MenuBuilder.AddMenuEntry(
+			Section.AddMenuEntry(
+				"AssignAssetChunk",
 				LOCTEXT("AssignAssetChunk", "Assign to Chunk..."),
 				LOCTEXT("AssignAssetChunkTooltip", "Assign this asset to a specific Chunk"),
 				FSlateIcon(),
 				FUIAction( FExecuteAction::CreateSP(this, &FAssetContextMenu::ExecuteAssignChunkID) )
 				);
 
-			MenuBuilder.AddSubMenu(
+			Section.AddEntry(FEditorMenuEntry::InitSubMenu(
+				Menu->GetMenuName(),
+				"RemoveAssetFromChunk",
 				LOCTEXT("RemoveAssetFromChunk", "Remove from Chunk..."),
 				LOCTEXT("RemoveAssetFromChunkTooltip", "Removed an asset from a Chunk it's assigned to."),
-				FNewMenuDelegate::CreateRaw(this, &FAssetContextMenu::MakeChunkIDListMenu)
-				);
+				FNewEditorMenuDelegate::CreateRaw(this, &FAssetContextMenu::MakeChunkIDListMenu)
+				));
 
-			MenuBuilder.AddMenuEntry(
+			Section.AddMenuEntry(
+				"RemoveAllChunkAssignments",
 				LOCTEXT("RemoveAllChunkAssignments", "Remove from all Chunks"),
 				LOCTEXT("RemoveAllChunkAssignmentsTooltip", "Removed an asset from all Chunks it's assigned to."),
 				FSlateIcon(),
@@ -633,20 +768,19 @@ void FAssetContextMenu::MakeAssetActionsSubMenu(FMenuBuilder& MenuBuilder)
 				);
 		}
 	}
-	MenuBuilder.EndSection();
 
 	if (GetDefault<UEditorExperimentalSettings>()->bTextAssetFormatSupport)
 	{
-		MenuBuilder.BeginSection("AssetContextTextAssetFormatActions", LOCTEXT("AssetContextTextAssetFormatActionsHeading", "Text Assets"));
+		FEditorMenuSection& FormatActionsSection = Menu->AddSection("AssetContextTextAssetFormatActions", LOCTEXT("AssetContextTextAssetFormatActionsHeading", "Text Assets"));
 		{
-			MenuBuilder.AddMenuEntry(
+			FormatActionsSection.AddMenuEntry(
+				"ExportToTextFormat",
 				LOCTEXT("ExportToTextFormat", "Export to text format"),
 				LOCTEXT("ExportToTextFormatTooltip", "Exports the selected asset(s) to the experimental text asset format"),
 				FSlateIcon(),
 				FUIAction(FExecuteAction::CreateSP(this, &FAssetContextMenu::ExportSelectedAssetsToText))
 			);
 		}
-		MenuBuilder.EndSection();
 	}
 }
 
@@ -688,7 +822,7 @@ bool FAssetContextMenu::CanExecuteAssetActions() const
 	return !bAtLeastOneClassSelected;
 }
 
-void FAssetContextMenu::MakeAssetLocalizationSubMenu(FMenuBuilder& MenuBuilder)
+void FAssetContextMenu::MakeAssetLocalizationSubMenu(UEditorMenu* Menu)
 {
 	TArray<FCultureRef> CurrentCultures;
 
@@ -818,13 +952,14 @@ void FAssetContextMenu::MakeAssetLocalizationSubMenu(FMenuBuilder& MenuBuilder)
 #if USE_STABLE_LOCALIZATION_KEYS
 	// Add the Localization ID options
 	{
-		MenuBuilder.BeginSection(NAME_None, LOCTEXT("LocalizationIdHeading", "Localization ID"));
+		FEditorMenuSection& Section = Menu->AddSection("LocalizationId", LOCTEXT("LocalizationIdHeading", "Localization ID"));
 		{
 			// Show the localization ID if we have a single asset selected
 			if (SelectedAssets.Num() == 1)
 			{
 				const FString LocalizationId = TextNamespaceUtil::GetPackageNamespace(SelectedAssets[0].GetAsset());
-				MenuBuilder.AddMenuEntry(
+				Section.AddMenuEntry(
+					"CopyLocalizationId",
 					FText::Format(LOCTEXT("CopyLocalizationIdFmt", "ID: {0}"), LocalizationId.IsEmpty() ? LOCTEXT("EmptyLocalizationId", "None") : FText::FromString(LocalizationId)),
 					LOCTEXT("CopyLocalizationIdTooltip", "Copy the localization ID to the clipboard."),
 					FSlateIcon(),
@@ -833,14 +968,14 @@ void FAssetContextMenu::MakeAssetLocalizationSubMenu(FMenuBuilder& MenuBuilder)
 			}
 
 			// Always show the reset localization ID option
-			MenuBuilder.AddMenuEntry(
+			Section.AddMenuEntry(
+				"ResetLocalizationId",
 				LOCTEXT("ResetLocalizationId", "Reset Localization ID"),
 				LOCTEXT("ResetLocalizationIdTooltip", "Reset the localization ID. Note: This will re-key all the text within this asset."),
 				FSlateIcon(),
 				FUIAction(FExecuteAction::CreateSP(this, &FAssetContextMenu::ExecuteResetLocalizationId))
 				);
 		}
-		MenuBuilder.EndSection();
 	}
 #endif // USE_STABLE_LOCALIZATION_KEYS
 
@@ -850,57 +985,59 @@ void FAssetContextMenu::MakeAssetLocalizationSubMenu(FMenuBuilder& MenuBuilder)
 		FString PackageFilename;
 		if (FPackageName::DoesPackageExist(SelectedAssets[0].PackageName.ToString(), nullptr, &PackageFilename))
 		{
-			MenuBuilder.BeginSection(NAME_None, LOCTEXT("LocalizationCacheHeading", "Localization Cache"));
+			FEditorMenuSection& Section = Menu->AddSection("LocalizationCache", LOCTEXT("LocalizationCacheHeading", "Localization Cache"));
 			{
 				// Always show the reset localization ID option
-				MenuBuilder.AddMenuEntry(
+				Section.AddMenuEntry(
+					"ShowLocalizationCache",
 					LOCTEXT("ShowLocalizationCache", "Show Localization Cache"),
 					LOCTEXT("ShowLocalizationCacheTooltip", "Show the cached list of localized texts stored in the package header."),
 					FSlateIcon(),
 					FUIAction(FExecuteAction::CreateSP(this, &FAssetContextMenu::ExecuteShowLocalizationCache, PackageFilename))
 					);
 			}
-			MenuBuilder.EndSection();
 		}
 	}
 
 	// If we found source assets for localized assets, then we can show the Source Asset options
 	if (SourceAssetsState.CurrentAssets.Num() > 0)
 	{
-		MenuBuilder.BeginSection(NAME_None, LOCTEXT("ManageSourceAssetHeading", "Manage Source Asset"));
+		FEditorMenuSection& Section = Menu->AddSection("ManageSourceAsset", LOCTEXT("ManageSourceAssetHeading", "Manage Source Asset"));
 		{
-			MenuBuilder.AddMenuEntry(
+			Section.AddMenuEntry(
+				"ShowSourceAsset",
 				LOCTEXT("ShowSourceAsset", "Show Source Asset"),
 				LOCTEXT("ShowSourceAssetTooltip", "Show the source asset in the Content Browser."),
 				FSlateIcon(FEditorStyle::GetStyleSetName(), "SystemWideCommands.FindInContentBrowser"),
 				FUIAction(FExecuteAction::CreateSP(this, &FAssetContextMenu::ExecuteFindInAssetTree, SourceAssetsState.CurrentAssets.Array()))
 				);
 
-			MenuBuilder.AddMenuEntry(
+			Section.AddMenuEntry(
+				"EditSourceAsset",
 				LOCTEXT("EditSourceAsset", "Edit Source Asset"),
 				LOCTEXT("EditSourceAssetTooltip", "Edit the source asset."),
 				FSlateIcon(FEditorStyle::GetStyleSetName(), "ContentBrowser.AssetActions.Edit"),
 				FUIAction(FExecuteAction::CreateSP(this, &FAssetContextMenu::ExecuteOpenEditorsForAssets, SourceAssetsState.CurrentAssets.Array()))
 				);
 		}
-		MenuBuilder.EndSection();
 	}
 
 	// If we currently have source assets selected, then we can show the Localized Asset options
 	if (SourceAssetsState.SelectedAssets.Num() > 0)
 	{
-		MenuBuilder.BeginSection(NAME_None, LOCTEXT("ManageLocalizedAssetHeading", "Manage Localized Asset"));
+		FEditorMenuSection& Section = Menu->AddSection("ManageLocalizedAsset", LOCTEXT("ManageLocalizedAssetHeading", "Manage Localized Asset"));
 		{
-			MenuBuilder.AddSubMenu(
+			Section.AddEntry(FEditorMenuEntry::InitSubMenu(
+				Menu->GetMenuName(),
+				"CreateLocalizedAsset",
 				LOCTEXT("CreateLocalizedAsset", "Create Localized Asset"),
 				LOCTEXT("CreateLocalizedAssetTooltip", "Create a new localized asset."),
-				FNewMenuDelegate::CreateSP(this, &FAssetContextMenu::MakeCreateLocalizedAssetSubMenu, SourceAssetsState.SelectedAssets, LocalizedAssetsState),
+				FNewEditorMenuDelegate::CreateSP(this, &FAssetContextMenu::MakeCreateLocalizedAssetSubMenu, SourceAssetsState.SelectedAssets, LocalizedAssetsState),
 				FUIAction(),
-				NAME_None,
 				EUserInterfaceActionType::Button,
 				false,
 				FSlateIcon(FEditorStyle::GetStyleSetName(), "ContentBrowser.AssetActions.Duplicate")
-				);
+				));
 
 			int32 NumLocalizedAssets = 0;
 			for (const FLocalizedAssetsState& LocalizedAssetsStateForCulture : LocalizedAssetsState)
@@ -910,41 +1047,45 @@ void FAssetContextMenu::MakeAssetLocalizationSubMenu(FMenuBuilder& MenuBuilder)
 
 			if (NumLocalizedAssets > 0)
 			{
-				MenuBuilder.AddSubMenu(
+				Section.AddEntry(FEditorMenuEntry::InitSubMenu(
+					Menu->GetMenuName(),
+					"ShowLocalizedAsset",
 					LOCTEXT("ShowLocalizedAsset", "Show Localized Asset"),
 					LOCTEXT("ShowLocalizedAssetTooltip", "Show the localized asset in the Content Browser."),
-					FNewMenuDelegate::CreateSP(this, &FAssetContextMenu::MakeShowLocalizedAssetSubMenu, LocalizedAssetsState),
+					FNewEditorMenuDelegate::CreateSP(this, &FAssetContextMenu::MakeShowLocalizedAssetSubMenu, LocalizedAssetsState),
 					FUIAction(),
-					NAME_None,
 					EUserInterfaceActionType::Button,
 					false,
 					FSlateIcon(FEditorStyle::GetStyleSetName(), "SystemWideCommands.FindInContentBrowser")
-					);
+					));
 
-				MenuBuilder.AddSubMenu(
+				Section.AddEntry(FEditorMenuEntry::InitSubMenu(
+					Menu->GetMenuName(),
+					"EditLocalizedAsset",
 					LOCTEXT("EditLocalizedAsset", "Edit Localized Asset"),
 					LOCTEXT("EditLocalizedAssetTooltip", "Edit the localized asset."),
-					FNewMenuDelegate::CreateSP(this, &FAssetContextMenu::MakeEditLocalizedAssetSubMenu, LocalizedAssetsState),
+					FNewEditorMenuDelegate::CreateSP(this, &FAssetContextMenu::MakeEditLocalizedAssetSubMenu, LocalizedAssetsState),
 					FUIAction(),
-					NAME_None,
 					EUserInterfaceActionType::Button,
 					false,
 					FSlateIcon(FEditorStyle::GetStyleSetName(), "ContentBrowser.AssetActions.Edit")
-					);
+					));
 			}
 		}
-		MenuBuilder.EndSection();
 	}
 }
 
-void FAssetContextMenu::MakeCreateLocalizedAssetSubMenu(FMenuBuilder& MenuBuilder, TSet<FName> InSelectedSourceAssets, TArray<FLocalizedAssetsState> InLocalizedAssetsState)
+void FAssetContextMenu::MakeCreateLocalizedAssetSubMenu(UEditorMenu* Menu, TSet<FName> InSelectedSourceAssets, TArray<FLocalizedAssetsState> InLocalizedAssetsState)
 {
+	FEditorMenuSection& Section = Menu->AddSection("Section");
+
 	for (const FLocalizedAssetsState& LocalizedAssetsStateForCulture : InLocalizedAssetsState)
 	{
 		// If we have less localized assets than we have selected source assets, then we'll have some assets to create localized variants of
 		if (LocalizedAssetsStateForCulture.CurrentAssets.Num() < InSelectedSourceAssets.Num())
 		{
-			MenuBuilder.AddMenuEntry(
+			Section.AddMenuEntry(
+				NAME_None,
 				FText::FromString(LocalizedAssetsStateForCulture.Culture->GetDisplayName()),
 				FText::GetEmpty(),
 				FSlateIcon(),
@@ -954,13 +1095,16 @@ void FAssetContextMenu::MakeCreateLocalizedAssetSubMenu(FMenuBuilder& MenuBuilde
 	}
 }
 
-void FAssetContextMenu::MakeShowLocalizedAssetSubMenu(FMenuBuilder& MenuBuilder, TArray<FLocalizedAssetsState> InLocalizedAssetsState)
+void FAssetContextMenu::MakeShowLocalizedAssetSubMenu(UEditorMenu* Menu, TArray<FLocalizedAssetsState> InLocalizedAssetsState)
 {
+	FEditorMenuSection& Section = Menu->AddSection("Section");
+
 	for (const FLocalizedAssetsState& LocalizedAssetsStateForCulture : InLocalizedAssetsState)
 	{
 		if (LocalizedAssetsStateForCulture.CurrentAssets.Num() > 0)
 		{
-			MenuBuilder.AddMenuEntry(
+			Section.AddMenuEntry(
+				NAME_None,
 				FText::FromString(LocalizedAssetsStateForCulture.Culture->GetDisplayName()),
 				FText::GetEmpty(),
 				FSlateIcon(),
@@ -970,13 +1114,16 @@ void FAssetContextMenu::MakeShowLocalizedAssetSubMenu(FMenuBuilder& MenuBuilder,
 	}
 }
 
-void FAssetContextMenu::MakeEditLocalizedAssetSubMenu(FMenuBuilder& MenuBuilder, TArray<FLocalizedAssetsState> InLocalizedAssetsState)
+void FAssetContextMenu::MakeEditLocalizedAssetSubMenu(UEditorMenu* Menu, TArray<FLocalizedAssetsState> InLocalizedAssetsState)
 {
+	FEditorMenuSection& Section = Menu->AddSection("Section");
+
 	for (const FLocalizedAssetsState& LocalizedAssetsStateForCulture : InLocalizedAssetsState)
 	{
 		if (LocalizedAssetsStateForCulture.CurrentAssets.Num() > 0)
 		{
-			MenuBuilder.AddMenuEntry(
+			Section.AddMenuEntry(
+				NAME_None,
 				FText::FromString(LocalizedAssetsStateForCulture.Culture->GetDisplayName()),
 				FText::GetEmpty(),
 				FSlateIcon(),
@@ -1051,23 +1198,23 @@ void FAssetContextMenu::ExecuteOpenEditorsForAssets(TArray<FName> InAssets)
 	FAssetEditorManager::Get().OpenEditorsForAssets(InAssets);
 }
 
-bool FAssetContextMenu::AddReferenceMenuOptions(FMenuBuilder& MenuBuilder)
+bool FAssetContextMenu::AddReferenceMenuOptions(UEditorMenu* Menu)
 {
-	MenuBuilder.BeginSection("AssetContextReferences", LOCTEXT("ReferencesMenuHeading", "References"));
+	FEditorMenuSection& Section = Menu->AddSection("AssetContextReferences", LOCTEXT("ReferencesMenuHeading", "References"));
 	{
-		MenuBuilder.AddMenuEntry(
+		Section.AddMenuEntry(
+			"CopyReference",
 			LOCTEXT("CopyReference", "Copy Reference"),
 			LOCTEXT("CopyReferenceTooltip", "Copies reference paths for the selected assets to the clipboard."),
 			FSlateIcon(),
 			FUIAction( FExecuteAction::CreateSP( this, &FAssetContextMenu::ExecuteCopyReference ) )
 			);
 	}
-	MenuBuilder.EndSection();
 
 	return true;
 }
 
-bool FAssetContextMenu::AddDocumentationMenuOptions(FMenuBuilder& MenuBuilder)
+bool FAssetContextMenu::AddDocumentationMenuOptions(UEditorMenu* Menu)
 {
 	bool bAddedOption = false;
 
@@ -1105,16 +1252,16 @@ bool FAssetContextMenu::AddDocumentationMenuOptions(FMenuBuilder& MenuBuilder)
 
 				const FString CodeFileName = FPaths::GetCleanFilename( *ClassHeaderPath );
 
-				MenuBuilder.BeginSection( "AssetCode"/*, LOCTEXT("AssetCodeHeading", "C++")*/ );
+				FEditorMenuSection& Section = Menu->AddSection( "AssetCode"/*, LOCTEXT("AssetCodeHeading", "C++")*/ );
 				{
-					MenuBuilder.AddMenuEntry(
+					Section.AddMenuEntry(
+						"GoToCodeForAsset",
 						FText::Format( LOCTEXT("GoToCodeForAsset", "Open {0}"), FText::FromString( CodeFileName ) ),
 						FText::Format( LOCTEXT("GoToCodeForAsset_ToolTip", "Opens the header file for this asset ({0}) in a code editing program"), FText::FromString( CodeFileName ) ),
 						FSlateIcon(FEditorStyle::GetStyleSetName(), "ContentBrowser.AssetActions.GoToCodeForAsset"),
 						FUIAction( FExecuteAction::CreateSP( this, &FAssetContextMenu::ExecuteGoToCodeForAsset, SelectedClass ) )
 						);
 				}
-				MenuBuilder.EndSection();
 			}
 		}
 
@@ -1123,13 +1270,14 @@ bool FAssetContextMenu::AddDocumentationMenuOptions(FMenuBuilder& MenuBuilder)
 		{
 			bAddedOption = true;
 
-			MenuBuilder.BeginSection( "AssetDocumentation"/*, LOCTEXT("AseetDocsHeading", "Documentation")*/ );
+			FEditorMenuSection& Section = Menu->AddSection( "AssetDocumentation"/*, LOCTEXT("AseetDocsHeading", "Documentation")*/ );
 			{
 					if (bIsBlueprint)
 					{
 						if (!DocumentationLink.IsEmpty())
 						{
-							MenuBuilder.AddMenuEntry(
+							Section.AddMenuEntry(
+								"GoToDocsForAssetWithClass",
 								FText::Format( LOCTEXT("GoToDocsForAssetWithClass", "View Documentation - {0}"), SelectedClass->GetDisplayNameText() ),
 								FText::Format( LOCTEXT("GoToDocsForAssetWithClass_ToolTip", "Click to open documentation for {0}"), SelectedClass->GetDisplayNameText() ),
 								FSlateIcon(FEditorStyle::GetStyleSetName(), "HelpIcon.Hovered" ),
@@ -1144,7 +1292,8 @@ bool FAssetContextMenu::AddDocumentationMenuOptions(FMenuBuilder& MenuBuilder)
 						switch (BlueprintType)
 						{
 						case BPTYPE_FunctionLibrary:
-							MenuBuilder.AddMenuEntry(
+							Section.AddMenuEntry(
+								"GoToDocsForMacroBlueprint",
 								LOCTEXT("GoToDocsForMacroBlueprint", "View Documentation - Function Library"),
 								LOCTEXT("GoToDocsForMacroBlueprint_ToolTip", "Click to open documentation on blueprint function libraries"),
 								FSlateIcon(FEditorStyle::GetStyleSetName(), "HelpIcon.Hovered" ),
@@ -1152,7 +1301,8 @@ bool FAssetContextMenu::AddDocumentationMenuOptions(FMenuBuilder& MenuBuilder)
 								);
 							break;
 						case BPTYPE_Interface:
-							MenuBuilder.AddMenuEntry(
+							Section.AddMenuEntry(
+								"GoToDocsForInterfaceBlueprint",
 								LOCTEXT("GoToDocsForInterfaceBlueprint", "View Documentation - Interface"),
 								LOCTEXT("GoToDocsForInterfaceBlueprint_ToolTip", "Click to open documentation on blueprint interfaces"),
 								FSlateIcon(FEditorStyle::GetStyleSetName(), "HelpIcon.Hovered" ),
@@ -1160,7 +1310,8 @@ bool FAssetContextMenu::AddDocumentationMenuOptions(FMenuBuilder& MenuBuilder)
 								);
 							break;
 						case BPTYPE_MacroLibrary:
-							MenuBuilder.AddMenuEntry(
+							Section.AddMenuEntry(
+								"GoToDocsForMacroLibrary",
 								LOCTEXT("GoToDocsForMacroLibrary", "View Documentation - Macro"),
 								LOCTEXT("GoToDocsForMacroLibrary_ToolTip", "Click to open documentation on blueprint macros"),
 								FSlateIcon(FEditorStyle::GetStyleSetName(), "HelpIcon.Hovered" ),
@@ -1168,7 +1319,8 @@ bool FAssetContextMenu::AddDocumentationMenuOptions(FMenuBuilder& MenuBuilder)
 								);
 							break;
 						default:
-							MenuBuilder.AddMenuEntry(
+							Section.AddMenuEntry(
+								"GoToDocsForBlueprint",
 								LOCTEXT("GoToDocsForBlueprint", "View Documentation - Blueprint"),
 								LOCTEXT("GoToDocsForBlueprint_ToolTip", "Click to open documentation on blueprints"),
 								FSlateIcon(FEditorStyle::GetStyleSetName(), "HelpIcon.Hovered" ),
@@ -1178,7 +1330,8 @@ bool FAssetContextMenu::AddDocumentationMenuOptions(FMenuBuilder& MenuBuilder)
 					}
 					else
 					{
-						MenuBuilder.AddMenuEntry(
+						Section.AddMenuEntry(
+							"GoToDocsForAsset",
 							LOCTEXT("GoToDocsForAsset", "View Documentation"),
 							LOCTEXT("GoToDocsForAsset_ToolTip", "Click to open documentation"),
 							FSlateIcon(FEditorStyle::GetStyleSetName(), "HelpIcon.Hovered" ),
@@ -1186,59 +1339,65 @@ bool FAssetContextMenu::AddDocumentationMenuOptions(FMenuBuilder& MenuBuilder)
 							);
 					}
 			}
-			MenuBuilder.EndSection();
 		}
 	}
 
 	return bAddedOption;
 }
 
-bool FAssetContextMenu::AddAssetTypeMenuOptions(FMenuBuilder& MenuBuilder)
+bool FAssetContextMenu::AddAssetTypeMenuOptions(UEditorMenu* Menu, TArray<UObject*>& SelectedObjects)
 {
 	bool bAnyTypeOptions = false;
 
-	// Objects must be loaded for this operation... for now
-	TArray<FString> ObjectPaths;
-	for (int32 AssetIdx = 0; AssetIdx < SelectedAssets.Num(); ++AssetIdx)
+	if (SelectedObjects.Num() > 0)
 	{
-		ObjectPaths.Add(SelectedAssets[AssetIdx].ObjectPath.ToString());
-	}
-
-	TArray<UObject*> SelectedObjects;
-	if ( ContentBrowserUtils::LoadAssetsIfNeeded(ObjectPaths, SelectedObjects) )
-	{
-		// Load the asset tools module
-		FAssetToolsModule& AssetToolsModule = FModuleManager::LoadModuleChecked<FAssetToolsModule>(TEXT("AssetTools"));
-		bAnyTypeOptions = AssetToolsModule.Get().GetAssetActions(SelectedObjects, MenuBuilder, /*bIncludeHeading=*/true);
+		// Label "GetAssetActions" section
+		UContentBrowserAssetContextMenuContext* Context = Menu->FindContext<UContentBrowserAssetContextMenuContext>();
+		FEditorMenuSection* Section = Menu->FindSection("GetAssetActions");
+		if (Section && Section->Blocks.Num() > 0 && Context)
+		{
+			if (Context->CommonAssetTypeActions.IsValid())
+			{
+				Section->Label = FText::Format(NSLOCTEXT("AssetTools", "AssetSpecificOptionsMenuHeading", "{0} Actions"), Context->CommonAssetTypeActions.Pin()->GetName());
+			}
+			else if (Context->CommonClass)
+			{
+				Section->Label = FText::Format(NSLOCTEXT("AssetTools", "AssetSpecificOptionsMenuHeading", "{0} Actions"), FText::FromName(Context->CommonClass->GetFName()));
+			}
+			
+			bAnyTypeOptions = true;
+		}
 	}
 
 	return bAnyTypeOptions;
 }
 
-bool FAssetContextMenu::AddSourceControlMenuOptions(FMenuBuilder& MenuBuilder)
+bool FAssetContextMenu::AddSourceControlMenuOptions(UEditorMenu* Menu)
 {
-	MenuBuilder.BeginSection("AssetContextSourceControl");
+	FEditorMenuSection& Section = Menu->AddSection("AssetContextSourceControl");
 	
 	if ( ISourceControlModule::Get().IsEnabled() )
 	{
 		// SCC sub menu
-		MenuBuilder.AddSubMenu(
+		Section.AddEntry(FEditorMenuEntry::InitSubMenu(
+			Menu->GetMenuName(),
+			"SourceControlSubMenu",
 			LOCTEXT("SourceControlSubMenuLabel", "Source Control"),
 			LOCTEXT("SourceControlSubMenuToolTip", "Source control actions."),
-			FNewMenuDelegate::CreateSP(this, &FAssetContextMenu::FillSourceControlSubMenu),
+			FNewEditorMenuDelegate::CreateSP(this, &FAssetContextMenu::FillSourceControlSubMenu),
 			FUIAction(
 				FExecuteAction(),
 				FCanExecuteAction::CreateSP( this, &FAssetContextMenu::CanExecuteSourceControlActions )
 				),
-			NAME_None,
 			EUserInterfaceActionType::Button,
 			false,
 			FSlateIcon(FEditorStyle::GetStyleSetName(), "SourceControl.StatusIcon.On")
-			);
+			));
 	}
 	else
 	{
-		MenuBuilder.AddMenuEntry(
+		Section.AddMenuEntry(
+			"SCCConnectToSourceControl",
 			LOCTEXT("SCCConnectToSourceControl", "Connect To Source Control..."),
 			LOCTEXT("SCCConnectToSourceControlTooltip", "Connect to source control to allow source control operations to be performed on content and levels."),
 			FSlateIcon(FEditorStyle::GetStyleSetName(), "SourceControl.Actions.Connect"),
@@ -1252,7 +1411,8 @@ bool FAssetContextMenu::AddSourceControlMenuOptions(FMenuBuilder& MenuBuilder)
 	// Diff selected
 	if (CanExecuteDiffSelected())
 	{
-		MenuBuilder.AddMenuEntry(
+		Section.AddMenuEntry(
+			"DiffSelected",
 			LOCTEXT("DiffSelected", "Diff Selected"),
 			LOCTEXT("DiffSelectedTooltip", "Diff the two assets that you have selected."),
 			FSlateIcon(FEditorStyle::GetStyleSetName(), "SourceControl.Actions.Diff"),
@@ -1262,17 +1422,17 @@ bool FAssetContextMenu::AddSourceControlMenuOptions(FMenuBuilder& MenuBuilder)
 		);
 	}
 
-	MenuBuilder.EndSection();
 	return true;
 }
 
-void FAssetContextMenu::FillSourceControlSubMenu(FMenuBuilder& MenuBuilder)
+void FAssetContextMenu::FillSourceControlSubMenu(UEditorMenu* Menu)
 {
-	MenuBuilder.BeginSection("AssetSourceControlActions", LOCTEXT("AssetSourceControlActionsMenuHeading", "Source Control"));
+	FEditorMenuSection& Section = Menu->AddSection("AssetSourceControlActions", LOCTEXT("AssetSourceControlActionsMenuHeading", "Source Control"));
 
 	if( CanExecuteSCCMerge() )
 	{
-		MenuBuilder.AddMenuEntry(
+		Section.AddMenuEntry(
+			"SCCMerge",
 			LOCTEXT("SCCMerge", "Merge"),
 			LOCTEXT("SCCMergeTooltip", "Opens the blueprint editor with the merge tool open."),
 			FSlateIcon(),
@@ -1285,7 +1445,8 @@ void FAssetContextMenu::FillSourceControlSubMenu(FMenuBuilder& MenuBuilder)
 
 	if( CanExecuteSCCSync() )
 	{
-		MenuBuilder.AddMenuEntry(
+		Section.AddMenuEntry(
+			"SCCSync",
 			LOCTEXT("SCCSync", "Sync"),
 			LOCTEXT("SCCSyncTooltip", "Updates the item to the latest version in source control."),
 			FSlateIcon(FEditorStyle::GetStyleSetName(), "SourceControl.Actions.Sync"),
@@ -1298,7 +1459,8 @@ void FAssetContextMenu::FillSourceControlSubMenu(FMenuBuilder& MenuBuilder)
 
 	if ( CanExecuteSCCCheckOut() )
 	{
-		MenuBuilder.AddMenuEntry(
+		Section.AddMenuEntry(
+			"SCCCheckOut",
 			LOCTEXT("SCCCheckOut", "Check Out"),
 			LOCTEXT("SCCCheckOutTooltip", "Checks out the selected asset from source control."),
 			FSlateIcon(FEditorStyle::GetStyleSetName(), "SourceControl.Actions.CheckOut"),
@@ -1311,7 +1473,8 @@ void FAssetContextMenu::FillSourceControlSubMenu(FMenuBuilder& MenuBuilder)
 
 	if ( CanExecuteSCCOpenForAdd() )
 	{
-		MenuBuilder.AddMenuEntry(
+		Section.AddMenuEntry(
+			"SCCOpenForAdd",
 			LOCTEXT("SCCOpenForAdd", "Mark For Add"),
 			LOCTEXT("SCCOpenForAddTooltip", "Adds the selected asset to source control."),
 			FSlateIcon(FEditorStyle::GetStyleSetName(), "SourceControl.Actions.Add"),
@@ -1324,7 +1487,8 @@ void FAssetContextMenu::FillSourceControlSubMenu(FMenuBuilder& MenuBuilder)
 
 	if ( CanExecuteSCCCheckIn() )
 	{
-		MenuBuilder.AddMenuEntry(
+		Section.AddMenuEntry(
+			"SCCCheckIn",
 			LOCTEXT("SCCCheckIn", "Check In"),
 			LOCTEXT("SCCCheckInTooltip", "Checks in the selected asset to source control."),
 			FSlateIcon(FEditorStyle::GetStyleSetName(), "SourceControl.Actions.Submit"),
@@ -1335,7 +1499,8 @@ void FAssetContextMenu::FillSourceControlSubMenu(FMenuBuilder& MenuBuilder)
 		);
 	}
 
-	MenuBuilder.AddMenuEntry(
+	Section.AddMenuEntry(
+		"SCCRefresh",
 		LOCTEXT("SCCRefresh", "Refresh"),
 		LOCTEXT("SCCRefreshTooltip", "Updates the source control status of the asset."),
 		FSlateIcon(FEditorStyle::GetStyleSetName(), "SourceControl.Actions.Refresh"),
@@ -1347,7 +1512,8 @@ void FAssetContextMenu::FillSourceControlSubMenu(FMenuBuilder& MenuBuilder)
 
 	if( CanExecuteSCCHistory() )
 	{
-		MenuBuilder.AddMenuEntry(
+		Section.AddMenuEntry(
+			"SCCHistory",
 			LOCTEXT("SCCHistory", "History"),
 			LOCTEXT("SCCHistoryTooltip", "Displays the source control revision history of the selected asset."),
 			FSlateIcon(FEditorStyle::GetStyleSetName(), "SourceControl.Actions.History"),
@@ -1357,7 +1523,8 @@ void FAssetContextMenu::FillSourceControlSubMenu(FMenuBuilder& MenuBuilder)
 			)
 		);
 
-		MenuBuilder.AddMenuEntry(
+		Section.AddMenuEntry(
+			"SCCDiffAgainstDepot",
 			LOCTEXT("SCCDiffAgainstDepot", "Diff Against Depot"),
 			LOCTEXT("SCCDiffAgainstDepotTooltip", "Look at differences between your version of the asset and that in source control."),
 			FSlateIcon(FEditorStyle::GetStyleSetName(), "SourceControl.Actions.Diff"),
@@ -1370,7 +1537,8 @@ void FAssetContextMenu::FillSourceControlSubMenu(FMenuBuilder& MenuBuilder)
 
 	if( CanExecuteSCCRevert() )
 	{
-		MenuBuilder.AddMenuEntry(
+		Section.AddMenuEntry(
+			"SCCRevert",
 			LOCTEXT("SCCRevert", "Revert"),
 			LOCTEXT("SCCRevertTooltip", "Reverts the asset to the state it was before it was checked out."),
 			FSlateIcon(FEditorStyle::GetStyleSetName(), "SourceControl.Actions.Revert"),
@@ -1380,8 +1548,6 @@ void FAssetContextMenu::FillSourceControlSubMenu(FMenuBuilder& MenuBuilder)
 			)
 		);
 	}
-
-	MenuBuilder.EndSection();
 }
 
 bool FAssetContextMenu::CanExecuteSourceControlActions() const
@@ -1389,22 +1555,22 @@ bool FAssetContextMenu::CanExecuteSourceControlActions() const
 	return !bAtLeastOneClassSelected;
 }
 
-bool FAssetContextMenu::AddCollectionMenuOptions(FMenuBuilder& MenuBuilder)
+bool FAssetContextMenu::AddCollectionMenuOptions(UEditorMenu* Menu)
 {
 	class FManageCollectionsContextMenu
 	{
 	public:
-		static void CreateManageCollectionsSubMenu(FMenuBuilder& SubMenuBuilder, TSharedRef<FCollectionAssetManagement> QuickAssetManagement)
+		static void CreateManageCollectionsSubMenu(UEditorMenu* SubMenu, TSharedRef<FCollectionAssetManagement> QuickAssetManagement)
 		{
 			FCollectionManagerModule& CollectionManagerModule = FCollectionManagerModule::GetModule();
 
 			TArray<FCollectionNameType> AvailableCollections;
 			CollectionManagerModule.Get().GetRootCollections(AvailableCollections);
 
-			CreateManageCollectionsSubMenu(SubMenuBuilder, QuickAssetManagement, MoveTemp(AvailableCollections));
+			CreateManageCollectionsSubMenu(SubMenu, QuickAssetManagement, MoveTemp(AvailableCollections));
 		}
 
-		static void CreateManageCollectionsSubMenu(FMenuBuilder& SubMenuBuilder, TSharedRef<FCollectionAssetManagement> QuickAssetManagement, TArray<FCollectionNameType> AvailableCollections)
+		static void CreateManageCollectionsSubMenu(UEditorMenu* SubMenu, TSharedRef<FCollectionAssetManagement> QuickAssetManagement, TArray<FCollectionNameType> AvailableCollections)
 		{
 			FCollectionManagerModule& CollectionManagerModule = FCollectionManagerModule::GetModule();
 
@@ -1413,6 +1579,7 @@ bool FAssetContextMenu::AddCollectionMenuOptions(FMenuBuilder& MenuBuilder)
 				return One.Name.LexicalLess(Two.Name);
 			});
 
+			FEditorMenuSection& Section = SubMenu->AddSection("Section");
 			for (const FCollectionNameType& AvailableCollection : AvailableCollections)
 			{
 				// Never display system collections
@@ -1434,24 +1601,26 @@ bool FAssetContextMenu::AddCollectionMenuOptions(FMenuBuilder& MenuBuilder)
 
 				if (AvailableChildCollections.Num() > 0)
 				{
-					SubMenuBuilder.AddSubMenu(
+					Section.AddEntry(FEditorMenuEntry::InitSubMenu(
+						SubMenu->GetMenuName(),
+						NAME_None,
 						FText::FromName(AvailableCollection.Name), 
 						FText::GetEmpty(), 
-						FNewMenuDelegate::CreateStatic(&FManageCollectionsContextMenu::CreateManageCollectionsSubMenu, QuickAssetManagement, AvailableChildCollections),
+						FNewEditorMenuDelegate::CreateStatic(&FManageCollectionsContextMenu::CreateManageCollectionsSubMenu, QuickAssetManagement, AvailableChildCollections),
 						FUIAction(
 							FExecuteAction::CreateStatic(&FManageCollectionsContextMenu::OnCollectionClicked, QuickAssetManagement, AvailableCollection),
 							FCanExecuteAction::CreateStatic(&FManageCollectionsContextMenu::IsCollectionEnabled, QuickAssetManagement, AvailableCollection),
 							FGetActionCheckState::CreateStatic(&FManageCollectionsContextMenu::GetCollectionCheckState, QuickAssetManagement, AvailableCollection)
 							), 
-						NAME_None, 
 						EUserInterfaceActionType::ToggleButton,
 						false,
 						FSlateIcon(FEditorStyle::GetStyleSetName(), ECollectionShareType::GetIconStyleName(AvailableCollection.Type))
-						);
+						));
 				}
 				else
 				{
-					SubMenuBuilder.AddMenuEntry(
+					Section.AddMenuEntry(
+						NAME_None,
 						FText::FromName(AvailableCollection.Name), 
 						FText::GetEmpty(), 
 						FSlateIcon(FEditorStyle::GetStyleSetName(), ECollectionShareType::GetIconStyleName(AvailableCollection.Type)), 
@@ -1460,7 +1629,6 @@ bool FAssetContextMenu::AddCollectionMenuOptions(FMenuBuilder& MenuBuilder)
 							FCanExecuteAction::CreateStatic(&FManageCollectionsContextMenu::IsCollectionEnabled, QuickAssetManagement, AvailableCollection),
 							FGetActionCheckState::CreateStatic(&FManageCollectionsContextMenu::GetCollectionCheckState, QuickAssetManagement, AvailableCollection)
 							), 
-						NAME_None, 
 						EUserInterfaceActionType::ToggleButton
 						);
 				}
@@ -1497,7 +1665,7 @@ bool FAssetContextMenu::AddCollectionMenuOptions(FMenuBuilder& MenuBuilder)
 
 	FCollectionManagerModule& CollectionManagerModule = FCollectionManagerModule::GetModule();
 
-	MenuBuilder.BeginSection("AssetContextCollections", LOCTEXT("AssetCollectionOptionsMenuHeading", "Collections"));
+	FEditorMenuSection& Section = Menu->AddSection("AssetContextCollections", LOCTEXT("AssetCollectionOptionsMenuHeading", "Collections"));
 
 	// Show a sub-menu that allows you to quickly add or remove the current asset selection from the available collections
 	if (CollectionManagerModule.Get().HasCollections())
@@ -1505,11 +1673,13 @@ bool FAssetContextMenu::AddCollectionMenuOptions(FMenuBuilder& MenuBuilder)
 		TSharedRef<FCollectionAssetManagement> QuickAssetManagement = MakeShareable(new FCollectionAssetManagement());
 		QuickAssetManagement->SetCurrentAssets(SelectedAssets);
 
-		MenuBuilder.AddSubMenu(
+		Section.AddEntry(FEditorMenuEntry::InitSubMenu(
+			Menu->GetMenuName(),
+			"ManageCollections",
 			LOCTEXT("ManageCollections", "Manage Collections"),
 			LOCTEXT("ManageCollections_ToolTip", "Manage the collections that the selected asset(s) belong to."),
-			FNewMenuDelegate::CreateStatic(&FManageCollectionsContextMenu::CreateManageCollectionsSubMenu, QuickAssetManagement)
-			);
+			FNewEditorMenuDelegate::CreateStatic(&FManageCollectionsContextMenu::CreateManageCollectionsSubMenu, QuickAssetManagement)
+			));
 
 		bHasAddedItems = true;
 	}
@@ -1517,7 +1687,8 @@ bool FAssetContextMenu::AddCollectionMenuOptions(FMenuBuilder& MenuBuilder)
 	// "Remove from collection" (only display option if exactly one collection is selected)
 	if ( SourcesData.Collections.Num() == 1 && !SourcesData.IsDynamicCollection() )
 	{
-		MenuBuilder.AddMenuEntry(
+		Section.AddMenuEntry(
+			"RemoveFromCollection",
 			FText::Format(LOCTEXT("RemoveFromCollectionFmt", "Remove From {0}"), FText::FromName(SourcesData.Collections[0].Name)),
 			LOCTEXT("RemoveFromCollection_ToolTip", "Removes the selected asset from the current collection."),
 			FSlateIcon(),
@@ -1530,7 +1701,6 @@ bool FAssetContextMenu::AddCollectionMenuOptions(FMenuBuilder& MenuBuilder)
 		bHasAddedItems = true;
 	}
 
-	MenuBuilder.EndSection();
 
 	return bHasAddedItems;
 }
@@ -3015,7 +3185,7 @@ void FAssetContextMenu::GetSelectedPackages(TArray<UPackage*>& OutPackages) cons
 	}
 }
 
-void FAssetContextMenu::MakeChunkIDListMenu(FMenuBuilder& MenuBuilder)
+void FAssetContextMenu::MakeChunkIDListMenu(UEditorMenu* Menu)
 {
 	TArray<int32> FoundChunks;
 	TArray< FAssetData > AssetViewSelectedAssets = AssetView.Pin()->GetSelectedAssets();
@@ -3032,9 +3202,11 @@ void FAssetContextMenu::MakeChunkIDListMenu(FMenuBuilder& MenuBuilder)
 		}
 	}
 
+	FEditorMenuSection& Section = Menu->AddSection("Chunks");
 	for (auto ChunkID : FoundChunks)
 	{
-		MenuBuilder.AddMenuEntry(
+		Section.AddMenuEntry(
+			NAME_None,
 			FText::Format(LOCTEXT("PackageChunk", "Chunk {0}"), FText::AsNumber(ChunkID)),
 			FText(),
 			FSlateIcon(),
diff --git a/Engine/Source/Editor/ContentBrowser/Private/AssetContextMenu.h b/Engine/Source/Editor/ContentBrowser/Private/AssetContextMenu.h
--- a/Engine/Source/Editor/ContentBrowser/Private/AssetContextMenu.h
+++ b/Engine/Source/Editor/ContentBrowser/Private/AssetContextMenu.h
@@ -9,7 +9,7 @@
 #include "ContentBrowserDelegates.h"
 #include "SourcesData.h"
 
-class FMenuBuilder;
+class UEditorMenu;
 class FUICommandList;
 class SAssetView;
 class SWindow;
@@ -118,34 +118,34 @@ private:
 
 private:
 	/** Adds asset type-specific menu options to a menu builder. Returns true if any options were added. */
-	bool AddAssetTypeMenuOptions(FMenuBuilder& MenuBuilder);
+	bool AddAssetTypeMenuOptions(UEditorMenu* Menu, TArray<UObject*>& SelectedObjects);
 
 	/** Adds asset type-specific menu options to a menu builder. Returns true if any options were added. */
-	bool AddImportedAssetMenuOptions(FMenuBuilder& MenuBuilder);
+	bool AddImportedAssetMenuOptions(UEditorMenu* Menu);
 	
 	/** Adds common menu options to a menu builder. Returns true if any options were added. */
-	bool AddCommonMenuOptions(FMenuBuilder& MenuBuilder);
+	bool AddCommonMenuOptions(UEditorMenu* Menu);
 
 	/** Adds explore menu options to a menu builder. */
-	void AddExploreMenuOptions(FMenuBuilder& MenuBuilder);
+	void AddExploreMenuOptions(UEditorMenu* Menu);
 
 	/** Adds Asset Actions sub-menu to a menu builder. */
-	void MakeAssetActionsSubMenu(FMenuBuilder& MenuBuilder);
+	void MakeAssetActionsSubMenu(UEditorMenu* Menu);
 
 	/** Handler to check to see if "Asset Actions" are allowed */
 	bool CanExecuteAssetActions() const;
 
 	/** Adds Asset Localization sub-menu to a menu builder. */
-	void MakeAssetLocalizationSubMenu(FMenuBuilder& MenuBuilder);
+	void MakeAssetLocalizationSubMenu(UEditorMenu* Menu);
 
 	/** Adds the Create Localized Asset sub-menu to a menu builder. */
-	void MakeCreateLocalizedAssetSubMenu(FMenuBuilder& MenuBuilder, TSet<FName> InSelectedSourceAssets, TArray<FLocalizedAssetsState> InLocalizedAssetsState);
+	void MakeCreateLocalizedAssetSubMenu(UEditorMenu* Menu, TSet<FName> InSelectedSourceAssets, TArray<FLocalizedAssetsState> InLocalizedAssetsState);
 
 	/** Adds the Show Localized Assets sub-menu to a menu builder. */
-	void MakeShowLocalizedAssetSubMenu(FMenuBuilder& MenuBuilder, TArray<FLocalizedAssetsState> InLocalizedAssetsState);
+	void MakeShowLocalizedAssetSubMenu(UEditorMenu* Menu, TArray<FLocalizedAssetsState> InLocalizedAssetsState);
 
 	/** Adds the Edit Localized Assets sub-menu to a menu builder. */
-	void MakeEditLocalizedAssetSubMenu(FMenuBuilder& MenuBuilder, TArray<FLocalizedAssetsState> InLocalizedAssetsState);
+	void MakeEditLocalizedAssetSubMenu(UEditorMenu* Menu, TArray<FLocalizedAssetsState> InLocalizedAssetsState);
 
 	/** Create new localized assets for the given culture */
 	void ExecuteCreateLocalizedAsset(TSet<FName> InSelectedSourceAssets, FLocalizedAssetsState InLocalizedAssetsStateForCulture);
@@ -157,25 +157,25 @@ private:
 	void ExecuteOpenEditorsForAssets(TArray<FName> InAssets);
 
 	/** Adds asset reference menu options to a menu builder. Returns true if any options were added. */
-	bool AddReferenceMenuOptions(FMenuBuilder& MenuBuilder);
+	bool AddReferenceMenuOptions(UEditorMenu* Menu);
 
 	/** Adds asset documentation menu options to a menu builder. Returns true if any options were added. */
-	bool AddDocumentationMenuOptions(FMenuBuilder& MenuBuilder);
+	bool AddDocumentationMenuOptions(UEditorMenu* Menu);
 	
 	/** Adds source control menu options to a menu builder. */
-	bool AddSourceControlMenuOptions(FMenuBuilder& MenuBuilder);
+	bool AddSourceControlMenuOptions(UEditorMenu* Menu);
 
 	/** Fills the source control sub-menu */
-	void FillSourceControlSubMenu(FMenuBuilder& MenuBuilder);
+	void FillSourceControlSubMenu(UEditorMenu* Menu);
 
 	/** Handler to check to see if SCC actions are allowed */
 	bool CanExecuteSourceControlActions() const;
 
 	/** Adds menu options related to working with collections */
-	bool AddCollectionMenuOptions(FMenuBuilder& MenuBuilder);
+	bool AddCollectionMenuOptions(UEditorMenu* Menu);
 
 	/** Creates a sub-menu of Chunk IDs that are are assigned to all selected assets */
-	void MakeChunkIDListMenu(FMenuBuilder& MenuBuilder);
+	void MakeChunkIDListMenu(UEditorMenu* Menu);
 
 	/** Handler for when sync to asset tree is selected */
 	void ExecuteSyncToAssetTree();
@@ -386,6 +386,10 @@ private:
 	FText GetExecutePropertyMatrixTooltip() const;
 
 private:
+
+	/** Registers the base context menu for assets */
+	void RegisterContextMenu(const FName MenuName);
+
 	/** Generates a list of selected assets in the content browser */
 	void GetSelectedAssets(TArray<UObject*>& Assets, bool SkipRedirectors) const;
 
diff --git a/Engine/Source/Editor/ContentBrowser/Public/ContentBrowserMenuContexts.h b/Engine/Source/Editor/ContentBrowser/Public/ContentBrowserMenuContexts.h
new file mode 100644
--- /dev/null
+++ b/Engine/Source/Editor/ContentBrowser/Public/ContentBrowserMenuContexts.h
@@ -0,0 +1,30 @@
+// Copyright 1998-2019 Epic Games, Inc. All Rights Reserved.
+
+#pragma once
+
+#include "CoreMinimal.h"
+#include "UObject/ObjectMacros.h"
+#include "UObject/Object.h"
+
+#include "ContentBrowserMenuContexts.generated.h"
+
+class FAssetContextMenu;
+class IAssetTypeActions;
+
+UCLASS()
+class CONTENTBROWSER_API UContentBrowserAssetContextMenuContext : public UObject
+{
+	GENERATED_BODY()
+
+public:
+
+	TWeakPtr<FAssetContextMenu> AssetContextMenu;
+
+	TWeakPtr<IAssetTypeActions> CommonAssetTypeActions;
+	
+	UPROPERTY()
+	TArray<UObject*> SelectedObjects;
+
+	UPROPERTY()
+	UClass* CommonClass;
+};
