ToolMenus, simplify add sub menu API to not require menu's name
#rnx

diff --git a/Engine/Plugins/Editor/AssetManagerEditor/Source/AssetManagerEditor/Private/ReferenceViewer/ReferenceViewerSchema.cpp b/Engine/Plugins/Editor/AssetManagerEditor/Source/AssetManagerEditor/Private/ReferenceViewer/ReferenceViewerSchema.cpp
--- a/Engine/Plugins/Editor/AssetManagerEditor/Source/AssetManagerEditor/Private/ReferenceViewer/ReferenceViewerSchema.cpp
+++ b/Engine/Plugins/Editor/AssetManagerEditor/Source/AssetManagerEditor/Private/ReferenceViewer/ReferenceViewerSchema.cpp
@@ -65,7 +65,6 @@ void UReferenceViewerSchema::GetContextMenuActions(UToolMenu* Menu, UGraphNodeCo
 		Section.AddMenuEntry(FAssetManagerEditorCommands::Get().ZoomToFit);
 		Section.AddMenuEntry(FAssetManagerEditorCommands::Get().ReCenterGraph);
 		Section.AddSubMenu(
-			Menu->GetMenuName(),
 			"MakeCollectionWith",
 			NSLOCTEXT("ReferenceViewerSchema", "MakeCollectionWithTitle", "Make Collection with"),
 			NSLOCTEXT("ReferenceViewerSchema", "MakeCollectionWithTooltip", "Makes a collection with either the referencers or dependencies of the selected nodes."),
@@ -145,7 +144,6 @@ void UReferenceViewerSchema::GetMakeCollectionWithSubMenu(UToolMenu* Menu)
 	FToolMenuSection& Section = Menu->AddSection("Section");
 
 	Section.AddSubMenu(
-		Menu->GetMenuName(),
 		"MakeCollectionWithReferencers",
 		NSLOCTEXT("ReferenceViewerSchema", "MakeCollectionWithReferencersTitle", "Referencers <-"),
 		NSLOCTEXT("ReferenceViewerSchema", "MakeCollectionWithReferencersTooltip", "Makes a collection with assets one connection to the left of selected nodes."),
@@ -153,7 +151,6 @@ void UReferenceViewerSchema::GetMakeCollectionWithSubMenu(UToolMenu* Menu)
 		);
 
 	Section.AddSubMenu(
-		Menu->GetMenuName(),
 		"MakeCollectionWithDependencies",
 		NSLOCTEXT("ReferenceViewerSchema", "MakeCollectionWithDependenciesTitle", "Dependencies ->"),
 		NSLOCTEXT("ReferenceViewerSchema", "MakeCollectionWithDependenciesTooltip", "Makes a collection with assets one connection to the right of selected nodes."),
diff --git a/Engine/Plugins/Experimental/PythonScriptPlugin/Source/PythonScriptPlugin/Private/PythonScriptPlugin.cpp b/Engine/Plugins/Experimental/PythonScriptPlugin/Source/PythonScriptPlugin/Private/PythonScriptPlugin.cpp
--- a/Engine/Plugins/Experimental/PythonScriptPlugin/Source/PythonScriptPlugin/Private/PythonScriptPlugin.cpp
+++ b/Engine/Plugins/Experimental/PythonScriptPlugin/Source/PythonScriptPlugin/Private/PythonScriptPlugin.cpp
@@ -350,15 +350,14 @@ private:
 			FSlateIcon(),
 			FUIAction(FExecuteAction::CreateRaw(this, &FPythonCommandMenuImpl::Menu_ExecutePython))
 		);
-		Section.AddEntry(FToolMenuEntry::InitSubMenu(
-			"LevelEditor.MainMenu.File",
+		Section.AddSubMenu(
 			"RecentPythonsSubMenu",
 			LOCTEXT("RecentPythonsSubMenu", "Recent Python Scripts"),
 			LOCTEXT("RecentPythonsSubMenu_ToolTip", "Select a recent Python Script file and Execute it."),
 			FNewToolMenuDelegate::CreateRaw(this, &FPythonCommandMenuImpl::MakeRecentPythonScriptMenu),
 			false,
 			FSlateIcon(FEditorStyle::GetStyleSetName(), "MainFrame.RecentLevels")
-		));
+		);
 	}
 
 	void Menu_ExecutePythonRecent(int32 Index)
diff --git a/Engine/Plugins/FX/Niagara/Source/NiagaraEditor/Private/EdGraphSchema_Niagara.cpp b/Engine/Plugins/FX/Niagara/Source/NiagaraEditor/Private/EdGraphSchema_Niagara.cpp
--- a/Engine/Plugins/FX/Niagara/Source/NiagaraEditor/Private/EdGraphSchema_Niagara.cpp
+++ b/Engine/Plugins/FX/Niagara/Source/NiagaraEditor/Private/EdGraphSchema_Niagara.cpp
@@ -1660,7 +1660,6 @@ void UEdGraphSchema_Niagara::GetContextMenuActions(UToolMenu* Menu, UGraphNodeCo
 			if (PinToTypeDefinition(InGraphPin) == FNiagaraTypeDefinition::GetGenericNumericDef() && InGraphPin->LinkedTo.Num() == 0)
 			{
 				Section.AddSubMenu(
-					Menu->GetMenuName(),
 					"ConvertNumericSpecific",
 					LOCTEXT("ConvertNumericSpecific", "Convert Numeric To..."),
 					LOCTEXT("ConvertNumericSpecificToolTip", "Convert Numeric pin to the specific typed pin."),
@@ -1676,7 +1675,6 @@ void UEdGraphSchema_Niagara::GetContextMenuActions(UToolMenu* Menu, UGraphNodeCo
 				if (InGraphPin->LinkedTo.Num() > 1)
 				{
 					Section.AddSubMenu(
-						Menu->GetMenuName(),
 						"BreakLinkTo",
 						LOCTEXT("BreakLinkTo", "Break Link To..."),
 						LOCTEXT("BreakSpecificLinks", "Break a specific link..."),
@@ -1736,7 +1734,6 @@ void UEdGraphSchema_Niagara::GetContextMenuActions(UToolMenu* Menu, UGraphNodeCo
 			const FName SectionName = "EdGraphSchema_NiagaraNodeActions";
 			FToolMenuSection& Section = Menu->AddSection(SectionName, LOCTEXT("PinConversionMenuHeader", "Convert Pins"));
 			Section.AddSubMenu(
-				Menu->GetMenuName(),
 				"ConvertNumericSpecific",
 				LOCTEXT("ConvertNumericSpecific", "Convert All Numerics To..."),
 				LOCTEXT("ConvertNumericSpecificToolTip", "Convert all Numeric pins to the specific typed pin."),
diff --git a/Engine/Plugins/FX/Niagara/Source/NiagaraEditor/Public/NiagaraNodeAssignment.h b/Engine/Plugins/FX/Niagara/Source/NiagaraEditor/Public/NiagaraNodeAssignment.h
--- a/Engine/Plugins/FX/Niagara/Source/NiagaraEditor/Public/NiagaraNodeAssignment.h
+++ b/Engine/Plugins/FX/Niagara/Source/NiagaraEditor/Public/NiagaraNodeAssignment.h
@@ -34,8 +34,8 @@ public:
 	const TArray<FString>& GetAssignmentDefaults() const { return AssignmentDefaultValues; }
 
 	void MergeUp();
-	void BuildAddParameterMenu(FMenuBuilder& MenuBuilder, ENiagaraScriptUsage InUsage, UNiagaraNodeOutput* InGraphOutputNode);
-	void BuildCreateParameterMenu(FMenuBuilder& MenuBuilder, ENiagaraScriptUsage InUsage, UNiagaraNodeOutput* InGraphOutputNode);
+	void BuildAddParameterMenu(class FMenuBuilder& MenuBuilder, ENiagaraScriptUsage InUsage, UNiagaraNodeOutput* InGraphOutputNode);
+	void BuildCreateParameterMenu(class FMenuBuilder& MenuBuilder, ENiagaraScriptUsage InUsage, UNiagaraNodeOutput* InGraphOutputNode);
 
 	//~ Begin EdGraphNode Interface
 	virtual void PostLoad() override;
diff --git a/Engine/Source/Developer/AssetTools/Private/AssetTypeActions/AssetTypeActions_AnimBlueprint.cpp b/Engine/Source/Developer/AssetTools/Private/AssetTypeActions/AssetTypeActions_AnimBlueprint.cpp
--- a/Engine/Source/Developer/AssetTools/Private/AssetTypeActions/AssetTypeActions_AnimBlueprint.cpp
+++ b/Engine/Source/Developer/AssetTools/Private/AssetTypeActions/AssetTypeActions_AnimBlueprint.cpp
@@ -38,15 +38,14 @@ void FAssetTypeActions_AnimBlueprint::GetActions(const TArray<UObject*>& InObjec
 			)
 		);
 
-	Section.AddEntry(FToolMenuEntry::InitSubMenu(
-		NAME_None,
+	Section.AddSubMenu(
 		"RetargetBlueprintSubmenu",
 		LOCTEXT("RetargetBlueprintSubmenu", "Retarget Anim Blueprints"),
 		LOCTEXT("RetargetBlueprintSubmenu_ToolTip", "Opens the retarget blueprints menu"),
 		FNewToolMenuDelegate::CreateSP( this, &FAssetTypeActions_AnimBlueprint::FillRetargetMenu, InObjects),
 		false,
 		FSlateIcon(FEditorStyle::GetStyleSetName(), "Persona.AssetActions.RetargetSkeleton")
-		));
+		);
 }
 
 void FAssetTypeActions_AnimBlueprint::FillRetargetMenu(UToolMenu* MenuBuilder, const TArray<UObject*> InObjects)
diff --git a/Engine/Source/Developer/AssetTools/Private/AssetTypeActions/AssetTypeActions_AnimSequence.cpp b/Engine/Source/Developer/AssetTools/Private/AssetTypeActions/AssetTypeActions_AnimSequence.cpp
--- a/Engine/Source/Developer/AssetTools/Private/AssetTypeActions/AssetTypeActions_AnimSequence.cpp
+++ b/Engine/Source/Developer/AssetTools/Private/AssetTypeActions/AssetTypeActions_AnimSequence.cpp
@@ -32,15 +32,14 @@ void FAssetTypeActions_AnimSequence::GetActions(const TArray<UObject*>& InObject
 	auto Sequences = GetTypedWeakObjectPtrs<UAnimSequence>(InObjects);
 
 	// create menu
-	Section.AddEntry(FToolMenuEntry::InitSubMenu(
-		NAME_None,
+	Section.AddSubMenu(
 		"CreateAnimSubmenu",
 		LOCTEXT("CreateAnimSubmenu", "Create"),
 		LOCTEXT("CreateAnimSubmenu_ToolTip", "Create assets from this anim sequence"),
 		FNewMenuDelegate::CreateSP(this, &FAssetTypeActions_AnimSequence::FillCreateMenu, Sequences),
 		false,
 		FSlateIcon(FEditorStyle::GetStyleSetName(), "Persona.AssetActions.CreateAnimAsset")
-		));
+		);
 
 	Section.AddMenuEntry(
 		"AnimSequence_ReimportWithNewSource",
diff --git a/Engine/Source/Developer/AssetTools/Private/AssetTypeActions/AssetTypeActions_AnimationAsset.cpp b/Engine/Source/Developer/AssetTools/Private/AssetTypeActions/AssetTypeActions_AnimationAsset.cpp
--- a/Engine/Source/Developer/AssetTools/Private/AssetTypeActions/AssetTypeActions_AnimationAsset.cpp
+++ b/Engine/Source/Developer/AssetTools/Private/AssetTypeActions/AssetTypeActions_AnimationAsset.cpp
@@ -45,15 +45,14 @@ void FAssetTypeActions_AnimationAsset::GetActions(const TArray<UObject*>& InObje
 			)
 		);
 
-	Section.AddEntry(FToolMenuEntry::InitSubMenu(
-		NAME_None,
+	Section.AddSubMenu(
 		"RetargetAnimSubmenu", 
 		LOCTEXT("RetargetAnimSubmenu", "Retarget Anim Assets"),
 		LOCTEXT("RetargetAnimSubmenu_ToolTip", "Opens the retarget anim assets menu"),
 		FNewMenuDelegate::CreateSP( this, &FAssetTypeActions_AnimationAsset::FillRetargetMenu, InObjects ),
 		false,
 		FSlateIcon(FEditorStyle::GetStyleSetName(), "Persona.AssetActions.RetargetSkeleton")
-		));
+		);
 }
 
 void FAssetTypeActions_AnimationAsset::FillRetargetMenu( FMenuBuilder& MenuBuilder, const TArray<UObject*> InObjects )
diff --git a/Engine/Source/Developer/AssetTools/Private/AssetTypeActions/AssetTypeActions_SkeletalMesh.cpp b/Engine/Source/Developer/AssetTools/Private/AssetTypeActions/AssetTypeActions_SkeletalMesh.cpp
--- a/Engine/Source/Developer/AssetTools/Private/AssetTypeActions/AssetTypeActions_SkeletalMesh.cpp
+++ b/Engine/Source/Developer/AssetTools/Private/AssetTypeActions/AssetTypeActions_SkeletalMesh.cpp
@@ -452,23 +452,21 @@ void FAssetTypeActions_SkeletalMesh::GetActions(const TArray<UObject*>& InObject
 {
 	auto Meshes = GetTypedWeakObjectPtrs<USkeletalMesh>(InObjects);
 
-	Section.AddEntry(FToolMenuEntry::InitSubMenu(
-			NAME_None,
+	Section.AddSubMenu(
 			"CreateSkeletalMeshSubmenu",
 			LOCTEXT("CreateSkeletalMeshSubmenu", "Create"),
 			LOCTEXT("CreateSkeletalMeshSubmenu_ToolTip", "Create related assets"),
 			FNewMenuDelegate::CreateSP(this, &FAssetTypeActions_SkeletalMesh::FillCreateMenu, Meshes),
 			false,
 			FSlateIcon(FEditorStyle::GetStyleSetName(), "Persona.AssetActions.CreateAnimAsset")
-			));
+			);
 
-	Section.AddEntry(FToolMenuEntry::InitSubMenu(
-		NAME_None,
+	Section.AddSubMenu(
 		"SkeletalMesh_LODImport",	
 		LOCTEXT("SkeletalMesh_LODImport", "Import LOD"),
 		LOCTEXT("SkeletalMesh_LODImportTooltip", "Select which LODs to import."),
 		FNewMenuDelegate::CreateSP(this, &FAssetTypeActions_SkeletalMesh::GetLODMenu, Meshes)
-		));
+		);
 	
 	Section.AddMenuEntry(
 		"ImportClothing_Entry",
@@ -478,13 +476,12 @@ void FAssetTypeActions_SkeletalMesh::GetActions(const TArray<UObject*>& InObject
 		FUIAction(FExecuteAction::CreateSP(this, &FAssetTypeActions_SkeletalMesh::ExecuteImportClothing, Meshes)));
 
 	// skeleton menu
-	Section.AddEntry(FToolMenuEntry::InitSubMenu(
-		NAME_None,
+	Section.AddSubMenu(
 		"SkeletonSubmenu",
 		LOCTEXT("SkeletonSubmenu", "Skeleton"),
 		LOCTEXT("SkeletonSubmenu_ToolTip", "Skeleton related actions"),
 		FNewMenuDelegate::CreateSP(this, &FAssetTypeActions_SkeletalMesh::FillSkeletonMenu, Meshes)
-	));
+	);
 }
 
 void FAssetTypeActions_SkeletalMesh::FillCreateMenu(FMenuBuilder& MenuBuilder, TArray<TWeakObjectPtr<USkeletalMesh>> Meshes) const
diff --git a/Engine/Source/Developer/AssetTools/Private/AssetTypeActions/AssetTypeActions_Skeleton.cpp b/Engine/Source/Developer/AssetTools/Private/AssetTypeActions/AssetTypeActions_Skeleton.cpp
--- a/Engine/Source/Developer/AssetTools/Private/AssetTypeActions/AssetTypeActions_Skeleton.cpp
+++ b/Engine/Source/Developer/AssetTools/Private/AssetTypeActions/AssetTypeActions_Skeleton.cpp
@@ -430,15 +430,14 @@ void FAssetTypeActions_Skeleton::GetActions(const TArray<UObject*>& InObjects, F
 	auto Skeletons = GetTypedWeakObjectPtrs<USkeleton>(InObjects);
 
 	// create menu
-	Section.AddEntry(FToolMenuEntry::InitSubMenu(
-			NAME_None,
+	Section.AddSubMenu(
 			"CreateSkeletonSubmenu",
 			LOCTEXT("CreateSkeletonSubmenu", "Create"),
 			LOCTEXT("CreateSkeletonSubmenu_ToolTip", "Create assets for this skeleton"),
 			FNewMenuDelegate::CreateSP(this, &FAssetTypeActions_Skeleton::FillCreateMenu, Skeletons),
 			false, 
 			FSlateIcon(FEditorStyle::GetStyleSetName(), "Persona.AssetActions.CreateAnimAsset")
-			));
+			);
 
 	// only show if one is selected. It won't work since I changed the window to be normal window
 	if (Skeletons.Num() == 1)
diff --git a/Engine/Source/Developer/AssetTools/Private/AssetTypeActions/AssetTypeActions_StaticMesh.cpp b/Engine/Source/Developer/AssetTools/Private/AssetTypeActions/AssetTypeActions_StaticMesh.cpp
--- a/Engine/Source/Developer/AssetTools/Private/AssetTypeActions/AssetTypeActions_StaticMesh.cpp
+++ b/Engine/Source/Developer/AssetTools/Private/AssetTypeActions/AssetTypeActions_StaticMesh.cpp
@@ -42,15 +42,14 @@ void FAssetTypeActions_StaticMesh::GetActions(const TArray<UObject*>& InObjects,
 			);
 	}
 
-	Section.AddEntry(FToolMenuEntry::InitSubMenu(
-		NAME_None,
+	Section.AddSubMenu(
 		"StaticMesh_LODMenu",
 		NSLOCTEXT("AssetTypeActions_StaticMesh", "StaticMesh_LODMenu", "Level Of Detail"),
 		NSLOCTEXT("AssetTypeActions_StaticMesh", "StaticMesh_LODTooltip", "LOD Options and Tools"),
 		FNewMenuDelegate::CreateSP(this, &FAssetTypeActions_StaticMesh::GetLODMenu, Meshes),
 		false,
 		FSlateIcon(FEditorStyle::GetStyleSetName(), "ContentBrowser.AssetActions")
-		));
+		);
 
 	Section.AddMenuEntry(
 		"ObjectContext_ClearVertexColors",
diff --git a/Engine/Source/Developer/ToolMenus/Private/ToolMenu.cpp b/Engine/Source/Developer/ToolMenus/Private/ToolMenu.cpp
--- a/Engine/Source/Developer/ToolMenus/Private/ToolMenu.cpp
+++ b/Engine/Source/Developer/ToolMenus/Private/ToolMenu.cpp
@@ -166,7 +166,7 @@ UToolMenu* UToolMenu::AddSubMenuScript(const FName InOwner, const FName SectionN
 
 UToolMenu* UToolMenu::AddSubMenu(const FToolMenuOwner InOwner, const FName SectionName, const FName InName, const FText& InLabel, const FText& InToolTip)
 {
-	FToolMenuEntry Args = FToolMenuEntry::InitSubMenu(MenuName, InName, InLabel, InToolTip, FNewToolMenuChoice());
+	FToolMenuEntry Args = FToolMenuEntry::InitSubMenu(InName, InLabel, InToolTip, FNewToolMenuChoice());
 	Args.Owner = InOwner;
 	FindOrAddSection(SectionName).AddEntry(Args);
 	return UToolMenus::Get()->ExtendMenu(*(MenuName.ToString() + TEXT(".") + InName.ToString()));
diff --git a/Engine/Source/Developer/ToolMenus/Private/ToolMenuEntry.cpp b/Engine/Source/Developer/ToolMenus/Private/ToolMenuEntry.cpp
--- a/Engine/Source/Developer/ToolMenus/Private/ToolMenuEntry.cpp
+++ b/Engine/Source/Developer/ToolMenus/Private/ToolMenuEntry.cpp
@@ -103,7 +103,7 @@ FToolMenuEntry FToolMenuEntry::InitMenuEntry(const FName InName, const FToolUIAc
 	return Entry;
 }
 
-FToolMenuEntry FToolMenuEntry::InitSubMenu(const FName InParentMenu, const FName InName, const TAttribute<FText>& InLabel, const TAttribute<FText>& InToolTip, const FNewToolMenuChoice& InMakeMenu, const FToolUIActionChoice& InAction, const EUserInterfaceActionType InUserInterfaceActionType, bool bInOpenSubMenuOnClick, const TAttribute<FSlateIcon>& InIcon, const bool bInShouldCloseWindowAfterMenuSelection)
+FToolMenuEntry FToolMenuEntry::InitSubMenu(const FName InName, const TAttribute<FText>& InLabel, const TAttribute<FText>& InToolTip, const FNewToolMenuChoice& InMakeMenu, const FToolUIActionChoice& InAction, const EUserInterfaceActionType InUserInterfaceActionType, bool bInOpenSubMenuOnClick, const TAttribute<FSlateIcon>& InIcon, const bool bInShouldCloseWindowAfterMenuSelection)
 {
 	FToolMenuEntry Entry(UToolMenus::Get()->CurrentOwner(), InName, EMultiBlockType::MenuEntry);
 	Entry.Label = InLabel;
@@ -118,7 +118,7 @@ FToolMenuEntry FToolMenuEntry::InitSubMenu(const FName InParentMenu, const FName
 	return Entry;
 }
 
-FToolMenuEntry FToolMenuEntry::InitSubMenu(const FName InParentMenu, const FName InName, const TAttribute<FText>& InLabel, const TAttribute<FText>& InToolTip, const FNewToolMenuChoice& InMakeMenu, bool bInOpenSubMenuOnClick, const TAttribute<FSlateIcon>& InIcon, const bool bInShouldCloseWindowAfterMenuSelection)
+FToolMenuEntry FToolMenuEntry::InitSubMenu(const FName InName, const TAttribute<FText>& InLabel, const TAttribute<FText>& InToolTip, const FNewToolMenuChoice& InMakeMenu, bool bInOpenSubMenuOnClick, const TAttribute<FSlateIcon>& InIcon, const bool bInShouldCloseWindowAfterMenuSelection)
 {
 	FToolMenuEntry Entry(UToolMenus::Get()->CurrentOwner(), InName, EMultiBlockType::MenuEntry);
 	Entry.Label = InLabel;
@@ -131,7 +131,7 @@ FToolMenuEntry FToolMenuEntry::InitSubMenu(const FName InParentMenu, const FName
 	return Entry;
 }
 
-FToolMenuEntry FToolMenuEntry::InitSubMenu(const FName InParentMenu, const FName InName, const FToolUIActionChoice& InAction, const TSharedRef<SWidget>& InWidget, const FNewToolMenuChoice& InMakeMenu, bool bInShouldCloseWindowAfterMenuSelection)
+FToolMenuEntry FToolMenuEntry::InitSubMenu(const FName InName, const FToolUIActionChoice& InAction, const TSharedRef<SWidget>& InWidget, const FNewToolMenuChoice& InMakeMenu, bool bInShouldCloseWindowAfterMenuSelection)
 {
 	FToolMenuEntry Entry(UToolMenus::Get()->CurrentOwner(), InName, EMultiBlockType::MenuEntry);
 	Entry.Action = InAction;
diff --git a/Engine/Source/Developer/ToolMenus/Private/ToolMenuEntryScript.cpp b/Engine/Source/Developer/ToolMenus/Private/ToolMenuEntryScript.cpp
--- a/Engine/Source/Developer/ToolMenus/Private/ToolMenuEntryScript.cpp
+++ b/Engine/Source/Developer/ToolMenus/Private/ToolMenuEntryScript.cpp
@@ -117,7 +117,6 @@ void UToolMenuEntryScript::ToMenuEntry(FToolMenuEntry& Output)
 	if (Data.Advanced.bIsSubMenu)
 	{
 		Output = FToolMenuEntry::InitSubMenu(
-			Data.Menu,
 			Data.Name,
 			Data.Label,
 			Data.ToolTip,
diff --git a/Engine/Source/Developer/ToolMenus/Private/ToolMenuSection.cpp b/Engine/Source/Developer/ToolMenus/Private/ToolMenuSection.cpp
--- a/Engine/Source/Developer/ToolMenus/Private/ToolMenuSection.cpp
+++ b/Engine/Source/Developer/ToolMenus/Private/ToolMenuSection.cpp
@@ -103,14 +103,14 @@ FToolMenuEntry& FToolMenuSection::AddMenuSeparator(const FName InName)
 	return AddEntry(FToolMenuEntry::InitMenuSeparator(InName));
 }
 
-FToolMenuEntry& FToolMenuSection::AddSubMenu(const FName InParentMenu, const FName InName, const TAttribute<FText>& InLabel, const TAttribute<FText>& InToolTip, const FNewToolMenuChoice& InMakeMenu, bool bInOpenSubMenuOnClick, const TAttribute<FSlateIcon>& InIcon, const bool bShouldCloseWindowAfterMenuSelection)
+FToolMenuEntry& FToolMenuSection::AddSubMenu(const FName InName, const TAttribute<FText>& InLabel, const TAttribute<FText>& InToolTip, const FNewToolMenuChoice& InMakeMenu, const FToolUIActionChoice& InAction, const EUserInterfaceActionType InUserInterfaceActionType, bool bInOpenSubMenuOnClick, const TAttribute<FSlateIcon>& InIcon, const bool bInShouldCloseWindowAfterMenuSelection)
 {
-	return AddEntry(FToolMenuEntry::InitSubMenu(InParentMenu, InName, InLabel, InToolTip, InMakeMenu, bInOpenSubMenuOnClick, InIcon, bShouldCloseWindowAfterMenuSelection));
+	return AddEntry(FToolMenuEntry::InitSubMenu(InName, InLabel, InToolTip, InMakeMenu, InAction, InUserInterfaceActionType, bInOpenSubMenuOnClick, InIcon, bInShouldCloseWindowAfterMenuSelection));
 }
 
 FToolMenuEntry& FToolMenuSection::AddSubMenu(const FName InName, const TAttribute<FText>& InLabel, const TAttribute<FText>& InToolTip, const FNewToolMenuChoice& InMakeMenu, bool bInOpenSubMenuOnClick, const TAttribute<FSlateIcon>& InIcon, const bool bShouldCloseWindowAfterMenuSelection)
 {
-	return AddEntry(FToolMenuEntry::InitSubMenu(NAME_None, InName, InLabel, InToolTip, InMakeMenu, bInOpenSubMenuOnClick, InIcon, bShouldCloseWindowAfterMenuSelection));
+	return AddEntry(FToolMenuEntry::InitSubMenu(InName, InLabel, InToolTip, InMakeMenu, bInOpenSubMenuOnClick, InIcon, bShouldCloseWindowAfterMenuSelection));
 }
 
 int32 FToolMenuSection::IndexOfBlock(const FName InName) const
diff --git a/Engine/Source/Developer/ToolMenus/Public/ToolMenuEntry.h b/Engine/Source/Developer/ToolMenus/Public/ToolMenuEntry.h
--- a/Engine/Source/Developer/ToolMenus/Public/ToolMenuEntry.h
+++ b/Engine/Source/Developer/ToolMenus/Public/ToolMenuEntry.h
@@ -89,9 +89,9 @@ struct TOOLMENUS_API FToolMenuEntry
 	static FToolMenuEntry InitMenuEntryWithCommandList(const TSharedPtr<const FUICommandInfo>& InCommand, const TSharedPtr<const FUICommandList>& InCommandList, const TAttribute<FText>& InLabelOverride = TAttribute<FText>(), const TAttribute<FText>& InToolTipOverride = TAttribute<FText>(), const TAttribute<FSlateIcon>& InIconOverride = TAttribute<FSlateIcon>(), const FName InTutorialHighlightName = NAME_None, const FName InNameOverride = NAME_None);
 	static FToolMenuEntry InitMenuEntry(const FName InName, const FToolUIActionChoice& InAction, const TSharedRef<SWidget>& Widget);
 
-	static FToolMenuEntry InitSubMenu(const FName InParentMenu, const FName InName, const TAttribute<FText>& InLabel, const TAttribute<FText>& InToolTip, const FNewToolMenuChoice& InMakeMenu, const FToolUIActionChoice& InAction, const EUserInterfaceActionType InUserInterfaceActionType, bool bInOpenSubMenuOnClick = false, const TAttribute<FSlateIcon>& InIcon = TAttribute<FSlateIcon>(), const bool bShouldCloseWindowAfterMenuSelection = true);
-	static FToolMenuEntry InitSubMenu(const FName InParentMenu, const FName InName, const TAttribute<FText>& InLabel, const TAttribute<FText>& InToolTip, const FNewToolMenuChoice& InMakeMenu, bool bInOpenSubMenuOnClick = false, const TAttribute<FSlateIcon>& InIcon = TAttribute<FSlateIcon>(), const bool bShouldCloseWindowAfterMenuSelection = true);
-	static FToolMenuEntry InitSubMenu(const FName InParentMenu, const FName InName, const FToolUIActionChoice& InAction, const TSharedRef<SWidget>& InWidget, const FNewToolMenuChoice& InMakeMenu, bool bShouldCloseWindowAfterMenuSelection = true);
+	static FToolMenuEntry InitSubMenu(const FName InName, const TAttribute<FText>& InLabel, const TAttribute<FText>& InToolTip, const FNewToolMenuChoice& InMakeMenu, const FToolUIActionChoice& InAction, const EUserInterfaceActionType InUserInterfaceActionType, bool bInOpenSubMenuOnClick = false, const TAttribute<FSlateIcon>& InIcon = TAttribute<FSlateIcon>(), const bool bShouldCloseWindowAfterMenuSelection = true);
+	static FToolMenuEntry InitSubMenu(const FName InName, const TAttribute<FText>& InLabel, const TAttribute<FText>& InToolTip, const FNewToolMenuChoice& InMakeMenu, bool bInOpenSubMenuOnClick = false, const TAttribute<FSlateIcon>& InIcon = TAttribute<FSlateIcon>(), const bool bShouldCloseWindowAfterMenuSelection = true);
+	static FToolMenuEntry InitSubMenu(const FName InName, const FToolUIActionChoice& InAction, const TSharedRef<SWidget>& InWidget, const FNewToolMenuChoice& InMakeMenu, bool bShouldCloseWindowAfterMenuSelection = true);
 
 	static FToolMenuEntry InitToolBarButton(const FName InName, const FToolUIActionChoice& InAction, const TAttribute<FText>& InLabel = TAttribute<FText>(), const TAttribute<FText>& InToolTip = TAttribute<FText>(), const TAttribute<FSlateIcon>& InIcon = TAttribute<FSlateIcon>(), const EUserInterfaceActionType UserInterfaceActionType = EUserInterfaceActionType::Button, FName InTutorialHighlightName = NAME_None);
 	static FToolMenuEntry InitToolBarButton(const TSharedPtr< const FUICommandInfo >& InCommand, const TAttribute<FText>& InLabelOverride = TAttribute<FText>(), const TAttribute<FText>& InToolTipOverride = TAttribute<FText>(), const TAttribute<FSlateIcon>& InIconOverride = TAttribute<FSlateIcon>(), FName InTutorialHighlightName = NAME_None, const FName InNameOverride = NAME_None);
diff --git a/Engine/Source/Developer/ToolMenus/Public/ToolMenuSection.h b/Engine/Source/Developer/ToolMenus/Public/ToolMenuSection.h
--- a/Engine/Source/Developer/ToolMenus/Public/ToolMenuSection.h
+++ b/Engine/Source/Developer/ToolMenus/Public/ToolMenuSection.h
@@ -42,9 +42,9 @@ public:
 	FToolMenuEntry& AddDynamicEntry(const FName InName, const FNewToolMenuDelegateLegacy& InConstruct);
 	FToolMenuEntry& AddMenuSeparator(const FName InName);
 
-	FToolMenuEntry& AddSubMenu(const FName InParentMenu, const FName InName, const TAttribute<FText>& InLabel, const TAttribute<FText>& InToolTip, const FNewToolMenuChoice& InMakeMenu, bool bInOpenSubMenuOnClick = false, const TAttribute<FSlateIcon>& InIcon = TAttribute<FSlateIcon>(), const bool bShouldCloseWindowAfterMenuSelection = true);
+	FToolMenuEntry& AddSubMenu(const FName InName, const TAttribute<FText>& InLabel, const TAttribute<FText>& InToolTip, const FNewToolMenuChoice& InMakeMenu, const FToolUIActionChoice& InAction, const EUserInterfaceActionType InUserInterfaceActionType, bool bInOpenSubMenuOnClick = false, const TAttribute<FSlateIcon>& InIcon = TAttribute<FSlateIcon>(), const bool bShouldCloseWindowAfterMenuSelection = true);
 	FToolMenuEntry& AddSubMenu(const FName InName, const TAttribute<FText>& InLabel, const TAttribute<FText>& InToolTip, const FNewToolMenuChoice& InMakeMenu, bool bInOpenSubMenuOnClick = false, const TAttribute<FSlateIcon>& InIcon = TAttribute<FSlateIcon>(), const bool bShouldCloseWindowAfterMenuSelection = true);
-
+	
 	template <typename TContextType>
 	TContextType* FindContext() const
 	{
diff --git a/Engine/Source/Editor/AIGraph/Private/AIGraphSchema.cpp b/Engine/Source/Editor/AIGraph/Private/AIGraphSchema.cpp
--- a/Engine/Source/Editor/AIGraph/Private/AIGraphSchema.cpp
+++ b/Engine/Source/Editor/AIGraph/Private/AIGraphSchema.cpp
@@ -186,7 +186,6 @@ void UAIGraphSchema::GetContextMenuActions(class UToolMenu* Menu, class UGraphNo
 				if (Context->Pin->LinkedTo.Num() > 1)
 				{
 					Section.AddSubMenu(
-						Menu->GetMenuName(),
 						"BreakLinkTo",
 						LOCTEXT("BreakLinkTo", "Break Link To..."),
 						LOCTEXT("BreakSpecificLinks", "Break a specific link..."),
diff --git a/Engine/Source/Editor/AnimGraph/Private/AnimGraphNode_ModifyCurve.cpp b/Engine/Source/Editor/AnimGraph/Private/AnimGraphNode_ModifyCurve.cpp
--- a/Engine/Source/Editor/AnimGraph/Private/AnimGraphNode_ModifyCurve.cpp
+++ b/Engine/Source/Editor/AnimGraph/Private/AnimGraphNode_ModifyCurve.cpp
@@ -102,7 +102,6 @@ void UAnimGraphNode_ModifyCurve::GetNodeContextMenuActions(UToolMenu* Menu, UGra
 		if (GetCurvesToAdd().Num() > 0)
 		{
 			Section.AddSubMenu(
-				Menu->GetMenuName(),
 				"AddCurvePin",
 				LOCTEXT("AddCurvePin", "Add Curve Pin"),
 				LOCTEXT("AddCurvePinTooltip", "Add a new pin to drive a curve"),
@@ -113,7 +112,6 @@ void UAnimGraphNode_ModifyCurve::GetNodeContextMenuActions(UToolMenu* Menu, UGra
 		if (Node.CurveNames.Num() > 0)
 		{
 			Section.AddSubMenu(
-				Menu->GetMenuName(),
 				"RemoveCurvePin",
 				LOCTEXT("RemoveCurvePin", "Remove Curve Pin"),
 				LOCTEXT("RemoveCurvePinTooltip", "Remove a pin driving a curve"),
diff --git a/Engine/Source/Editor/AudioEditor/Private/AssetTypeActions/AssetTypeActions_SoundWave.cpp b/Engine/Source/Editor/AudioEditor/Private/AssetTypeActions/AssetTypeActions_SoundWave.cpp
--- a/Engine/Source/Editor/AudioEditor/Private/AssetTypeActions/AssetTypeActions_SoundWave.cpp
+++ b/Engine/Source/Editor/AudioEditor/Private/AssetTypeActions/AssetTypeActions_SoundWave.cpp
@@ -44,12 +44,11 @@ void FAssetTypeActions_SoundWave::GetActions(const TArray<UObject*>& InObjects,
 			)
 		);
 
-	Section.AddEntry(FToolMenuEntry::InitSubMenu(
-		NAME_None,
+	Section.AddSubMenu(
 		"SoundWave_CreateDialogue",
 		LOCTEXT("SoundWave_CreateDialogue", "Create Dialogue"),
 		LOCTEXT("SoundWave_CreateDialogueTooltip", "Creates a dialogue wave using this sound wave."),
-		FNewMenuDelegate::CreateSP(this, &FAssetTypeActions_SoundWave::FillVoiceMenu, SoundNodes)));
+		FNewMenuDelegate::CreateSP(this, &FAssetTypeActions_SoundWave::FillVoiceMenu, SoundNodes));
 }
 
 void FAssetTypeActions_SoundWave::GetResolvedSourceFilePaths(const TArray<UObject*>& TypeAssets, TArray<FString>& OutSourceFilePaths) const
diff --git a/Engine/Source/Editor/AudioEditor/Private/SoundClassGraphSchema.cpp b/Engine/Source/Editor/AudioEditor/Private/SoundClassGraphSchema.cpp
--- a/Engine/Source/Editor/AudioEditor/Private/SoundClassGraphSchema.cpp
+++ b/Engine/Source/Editor/AudioEditor/Private/SoundClassGraphSchema.cpp
@@ -105,7 +105,6 @@ void USoundClassGraphSchema::GetContextMenuActions(UToolMenu* Menu, UGraphNodeCo
 				if(Context->Pin->LinkedTo.Num() > 1)
 				{
 					Section.AddSubMenu(
-						Menu->GetMenuName(),
 						"BreakLinkTo",
 						LOCTEXT("BreakLinkTo", "Break Link To..." ),
 						LOCTEXT("BreakSpecificLinks", "Break a specific link..." ),
diff --git a/Engine/Source/Editor/AudioEditor/Private/SoundSubmixGraphSchema.cpp b/Engine/Source/Editor/AudioEditor/Private/SoundSubmixGraphSchema.cpp
--- a/Engine/Source/Editor/AudioEditor/Private/SoundSubmixGraphSchema.cpp
+++ b/Engine/Source/Editor/AudioEditor/Private/SoundSubmixGraphSchema.cpp
@@ -105,12 +105,11 @@ void USoundSubmixGraphSchema::GetContextMenuActions(UToolMenu* Menu, UGraphNodeC
 				// add sub menu for break link to
 				if(InGraphPin->LinkedTo.Num() > 1)
 				{
-					Section.AddEntry(FToolMenuEntry::InitSubMenu(
-						Menu->GetMenuName(),
+					Section.AddSubMenu(
 						"BreakLinkTo",
 						LOCTEXT("BreakLinkTo", "Break Link To..." ),
 						LOCTEXT("BreakSpecificLinks", "Break a specific link..." ),
-						FNewToolMenuDelegate::CreateUObject((USoundSubmixGraphSchema*const)this, &USoundSubmixGraphSchema::GetBreakLinkToSubMenuActions, SectionName, const_cast<UEdGraphPin*>(InGraphPin))));
+						FNewToolMenuDelegate::CreateUObject((USoundSubmixGraphSchema*const)this, &USoundSubmixGraphSchema::GetBreakLinkToSubMenuActions, SectionName, const_cast<UEdGraphPin*>(InGraphPin)));
 				}
 				else
 				{
diff --git a/Engine/Source/Editor/BlueprintGraph/Private/EdGraphSchema_K2.cpp b/Engine/Source/Editor/BlueprintGraph/Private/EdGraphSchema_K2.cpp
--- a/Engine/Source/Editor/BlueprintGraph/Private/EdGraphSchema_K2.cpp
+++ b/Engine/Source/Editor/BlueprintGraph/Private/EdGraphSchema_K2.cpp
@@ -1450,26 +1450,23 @@ void UEdGraphSchema_K2::GetContextMenuActions(UToolMenu* Menu, UGraphNodeContext
 
 					if(InGraphPin->LinkedTo.Num() > 1)
 					{
-						Section.AddEntry(FToolMenuEntry::InitSubMenu(
-							Menu->GetMenuName(),
+						Section.AddSubMenu(
 							"BreakLinkTo",
 							LOCTEXT("BreakLinkTo", "Break Link To..."),
 							LOCTEXT("BreakSpecificLinks", "Break a specific link..."),
-							FNewToolMenuDelegate::CreateUObject( (UEdGraphSchema_K2*const)this, &UEdGraphSchema_K2::GetBreakLinkToSubMenuActions, const_cast<UEdGraphPin*>(InGraphPin))));
+							FNewToolMenuDelegate::CreateUObject( (UEdGraphSchema_K2*const)this, &UEdGraphSchema_K2::GetBreakLinkToSubMenuActions, const_cast<UEdGraphPin*>(InGraphPin)));
 
-						Section.AddEntry(FToolMenuEntry::InitSubMenu(
-							Menu->GetMenuName(),
+						Section.AddSubMenu(
 							"JumpToConnection",
 							LOCTEXT("JumpToConnection", "Jump to Connection..."),
 							LOCTEXT("JumpToSpecificConnection", "Jump to specific connection..."),
-							FNewToolMenuDelegate::CreateUObject( (UEdGraphSchema_K2*const)this, &UEdGraphSchema_K2::GetJumpToConnectionSubMenuActions, const_cast<UEdGraphPin*>(InGraphPin))));
+							FNewToolMenuDelegate::CreateUObject( (UEdGraphSchema_K2*const)this, &UEdGraphSchema_K2::GetJumpToConnectionSubMenuActions, const_cast<UEdGraphPin*>(InGraphPin)));
 
-						Section.AddEntry(FToolMenuEntry::InitSubMenu(
-							Menu->GetMenuName(),
+						Section.AddSubMenu(
 							"StraightenConnection",
 							LOCTEXT("StraightenConnection", "Straighten Connection To..."),
 							LOCTEXT("StraightenConnection_Tip", "Straighten a specific connection"),
-							FNewToolMenuDelegate::CreateUObject( this, &UEdGraphSchema_K2::GetStraightenConnectionToSubMenuActions, const_cast<UEdGraphPin*>(InGraphPin))));
+							FNewToolMenuDelegate::CreateUObject( this, &UEdGraphSchema_K2::GetStraightenConnectionToSubMenuActions, const_cast<UEdGraphPin*>(InGraphPin)));
 					}
 					else
 					{
@@ -1674,7 +1671,7 @@ void UEdGraphSchema_K2::GetContextMenuActions(UToolMenu* Menu, UGraphNodeContext
 						Section.AddMenuEntry( FGraphEditorCommands::Get().PromoteSelectionToMacro );
 					}
 
-					Section.AddEntry(FToolMenuEntry::InitSubMenu(Menu->GetMenuName(), "Alignment", LOCTEXT("AlignmentHeader", "Alignment"), FText(), FNewToolMenuDelegate::CreateLambda([](UToolMenu* AlignmentMenu)
+					Section.AddSubMenu("Alignment", LOCTEXT("AlignmentHeader", "Alignment"), FText(), FNewToolMenuDelegate::CreateLambda([](UToolMenu* AlignmentMenu)
 					{
 						{
 							FToolMenuSection& InSection = AlignmentMenu->AddSection("EdGraphSchemaAlignment", LOCTEXT("AlignHeader", "Align"));
@@ -1692,7 +1689,7 @@ void UEdGraphSchema_K2::GetContextMenuActions(UToolMenu* Menu, UGraphNodeContext
 							InSection.AddMenuEntry(FGraphEditorCommands::Get().DistributeNodesHorizontally);
 							InSection.AddMenuEntry(FGraphEditorCommands::Get().DistributeNodesVertically);
 						}
-					})));
+					}));
 				}
 			}
 
@@ -1717,8 +1714,7 @@ void UEdGraphSchema_K2::GetContextMenuActions(UToolMenu* Menu, UGraphNodeContext
 									const FUIAction* SubMenuUIAction = Menu->Context.GetActionForCommand(FGraphEditorCommands::Get().EnableNodes);
 									if(ensure(SubMenuUIAction))
 									{
-										Section.AddEntry(FToolMenuEntry::InitSubMenu(
-											Menu->GetMenuName(),
+										Section.AddSubMenu(
 											"EnableCompileSubMenu",
 											LOCTEXT("EnableCompileSubMenu", "Enable Compile"),
 											LOCTEXT("EnableCompileSubMenuToolTip", "Options to enable selected node(s) for compile."),
@@ -1735,7 +1731,7 @@ void UEdGraphSchema_K2::GetContextMenuActions(UToolMenu* Menu, UGraphNodeContext
 													LOCTEXT("EnableCompileDevelopmentOnlyToolTip", "Compile selected node(s) for development only."));
 											}),
 											*SubMenuUIAction,
-											FGraphEditorCommands::Get().EnableNodes->GetUserInterfaceType()));
+											FGraphEditorCommands::Get().EnableNodes->GetUserInterfaceType());
 									}
 								}
 							}
@@ -1972,13 +1968,12 @@ void UEdGraphSchema_K2::GetReplaceVariableMenu(FToolMenuSection& Section, const
 					ReplaceVariableWithTooltip = LOCTEXT("ReplaceMissingVariableWithToolTip", "Variable '{0}' does not exist, replace with another variable?");
 				}
 
-				Section.AddEntry(FToolMenuEntry::InitSubMenu(
-					NAME_None,
+				Section.AddSubMenu(
 					"ReplaceVariableWith",
 					FText::Format( LOCTEXT("ReplaceVariableWith", "Replace variable '{0}' with..."), Variable->GetVarNameText()),
 					FText::Format( ReplaceVariableWithTooltip, Variable->GetVarNameText()),
 					FNewToolMenuDelegate::CreateStatic( &UEdGraphSchema_K2::GetReplaceVariableMenu,
-						const_cast<UK2Node_Variable*>(Variable), InOwnerBlueprint, bInReplaceExistingVariable)));
+						const_cast<UK2Node_Variable*>(Variable), InOwnerBlueprint, bInReplaceExistingVariable));
 			}
 		}
 	}
diff --git a/Engine/Source/Editor/ContentBrowser/Private/AssetContextMenu.cpp b/Engine/Source/Editor/ContentBrowser/Private/AssetContextMenu.cpp
--- a/Engine/Source/Editor/ContentBrowser/Private/AssetContextMenu.cpp
+++ b/Engine/Source/Editor/ContentBrowser/Private/AssetContextMenu.cpp
@@ -371,19 +371,17 @@ bool FAssetContextMenu::AddImportedAssetMenuOptions(UToolMenu* Menu)
 			//Reimport Menu
 			if (ValidSelectedAssetCount == 1 && SourceFileLabels.Num() > 1)
 			{
-				Section.AddEntry(FToolMenuEntry::InitSubMenu(
-					Menu->GetMenuName(),
+				Section.AddSubMenu(
 					"Reimport",
 					LOCTEXT("Reimport", "Reimport"),
 					LOCTEXT("ReimportEmptyTooltip", ""),
-					FNewToolMenuDelegate::CreateLambda(CreateSubMenu, false) ));
+					FNewToolMenuDelegate::CreateLambda(CreateSubMenu, false) );
 				//With new file
-				Section.AddEntry(FToolMenuEntry::InitSubMenu(
-					Menu->GetMenuName(),
+				Section.AddSubMenu(
 					"ReimportWithNewFile",
 					LOCTEXT("ReimportWithNewFile", "Reimport With New File"),
 					LOCTEXT("ReimportEmptyTooltip", ""),
-					FNewToolMenuDelegate::CreateLambda(CreateSubMenu, true)));
+					FNewToolMenuDelegate::CreateLambda(CreateSubMenu, true));
 			}
 			else
 			{
@@ -494,8 +492,7 @@ bool FAssetContextMenu::AddCommonMenuOptions(UToolMenu* Menu)
 				);
 
 			// Asset Actions sub-menu
-			Section.AddEntry(FToolMenuEntry::InitSubMenu(
-				Menu->GetMenuName(),
+			Section.AddSubMenu(
 				"AssetActionsSubMenu",
 				LOCTEXT("AssetActionsSubMenuLabel", "Asset Actions"),
 				LOCTEXT("AssetActionsSubMenuToolTip", "Other asset actions"),
@@ -507,13 +504,12 @@ bool FAssetContextMenu::AddCommonMenuOptions(UToolMenu* Menu)
 				EUserInterfaceActionType::Button,
 				false, 
 				FSlateIcon(FEditorStyle::GetStyleSetName(), "ContentBrowser.AssetActions")
-				));
+				);
 
 			if (NumClassItems == 0)
 			{
 				// Asset Localization sub-menu
-				Section.AddEntry(FToolMenuEntry::InitSubMenu(
-					Menu->GetMenuName(),
+				Section.AddSubMenu(
 					"LocalizationSubMenu",
 					LOCTEXT("LocalizationSubMenuLabel", "Asset Localization"),
 					LOCTEXT("LocalizationSubMenuToolTip", "Manage the localization of this asset"),
@@ -522,7 +518,7 @@ bool FAssetContextMenu::AddCommonMenuOptions(UToolMenu* Menu)
 					EUserInterfaceActionType::Button,
 					false,
 					FSlateIcon(FEditorStyle::GetStyleSetName(), "ContentBrowser.AssetLocalization")
-					));
+					);
 			}
 		}
 	}
@@ -752,13 +748,12 @@ void FAssetContextMenu::MakeAssetActionsSubMenu(UToolMenu* Menu)
 				FUIAction( FExecuteAction::CreateSP(this, &FAssetContextMenu::ExecuteAssignChunkID) )
 				);
 
-			Section.AddEntry(FToolMenuEntry::InitSubMenu(
-				Menu->GetMenuName(),
+			Section.AddSubMenu(
 				"RemoveAssetFromChunk",
 				LOCTEXT("RemoveAssetFromChunk", "Remove from Chunk..."),
 				LOCTEXT("RemoveAssetFromChunkTooltip", "Removed an asset from a Chunk it's assigned to."),
 				FNewToolMenuDelegate::CreateRaw(this, &FAssetContextMenu::MakeChunkIDListMenu)
-				));
+				);
 
 			Section.AddMenuEntry(
 				"RemoveAllChunkAssignments",
@@ -1028,8 +1023,7 @@ void FAssetContextMenu::MakeAssetLocalizationSubMenu(UToolMenu* Menu)
 	{
 		FToolMenuSection& Section = Menu->AddSection("ManageLocalizedAsset", LOCTEXT("ManageLocalizedAssetHeading", "Manage Localized Asset"));
 		{
-			Section.AddEntry(FToolMenuEntry::InitSubMenu(
-				Menu->GetMenuName(),
+			Section.AddSubMenu(
 				"CreateLocalizedAsset",
 				LOCTEXT("CreateLocalizedAsset", "Create Localized Asset"),
 				LOCTEXT("CreateLocalizedAssetTooltip", "Create a new localized asset."),
@@ -1038,7 +1032,7 @@ void FAssetContextMenu::MakeAssetLocalizationSubMenu(UToolMenu* Menu)
 				EUserInterfaceActionType::Button,
 				false,
 				FSlateIcon(FEditorStyle::GetStyleSetName(), "ContentBrowser.AssetActions.Duplicate")
-				));
+				);
 
 			int32 NumLocalizedAssets = 0;
 			for (const FLocalizedAssetsState& LocalizedAssetsStateForCulture : LocalizedAssetsState)
@@ -1048,8 +1042,7 @@ void FAssetContextMenu::MakeAssetLocalizationSubMenu(UToolMenu* Menu)
 
 			if (NumLocalizedAssets > 0)
 			{
-				Section.AddEntry(FToolMenuEntry::InitSubMenu(
-					Menu->GetMenuName(),
+				Section.AddSubMenu(
 					"ShowLocalizedAsset",
 					LOCTEXT("ShowLocalizedAsset", "Show Localized Asset"),
 					LOCTEXT("ShowLocalizedAssetTooltip", "Show the localized asset in the Content Browser."),
@@ -1058,10 +1051,9 @@ void FAssetContextMenu::MakeAssetLocalizationSubMenu(UToolMenu* Menu)
 					EUserInterfaceActionType::Button,
 					false,
 					FSlateIcon(FEditorStyle::GetStyleSetName(), "SystemWideCommands.FindInContentBrowser")
-					));
+					);
 
-				Section.AddEntry(FToolMenuEntry::InitSubMenu(
-					Menu->GetMenuName(),
+				Section.AddSubMenu(
 					"EditLocalizedAsset",
 					LOCTEXT("EditLocalizedAsset", "Edit Localized Asset"),
 					LOCTEXT("EditLocalizedAssetTooltip", "Edit the localized asset."),
@@ -1070,7 +1062,7 @@ void FAssetContextMenu::MakeAssetLocalizationSubMenu(UToolMenu* Menu)
 					EUserInterfaceActionType::Button,
 					false,
 					FSlateIcon(FEditorStyle::GetStyleSetName(), "ContentBrowser.AssetActions.Edit")
-					));
+					);
 			}
 		}
 	}
@@ -1380,8 +1372,7 @@ bool FAssetContextMenu::AddSourceControlMenuOptions(UToolMenu* Menu)
 	if ( ISourceControlModule::Get().IsEnabled() )
 	{
 		// SCC sub menu
-		Section.AddEntry(FToolMenuEntry::InitSubMenu(
-			Menu->GetMenuName(),
+		Section.AddSubMenu(
 			"SourceControlSubMenu",
 			LOCTEXT("SourceControlSubMenuLabel", "Source Control"),
 			LOCTEXT("SourceControlSubMenuToolTip", "Source control actions."),
@@ -1393,7 +1384,7 @@ bool FAssetContextMenu::AddSourceControlMenuOptions(UToolMenu* Menu)
 			EUserInterfaceActionType::Button,
 			false,
 			FSlateIcon(FEditorStyle::GetStyleSetName(), "SourceControl.StatusIcon.On")
-			));
+			);
 	}
 	else
 	{
@@ -1602,8 +1593,7 @@ bool FAssetContextMenu::AddCollectionMenuOptions(UToolMenu* Menu)
 
 				if (AvailableChildCollections.Num() > 0)
 				{
-					Section.AddEntry(FToolMenuEntry::InitSubMenu(
-						SubMenu->GetMenuName(),
+					Section.AddSubMenu(
 						NAME_None,
 						FText::FromName(AvailableCollection.Name), 
 						FText::GetEmpty(), 
@@ -1616,7 +1606,7 @@ bool FAssetContextMenu::AddCollectionMenuOptions(UToolMenu* Menu)
 						EUserInterfaceActionType::ToggleButton,
 						false,
 						FSlateIcon(FEditorStyle::GetStyleSetName(), ECollectionShareType::GetIconStyleName(AvailableCollection.Type))
-						));
+						);
 				}
 				else
 				{
@@ -1674,13 +1664,12 @@ bool FAssetContextMenu::AddCollectionMenuOptions(UToolMenu* Menu)
 		TSharedRef<FCollectionAssetManagement> QuickAssetManagement = MakeShareable(new FCollectionAssetManagement());
 		QuickAssetManagement->SetCurrentAssets(SelectedAssets);
 
-		Section.AddEntry(FToolMenuEntry::InitSubMenu(
-			Menu->GetMenuName(),
+		Section.AddSubMenu(
 			"ManageCollections",
 			LOCTEXT("ManageCollections", "Manage Collections"),
 			LOCTEXT("ManageCollections_ToolTip", "Manage the collections that the selected asset(s) belong to."),
 			FNewToolMenuDelegate::CreateStatic(&FManageCollectionsContextMenu::CreateManageCollectionsSubMenu, QuickAssetManagement)
-			));
+			);
 
 		bHasAddedItems = true;
 	}
diff --git a/Engine/Source/Editor/Kismet/Private/SSCSEditor.cpp b/Engine/Source/Editor/Kismet/Private/SSCSEditor.cpp
--- a/Engine/Source/Editor/Kismet/Private/SSCSEditor.cpp
+++ b/Engine/Source/Editor/Kismet/Private/SSCSEditor.cpp
@@ -3990,14 +3990,13 @@ void SSCSEditor::PopulateContextMenu(UToolMenu* Menu)
 							// Build an event submenu if we can generate events
 							if( FBlueprintEditorUtils::CanClassGenerateEvents( SelectedClass ))
 							{
-								BlueprintSCSSection.AddEntry(FToolMenuEntry::InitSubMenu(
-									Menu->GetMenuName(),
+								BlueprintSCSSection.AddSubMenu(
 									"AddEventSubMenu",
 									LOCTEXT("AddEventSubMenu", "Add Event"), 
 									LOCTEXT("ActtionsSubMenu_ToolTip", "Add Event"), 
 									FNewMenuDelegate::CreateStatic( &SSCSEditor::BuildMenuEventsSection,
 									GetBlueprint(), SelectedClass, FCanExecuteAction::CreateSP(this, &SSCSEditor::IsEditingAllowed),
-									FGetSelectedObjectsDelegate::CreateSP(this, &SSCSEditor::GetSelectedItemsForContextMenu))));
+									FGetSelectedObjectsDelegate::CreateSP(this, &SSCSEditor::GetSelectedItemsForContextMenu)));
 							}
 						}
 					}					
diff --git a/Engine/Source/Editor/LevelEditor/Private/LevelEditorContextMenu.cpp b/Engine/Source/Editor/LevelEditor/Private/LevelEditorContextMenu.cpp
--- a/Engine/Source/Editor/LevelEditor/Private/LevelEditorContextMenu.cpp
+++ b/Engine/Source/Editor/LevelEditor/Private/LevelEditorContextMenu.cpp
@@ -343,36 +343,32 @@ void FLevelEditorContextMenu::RegisterActorContextMenu()
 			FToolMenuSection& Section = InMenu->AddSection("ActorSelectVisibilityLevels");
 
 			// Add a sub-menu for "Select"
-			Section.AddEntry(FToolMenuEntry::InitSubMenu(
-				InMenu->MenuName,
+			Section.AddSubMenu(
 				"SelectSubMenu",
 				LOCTEXT("SelectSubMenu", "Select"),
 				LOCTEXT("SelectSubMenu_ToolTip", "Opens the actor selection menu"),
-				FNewToolMenuDelegate::CreateStatic(&FLevelEditorContextMenuImpl::FillSelectActorMenu)));
+				FNewToolMenuDelegate::CreateStatic(&FLevelEditorContextMenuImpl::FillSelectActorMenu));
 
-			Section.AddEntry(FToolMenuEntry::InitSubMenu(
-				InMenu->MenuName,
+			Section.AddSubMenu(
 				"EditSubMenu",
 				LOCTEXT("EditSubMenu", "Edit"),
 				FText::GetEmpty(),
-				FNewToolMenuDelegate::CreateStatic(&FLevelEditorContextMenuImpl::FillEditMenu)));
+				FNewToolMenuDelegate::CreateStatic(&FLevelEditorContextMenuImpl::FillEditMenu));
 
-			Section.AddEntry(FToolMenuEntry::InitSubMenu(
-				InMenu->MenuName,
+			Section.AddSubMenu(
 				"VisibilitySubMenu",
 				LOCTEXT("VisibilitySubMenu", "Visibility"),
 				LOCTEXT("VisibilitySubMenu_ToolTip", "Selected actor visibility options"),
-				FNewToolMenuDelegate::CreateStatic(&FLevelEditorContextMenuImpl::FillActorVisibilityMenu)));
+				FNewToolMenuDelegate::CreateStatic(&FLevelEditorContextMenuImpl::FillActorVisibilityMenu));
 
 			// Build the menu for grouping actors
 			BuildGroupMenu(InMenu, SelectionInfo);
 
-			Section.AddEntry(FToolMenuEntry::InitSubMenu(
-				InMenu->MenuName,
+			Section.AddSubMenu(
 				"LevelSubMenu",
 				LOCTEXT("LevelSubMenu", "Level"),
 				LOCTEXT("LevelSubMenu_ToolTip", "Options for interacting with this actor's level"),
-				FNewToolMenuDelegate::CreateStatic(&FLevelEditorContextMenuImpl::FillActorLevelMenu)));
+				FNewToolMenuDelegate::CreateStatic(&FLevelEditorContextMenuImpl::FillActorLevelMenu));
 		}
 
 		if (LevelEditorContext->ContextType == ELevelEditorMenuContext::Viewport)
@@ -407,33 +403,30 @@ void FLevelEditorContextMenu::RegisterActorContextMenu()
 					Section.AddMenuEntry(FLevelEditorCommands::Get().DetachFromParent);
 				}
 
-				Section.AddEntry(FToolMenuEntry::InitSubMenu(
-					InMenu->GetMenuName(),
+				Section.AddSubMenu(
 					"ActorAttachToSubMenu",
 					LOCTEXT("ActorAttachToSubMenu", "Attach To"),
 					LOCTEXT("ActorAttachToSubMenu_ToolTip", "Attach Actor as child"),
-					FNewToolMenuDelegate::CreateStatic(&FLevelEditorContextMenuImpl::FillActorMenu)));
+					FNewToolMenuDelegate::CreateStatic(&FLevelEditorContextMenuImpl::FillActorMenu));
 			}
 
 			// Add a heading for "Movement" if an actor is selected
 			if (GEditor->GetSelectedActorIterator())
 			{
 				// Add a sub-menu for "Transform"
-				Section.AddEntry(FToolMenuEntry::InitSubMenu(
-					InMenu->GetMenuName(),
+				Section.AddSubMenu(
 					"TransformSubMenu",
 					LOCTEXT("TransformSubMenu", "Transform"),
 					LOCTEXT("TransformSubMenu_ToolTip", "Actor transform utils"),
-					FNewToolMenuDelegate::CreateStatic(&FLevelEditorContextMenuImpl::FillTransformMenu)));
+					FNewToolMenuDelegate::CreateStatic(&FLevelEditorContextMenuImpl::FillTransformMenu));
 			}
 
 			// Add a sub-menu for "Pivot"
-			Section.AddEntry(FToolMenuEntry::InitSubMenu(
-				InMenu->MenuName,
+			Section.AddSubMenu(
 				"PivotSubMenu",
 				LOCTEXT("PivotSubMenu", "Pivot"),
 				LOCTEXT("PivotSubMenu_ToolTip", "Actor pivoting utils"),
-				FNewToolMenuDelegate::CreateStatic(&FLevelEditorContextMenuImpl::FillPivotMenu)));
+				FNewToolMenuDelegate::CreateStatic(&FLevelEditorContextMenuImpl::FillPivotMenu));
 		}
 
 		FLevelScriptEventMenuHelper::FillLevelBlueprintEventsMenu(InMenu, SelectedActors);
@@ -457,13 +450,12 @@ void FLevelEditorContextMenu::RegisterSceneOutlinerContextMenu()
 			if (SceneOutlinerPtr.IsValid())
 			{
 				FToolMenuSection& Section = InMenu->AddSection("SelectVisibilityLevels");
-				Section.AddEntry(FToolMenuEntry::InitSubMenu(
-					InMenu->MenuName,
+				Section.AddSubMenu(
 					"EditSubMenu",
 					LOCTEXT("EditSubMenu", "Edit"),
 					FText::GetEmpty(),
 					FNewToolMenuDelegate::CreateStatic(&FLevelEditorContextMenuImpl::FillEditMenu)
-				));
+				);
 			}
 		}
 	}));
@@ -729,12 +721,11 @@ void FLevelEditorContextMenu::BuildGroupMenu(UToolMenu* Menu, const FSelectedAct
 		
 		if( bNeedGroupSubMenu )
 		{
-			Section.AddEntry(FToolMenuEntry::InitSubMenu(
-				Menu->GetMenuName(),
+			Section.AddSubMenu(
 				"GroupMenu",
 				LOCTEXT("GroupMenu", "Groups"),
 				LOCTEXT("GroupMenu_ToolTip", "Opens the actor grouping menu"),
-				FNewToolMenuDelegate::CreateStatic( &FLevelEditorContextMenuImpl::FillGroupMenu)));
+				FNewToolMenuDelegate::CreateStatic( &FLevelEditorContextMenuImpl::FillGroupMenu));
 		}
 	}
 }
@@ -1006,12 +997,11 @@ void FLevelEditorContextMenuImpl::FillTransformMenu(UToolMenu* Menu)
 	{
 		{
 			FToolMenuSection& Section = Menu->AddSection("TransformSnapAlign");
-			Section.AddEntry(FToolMenuEntry::InitSubMenu(
-				Menu->GetMenuName(),
+			Section.AddSubMenu(
 				"SnapAlignSubMenu",
 				LOCTEXT("SnapAlignSubMenu", "Snap/Align"), 
 				LOCTEXT("SnapAlignSubMenu_ToolTip", "Actor snap/align utils"),
-				FNewToolMenuDelegate::CreateStatic(&FLevelEditorContextMenuImpl::FillSnapAlignMenu)));
+				FNewToolMenuDelegate::CreateStatic(&FLevelEditorContextMenuImpl::FillSnapAlignMenu));
 		}
 
 		{
@@ -1265,8 +1255,7 @@ void FLevelScriptEventMenuHelper::FillLevelBlueprintEventsMenu(UToolMenu* Menu,
 
 			if (bAnyEventExists)
 			{
-				Section.AddEntry(FToolMenuEntry::InitSubMenu(
-					Menu->GetMenuName(),
+				Section.AddSubMenu(
 					"JumpEventSubMenu",
 					LOCTEXT("JumpEventSubMenu", "Jump to Event"),
 					FText::GetEmpty(),
@@ -1274,13 +1263,12 @@ void FLevelScriptEventMenuHelper::FillLevelBlueprintEventsMenu(UToolMenu* Menu,
 					, ActorPtr
 					, true
 					, false
-					, true)));
+					, true));
 			}
 
 			if (bAnyEventCanBeAdded)
 			{
-				Section.AddEntry(FToolMenuEntry::InitSubMenu(
-					Menu->GetMenuName(),
+				Section.AddSubMenu(
 					"AddEventSubMenu",
 					LOCTEXT("AddEventSubMenu", "Add Event"),
 					FText::GetEmpty(),
@@ -1288,7 +1276,7 @@ void FLevelScriptEventMenuHelper::FillLevelBlueprintEventsMenu(UToolMenu* Menu,
 					, ActorPtr
 					, false
 					, true
-					, true)));
+					, true));
 			}
 		}
 	}
diff --git a/Engine/Source/Editor/LevelEditor/Private/LevelEditorCreateActorMenu.cpp b/Engine/Source/Editor/LevelEditor/Private/LevelEditorCreateActorMenu.cpp
--- a/Engine/Source/Editor/LevelEditor/Private/LevelEditorCreateActorMenu.cpp
+++ b/Engine/Source/Editor/LevelEditor/Private/LevelEditorCreateActorMenu.cpp
@@ -361,21 +361,19 @@ void LevelEditorCreateActorMenu::FillAddReplaceViewportContextMenuSections(UTool
 	{
 		{
 			FToolMenuSection& Section = Menu->AddSection("ActorType");
-			Section.AddEntry(FToolMenuEntry::InitSubMenu(
-				Menu->GetMenuName(),
+			Section.AddSubMenu(
 				"LevelViewportContextMenu",
 				NSLOCTEXT("LevelViewportContextMenu", "AddActorHeading", "Place Actor") , 
 				NSLOCTEXT("LevelViewportContextMenu", "AddActorMenu_ToolTip", "Templates for adding a new actor to the world"),
-				FNewToolMenuDelegate::CreateStatic(&LevelEditorCreateActorMenu::FillAddReplaceActorMenu, EActorCreateMode::Add)));
+				FNewToolMenuDelegate::CreateStatic(&LevelEditorCreateActorMenu::FillAddReplaceActorMenu, EActorCreateMode::Add));
 
 			if ( CanReplaceActors() )
 			{
-				Section.AddEntry(FToolMenuEntry::InitSubMenu(
-					Menu->GetMenuName(),
+				Section.AddSubMenu(
 					"LevelViewportContextMenu",
 					NSLOCTEXT("LevelViewportContextMenu", "ReplaceActorHeading", "Replace Selected Actors with") , 
 					NSLOCTEXT("LevelViewportContextMenu", "ReplaceActorMenu_ToolTip", "Templates for replacing selected with new actors in the world"),
-					FNewToolMenuDelegate::CreateStatic(&LevelEditorCreateActorMenu::FillAddReplaceActorMenu, EActorCreateMode::Replace)));
+					FNewToolMenuDelegate::CreateStatic(&LevelEditorCreateActorMenu::FillAddReplaceActorMenu, EActorCreateMode::Replace));
 			}
 		}
 	}
@@ -391,7 +389,6 @@ void LevelEditorCreateActorMenu::FillAddReplaceViewportContextMenuSections(UTool
 			FUIAction Action( FExecuteAction::CreateStatic( &FLevelEditorActionCallbacks::AddActor_Clicked, AssetMenuOptions[0].FactoryToUse,  AssetMenuOptions[0].AssetData, false ) );
 			TSharedRef<SWidget> Widget = SNew(SAssetMenuEntry, TargetAssetData, AssetMenuOptions);
 			Section.AddEntry(FToolMenuEntry::InitSubMenu(
-				Menu->GetMenuName(),
 				"AddActor",
 				Action,
 				Widget,
@@ -406,7 +403,6 @@ void LevelEditorCreateActorMenu::FillAddReplaceViewportContextMenuSections(UTool
 				FUIAction Action( FExecuteAction::CreateStatic( &FLevelEditorActionCallbacks::ReplaceActors_Clicked, AssetMenuOptions[0].FactoryToUse,  AssetMenuOptions[0].AssetData ) );
 				TSharedRef<SWidget> Widget = SNew(SAssetMenuEntry, TargetAssetData, AssetMenuOptions);
 				Section.AddEntry(FToolMenuEntry::InitSubMenu(
-					Menu->GetMenuName(),
 					"ReplaceActor",
 					Action,
 					Widget,
diff --git a/Engine/Source/Editor/LevelEditor/Private/LevelEditorMenu.cpp b/Engine/Source/Editor/LevelEditor/Private/LevelEditorMenu.cpp
--- a/Engine/Source/Editor/LevelEditor/Private/LevelEditorMenu.cpp
+++ b/Engine/Source/Editor/LevelEditor/Private/LevelEditorMenu.cpp
@@ -125,15 +125,14 @@ void FLevelEditorMenu::RegisterLevelEditorMenus()
 					const FMainMRUFavoritesList& RecentsAndFavorites = *MainFrameModule.GetMRUFavoritesList();
 					if (RecentsAndFavorites.GetNumItems() > 0)
 					{
-						InSection.AddEntry(FToolMenuEntry::InitSubMenu(
-							"LevelEditor.MainMenu.File",
+						InSection.AddSubMenu(
 							"FavoriteLevelsSubMenu",
 							LOCTEXT("FavoriteLevelsSubMenu", "Favorite Levels"),
 							LOCTEXT("RecentLevelsSubMenu_ToolTip", "Select a level to load"),
 							FNewToolMenuDelegate::CreateStatic(&FFavoriteLevelMenu::MakeFavoriteLevelMenu),
 							false,
 							FSlateIcon(FEditorStyle::GetStyleSetName(), "MainFrame.FavoriteLevels")
-						));
+						);
 					}
 				}));
 			}
@@ -169,15 +168,14 @@ void FLevelEditorMenu::RegisterLevelEditorMenus()
 					const FMainMRUFavoritesList& RecentsAndFavorites = *MainFrameModule.GetMRUFavoritesList();
 					if (RecentsAndFavorites.GetNumItems() > 0)
 					{
-						InSection.AddEntry(FToolMenuEntry::InitSubMenu(
-							"LevelEditor.MainMenu.File",
+						InSection.AddSubMenu(
 							"RecentLevelsSubMenu",
 							LOCTEXT("RecentLevelsSubMenu", "Recent Levels"),
 							LOCTEXT("RecentLevelsSubMenu_ToolTip", "Select a level to load"),
 							FNewToolMenuDelegate::CreateStatic(&FRecentLevelMenu::MakeRecentLevelMenu),
 							false,
 							FSlateIcon(FEditorStyle::GetStyleSetName(), "MainFrame.RecentLevels")
-						));
+						);
 					}
 				}));
 			}
diff --git a/Engine/Source/Editor/MainFrame/Private/Menus/MainMenu.cpp b/Engine/Source/Editor/MainFrame/Private/Menus/MainMenu.cpp
--- a/Engine/Source/Editor/MainFrame/Private/Menus/MainMenu.cpp
+++ b/Engine/Source/Editor/MainFrame/Private/Menus/MainMenu.cpp
@@ -151,25 +151,23 @@ void FMainMenu::RegisterEditMenu()
 
 		if (GetDefault<UEditorStyleSettings>()->bExpandConfigurationMenus)
 		{
-			Section.AddEntry(FToolMenuEntry::InitSubMenu(
-				EditMenu->GetMenuName(),
+			Section.AddSubMenu(
 				"EditorPreferencesSubMenu",
 				LOCTEXT("EditorPreferencesSubMenuLabel", "Editor Preferences"),
 				LOCTEXT("EditorPreferencesSubMenuToolTip", "Configure the behavior and features of this Editor"),
 				FNewToolMenuDelegate::CreateStatic(&FSettingsMenu::MakeMenu, FName("Editor")),
 				false,
 				FSlateIcon(FEditorStyle::GetStyleSetName(), "EditorPreferences.TabIcon")
-			));
+			);
 
-			Section.AddEntry(FToolMenuEntry::InitSubMenu(
-				EditMenu->GetMenuName(),
+			Section.AddSubMenu(
 				"ProjectSettingsSubMenu",
 				LOCTEXT("ProjectSettingsSubMenuLabel", "Project Settings"),
 				LOCTEXT("ProjectSettingsSubMenuToolTip", "Change the settings of the currently loaded project"),
 				FNewToolMenuDelegate::CreateStatic(&FSettingsMenu::MakeMenu, FName("Project")),
 				false,
 				FSlateIcon(FEditorStyle::GetStyleSetName(), "ProjectSettings.TabIcon")
-			));
+			);
 		}
 		else
 		{
@@ -433,13 +431,12 @@ void FMainMenu::RegisterFileProjectMenu()
 		FText::Format(LOCTEXT("AddCodeToProjectTooltip", "Adds C++ code to the project. The code can only be compiled if you have {0} installed."), ShortIDEName)
 	);
 
-	Section.AddEntry(FToolMenuEntry::InitSubMenu(
-		"MainFrame.MainTabMenu.File",
+	Section.AddSubMenu(
 		"PackageProject",
 		LOCTEXT("PackageProjectSubMenuLabel", "Package Project"),
 		LOCTEXT("PackageProjectSubMenuToolTip", "Compile, cook and package your project and its content for distribution."),
 		FNewMenuDelegate::CreateStatic( &FPackageProjectMenu::MakeMenu ), false, FSlateIcon(FEditorStyle::GetStyleSetName(), "MainFrame.PackageProject")
-	));
+	);
 
 	/*
 	MenuBuilder.AddMenuEntry( FMainFrameCommands::Get().LocalizeProject,
@@ -512,15 +509,14 @@ void FMainMenu::RegisterRecentFileAndExitMenuItems()
 		FToolMenuSection& Section = MainTabFileMenu->AddSection("FileRecentFiles");
 		if (GetDefault<UEditorStyleSettings>()->bShowProjectMenus && FMainFrameActionCallbacks::ProjectNames.Num() > 0)
 		{
-			Section.AddEntry(FToolMenuEntry::InitSubMenu(
-				"MainFrame.MainTabMenu.File",
+			Section.AddSubMenu(
 				"RecentProjects",
 				LOCTEXT("SwitchProjectSubMenu", "Recent Projects"),
 				LOCTEXT("SwitchProjectSubMenu_ToolTip", "Select a project to switch to"),
 				FNewToolMenuDelegate::CreateStatic(&FRecentProjectsMenu::MakeMenu),
 				false,
 				FSlateIcon(FEditorStyle::GetStyleSetName(), "MainFrame.RecentProjects")
-			));
+			);
 		}
 	}
 
diff --git a/Engine/Source/Editor/SceneOutliner/Private/SSceneOutliner.cpp b/Engine/Source/Editor/SceneOutliner/Private/SSceneOutliner.cpp
--- a/Engine/Source/Editor/SceneOutliner/Private/SSceneOutliner.cpp
+++ b/Engine/Source/Editor/SceneOutliner/Private/SSceneOutliner.cpp
@@ -1779,12 +1779,11 @@ namespace SceneOutliner
 						// If we've only got folders selected, show the selection and edit sub menus
 						if (Context->NumSelectedItems > 0 && Context->NumSelectedFolders == Context->NumSelectedItems)
 						{
-							InMenu->FindOrAddSection("Section").AddEntry(FToolMenuEntry::InitSubMenu(
-								InMenu->MenuName,
+							InMenu->FindOrAddSection("Section").AddSubMenu(
 								"SelectSubMenu",
 								LOCTEXT("SelectSubmenu", "Select"),
 								LOCTEXT("SelectSubmenu_Tooltip", "Select the contents of the current selection"),
-								FNewToolMenuDelegate::CreateSP(SceneOutliner, &SSceneOutliner::FillSelectionSubMenu)));
+								FNewToolMenuDelegate::CreateSP(SceneOutliner, &SSceneOutliner::FillSelectionSubMenu));
 						}
 					}
 				}
@@ -1801,12 +1800,11 @@ namespace SceneOutliner
 					// Can't move worlds or level blueprints
 					if (Context->bShowParentTree && Context->NumSelectedItems > 0 && Context->NumWorldsSelected == 0 && Context->SceneOutliner.IsValid())
 					{
-						Section.AddEntry(FToolMenuEntry::InitSubMenu(
-							InMenu->MenuName,
+						Section.AddSubMenu(
 							"MoveActorsTo",
 							LOCTEXT("MoveActorsTo", "Move To"),
 							LOCTEXT("MoveActorsTo_Tooltip", "Move selection to another folder"),
-							FNewToolMenuDelegate::CreateSP(Context->SceneOutliner.Pin().Get(), &SSceneOutliner::FillFoldersSubMenu)));
+							FNewToolMenuDelegate::CreateSP(Context->SceneOutliner.Pin().Get(), &SSceneOutliner::FillFoldersSubMenu));
 					}
 				}
 			}));
diff --git a/Engine/Source/Editor/SceneOutliner/Public/SceneOutlinerFilters.h b/Engine/Source/Editor/SceneOutliner/Public/SceneOutlinerFilters.h
--- a/Engine/Source/Editor/SceneOutliner/Public/SceneOutlinerFilters.h
+++ b/Engine/Source/Editor/SceneOutliner/Public/SceneOutlinerFilters.h
@@ -12,6 +12,8 @@
 #include "ComponentTreeItem.h"
 #include "SubComponentTreeItem.h"
 
+class FMenuBuilder;
+
 namespace SceneOutliner
 {
 
diff --git a/Engine/Source/Editor/UnrealEd/Private/MaterialGraphSchema.cpp b/Engine/Source/Editor/UnrealEd/Private/MaterialGraphSchema.cpp
--- a/Engine/Source/Editor/UnrealEd/Private/MaterialGraphSchema.cpp
+++ b/Engine/Source/Editor/UnrealEd/Private/MaterialGraphSchema.cpp
@@ -419,12 +419,11 @@ void UMaterialGraphSchema::GetContextMenuActions(UToolMenu* Menu, UGraphNodeCont
 				// add sub menu for break link to
 				if(InGraphPin->LinkedTo.Num() > 1)
 				{
-					Section.AddEntry(FToolMenuEntry::InitSubMenu(
-						Menu->GetMenuName(),
+					Section.AddSubMenu(
 						"BreakLinkTo",
 						LOCTEXT("BreakLinkTo", "Break Link To..." ),
 						LOCTEXT("BreakSpecificLinks", "Break a specific link..." ),
-						FNewToolMenuDelegate::CreateUObject(this, &UMaterialGraphSchema::GetBreakLinkToSubMenuActions, const_cast<UEdGraphPin*>(InGraphPin))));
+						FNewToolMenuDelegate::CreateUObject(this, &UMaterialGraphSchema::GetBreakLinkToSubMenuActions, const_cast<UEdGraphPin*>(InGraphPin)));
 				}
 				else
 				{
diff --git a/Engine/Source/Editor/WorldBrowser/Private/SWorldHierarchyImpl.cpp b/Engine/Source/Editor/WorldBrowser/Private/SWorldHierarchyImpl.cpp
--- a/Engine/Source/Editor/WorldBrowser/Private/SWorldHierarchyImpl.cpp
+++ b/Engine/Source/Editor/WorldBrowser/Private/SWorldHierarchyImpl.cpp
@@ -508,25 +508,23 @@ TSharedPtr<SWidget> SWorldHierarchyImpl::ConstructLevelContextMenu() const
 			if (bAllSelectedItemsCanMove && FLevelFolders::IsAvailable())
 			{
 				FToolMenuSection& Section = Menu->AddSection("Section");
-				Section.AddEntry(FToolMenuEntry::InitSubMenu(
-					Menu->MenuName,
+				Section.AddSubMenu(
 					"MoveSelectionTo",
 					LOCTEXT("MoveSelectionTo", "Move To"),
 					LOCTEXT("MoveSelectionTo_Tooltip", "Move selection to another folder"),
 					FNewMenuDelegate::CreateSP(const_cast<SWorldHierarchyImpl*>(this), &SWorldHierarchyImpl::FillFoldersSubmenu)
-				));
+				);
 			}
 
 			if (bOnlyFoldersSelected)
 			{
 				FToolMenuSection& Section = Menu->AddSection("Section");
-				Section.AddEntry(FToolMenuEntry::InitSubMenu(
-					Menu->MenuName,
+				Section.AddSubMenu(
 					"SelectSubmenu",
 					LOCTEXT("SelectSubmenu", "Select"),
 					LOCTEXT("SelectSubmenu_Tooltip", "Select child items of the current selection"),
 					FNewMenuDelegate::CreateSP(const_cast<SWorldHierarchyImpl*>(this), &SWorldHierarchyImpl::FillSelectionSubmenu)
-				));
+				);
 			}
 		}
 
