Fix crash after clicking edit in context menu of leve assetl in context browser
#rnx

diff --git a/Engine/Plugins/Runtime/SoundUtilities/Source/SoundUtilitiesEditor/Private/SoundWaveAssetActionExtender.cpp b/Engine/Plugins/Runtime/SoundUtilities/Source/SoundUtilitiesEditor/Private/SoundWaveAssetActionExtender.cpp
--- a/Engine/Plugins/Runtime/SoundUtilities/Source/SoundUtilitiesEditor/Private/SoundWaveAssetActionExtender.cpp
+++ b/Engine/Plugins/Runtime/SoundUtilities/Source/SoundUtilitiesEditor/Private/SoundWaveAssetActionExtender.cpp
@@ -58,7 +58,7 @@ void FSoundWaveAssetActionExtender::ExecuteCreateSimpleSound(const FToolMenuCont
 		// Create the factory used to generate the asset
 		USoundSimpleFactory* Factory = NewObject<USoundSimpleFactory>();
 		Factory->SoundWaves.Reset();
-		for (UObject* Object : Context->SelectedObjects)
+		for (const TWeakObjectPtr<UObject>& Object : Context->SelectedObjects)
 		{
 			if (USoundWave* Wave = Cast<USoundWave>(Object))
 			{
diff --git a/Engine/Plugins/Runtime/TimeSynth/Source/TimeSynthEditor/Private/TimeSynthSoundWaveAssetActionExtender.cpp b/Engine/Plugins/Runtime/TimeSynth/Source/TimeSynthEditor/Private/TimeSynthSoundWaveAssetActionExtender.cpp
--- a/Engine/Plugins/Runtime/TimeSynth/Source/TimeSynthEditor/Private/TimeSynthSoundWaveAssetActionExtender.cpp
+++ b/Engine/Plugins/Runtime/TimeSynth/Source/TimeSynthEditor/Private/TimeSynthSoundWaveAssetActionExtender.cpp
@@ -66,7 +66,7 @@ void FTimeSynthSoundWaveAssetActionExtender::ExecuteCreateTimeSyncClip(const FTo
 	FAssetToolsModule& AssetToolsModule = FModuleManager::Get().LoadModuleChecked<FAssetToolsModule>("AssetTools");
 
 	UTimeSynthClipFactory* Factory = NewObject<UTimeSynthClipFactory>();
-	for (UObject* Object : Context->SelectedObjects)
+	for (const TWeakObjectPtr<UObject>& Object : Context->SelectedObjects)
 	{
 		USoundWave* SoundWave = Cast<USoundWave>(Object);
 
@@ -103,7 +103,7 @@ void FTimeSynthSoundWaveAssetActionExtender::ExecuteCreateTimeSyncClipSet(const
 	AssetToolsModule.Get().CreateUniqueAssetName(SoundWave->GetOutermost()->GetName(), DefaultSuffix, PackagePath, Name);
 
 	UTimeSynthClipFactory* Factory = NewObject<UTimeSynthClipFactory>();
-	Factory->SoundWaves = FObjectEditorUtils::GetTypedWeakObjectPtrs<USoundWave>(Context->SelectedObjects);
+	Factory->SoundWaves = FObjectEditorUtils::GetTypedWeakObjectPtrs<USoundWave>(Context->GetSelectedObjects());
 
 	AssetToolsModule.Get().CreateAsset(Name, FPackageName::GetLongPackagePath(PackagePath), UTimeSynthClip::StaticClass(), Factory);
 }
diff --git a/Engine/Source/Editor/ContentBrowser/Private/AssetContextMenu.cpp b/Engine/Source/Editor/ContentBrowser/Private/AssetContextMenu.cpp
--- a/Engine/Source/Editor/ContentBrowser/Private/AssetContextMenu.cpp
+++ b/Engine/Source/Editor/ContentBrowser/Private/AssetContextMenu.cpp
@@ -142,6 +142,8 @@ TSharedRef<SWidget> FAssetContextMenu::MakeContextMenu(const TArray<FAssetData>&
 	static const FName BaseMenuName("ContentBrowser.AssetContextMenu");
 	RegisterContextMenu(BaseMenuName);
 
+	TArray<UObject*> SelectedObjects;
+
 	// Create menu hierarchy based on class hierarchy
 	TArray<UToolMenu*> MenuHierarchy;
 	FName MenuName = BaseMenuName;
@@ -154,13 +156,15 @@ TSharedRef<SWidget> FAssetContextMenu::MakeContextMenu(const TArray<FAssetData>&
 		}
 
 		ContextObject->SelectedObjects.Reset();
-		if (ContentBrowserUtils::LoadAssetsIfNeeded(ObjectPaths, ContextObject->SelectedObjects) && ContextObject->SelectedObjects.Num() > 0)
+		if (ContentBrowserUtils::LoadAssetsIfNeeded(ObjectPaths, SelectedObjects) && SelectedObjects.Num() > 0)
 		{
+			ContextObject->SelectedObjects.Append(SelectedObjects);
+
 			// Find common class for selected objects
-			UClass* CommonClass = ContextObject->SelectedObjects[0]->GetClass();
-			for (int32 ObjIdx = 1; ObjIdx < ContextObject->SelectedObjects.Num(); ++ObjIdx)
+			UClass* CommonClass = SelectedObjects[0]->GetClass();
+			for (int32 ObjIdx = 1; ObjIdx < SelectedObjects.Num(); ++ObjIdx)
 			{
-				while (!ContextObject->SelectedObjects[ObjIdx]->IsA(CommonClass))
+				while (!SelectedObjects[ObjIdx]->IsA(CommonClass))
 				{
 					CommonClass = CommonClass->GetSuperClass();
 				}
@@ -171,7 +175,7 @@ TSharedRef<SWidget> FAssetContextMenu::MakeContextMenu(const TArray<FAssetData>&
 			FAssetToolsModule& AssetToolsModule = FModuleManager::LoadModuleChecked<FAssetToolsModule>("AssetTools");
 			TSharedPtr<IAssetTypeActions> CommonAssetTypeActions = AssetToolsModule.Get().GetAssetTypeActionsForClass(ContextObject->CommonClass).Pin();
 			UClass* AssetTypeActionsClass = nullptr;
-			if (CommonAssetTypeActions.IsValid() && CommonAssetTypeActions->HasActions(ContextObject->SelectedObjects))
+			if (CommonAssetTypeActions.IsValid() && CommonAssetTypeActions->HasActions(SelectedObjects))
 			{
 				ContextObject->CommonAssetTypeActions = CommonAssetTypeActions;
 				AssetTypeActionsClass = CommonAssetTypeActions->GetSupportedClass();
@@ -215,7 +219,7 @@ TSharedRef<SWidget> FAssetContextMenu::MakeContextMenu(const TArray<FAssetData>&
 	if ( SelectedAssets.Num() )
 	{
 		// Add any type-specific context menu options
-		AddAssetTypeMenuOptions(Menu, ContextObject->SelectedObjects);
+		AddAssetTypeMenuOptions(Menu, SelectedObjects);
 
 		// Add imported asset context menu options
 		AddImportedAssetMenuOptions(Menu);
@@ -255,7 +259,7 @@ void FAssetContextMenu::RegisterContextMenu(const FName MenuName)
 			UContentBrowserAssetContextMenuContext* Context = InSection.FindContext<UContentBrowserAssetContextMenuContext>();
 			if (Context && Context->CommonAssetTypeActions.IsValid())
 			{
-				Context->CommonAssetTypeActions.Pin()->GetActions(Context->SelectedObjects, InSection);
+				Context->CommonAssetTypeActions.Pin()->GetActions(Context->GetSelectedObjects(), InSection);
 			}
 		}));
 
@@ -264,7 +268,7 @@ void FAssetContextMenu::RegisterContextMenu(const FName MenuName)
 			UContentBrowserAssetContextMenuContext* Context = InMenu->FindContext<UContentBrowserAssetContextMenuContext>();
 			if (Context && Context->CommonAssetTypeActions.IsValid())
 			{
-				Context->CommonAssetTypeActions.Pin()->GetActions(Context->SelectedObjects, MenuBuilder);
+				Context->CommonAssetTypeActions.Pin()->GetActions(Context->GetSelectedObjects(), MenuBuilder);
 			}
 		}));
 	}
diff --git a/Engine/Source/Editor/ContentBrowser/Public/ContentBrowserMenuContexts.h b/Engine/Source/Editor/ContentBrowser/Public/ContentBrowserMenuContexts.h
--- a/Engine/Source/Editor/ContentBrowser/Public/ContentBrowserMenuContexts.h
+++ b/Engine/Source/Editor/ContentBrowser/Public/ContentBrowserMenuContexts.h
@@ -23,8 +23,21 @@ public:
 	TWeakPtr<IAssetTypeActions> CommonAssetTypeActions;
 	
 	UPROPERTY()
-	TArray<UObject*> SelectedObjects;
+	TArray<TWeakObjectPtr<UObject>> SelectedObjects;
 
 	UPROPERTY()
 	UClass* CommonClass;
+
+	UFUNCTION(BlueprintCallable, Category="Tool Menus")
+	TArray<UObject*> GetSelectedObjects() const
+	{
+		TArray<UObject*> Result;
+		Result.Reserve(SelectedObjects.Num());
+		for (const TWeakObjectPtr<UObject>& Object : SelectedObjects)
+		{
+			Result.Add(Object.Get());
+		}
+		return Result;
+	}
 };
+
