Fix label of content browser context menu's Get Asset Actions section.
#rnx

diff --git a/Engine/Source/Developer/ToolMenus/Private/ToolMenus.cpp b/Engine/Source/Developer/ToolMenus/Private/ToolMenus.cpp
--- a/Engine/Source/Developer/ToolMenus/Private/ToolMenus.cpp
+++ b/Engine/Source/Developer/ToolMenus/Private/ToolMenus.cpp
@@ -377,6 +377,14 @@ void UToolMenus::AssembleMenu(UToolMenu* GeneratedMenu, const UToolMenu* Other)
 					continue;
 				}
 			}
+			else
+			{
+				// Allow overriding label
+				if (!Section->Label.IsSet() && RemainingSection.Label.IsSet())
+				{
+					Section->Label = RemainingSection.Label;
+				}
+			}
 
 			AssembleMenuSection(GeneratedMenu, Other, Section, RemainingSection);
 			RemainingSections.RemoveAt(i);
diff --git a/Engine/Source/Editor/ContentBrowser/Private/AssetContextMenu.cpp b/Engine/Source/Editor/ContentBrowser/Private/AssetContextMenu.cpp
--- a/Engine/Source/Editor/ContentBrowser/Private/AssetContextMenu.cpp
+++ b/Engine/Source/Editor/ContentBrowser/Private/AssetContextMenu.cpp
@@ -281,8 +281,6 @@ void FAssetContextMenu::AddMenuOptions(UToolMenu* InMenu)
 
 	// Add source control options
 	AddSourceControlMenuOptions(InMenu);
-
-	UE_LOG(LogTemp, Log, TEXT("AddMenuOptions"));
 }
 
 void FAssetContextMenu::SetSelectedAssets(const TArray<FAssetData>& InSelectedAssets)
@@ -1361,18 +1359,22 @@ bool FAssetContextMenu::AddAssetTypeMenuOptions(UToolMenu* Menu, bool bHasObject
 	{
 		// Label "GetAssetActions" section
 		UContentBrowserAssetContextMenuContext* Context = Menu->FindContext<UContentBrowserAssetContextMenuContext>();
-		FToolMenuSection* Section = Menu->FindSection("GetAssetActions");
-		if (Section && Section->Blocks.Num() > 0 && Context)
+		if (Context)
 		{
+			FToolMenuSection& Section = Menu->FindOrAddSection("GetAssetActions");
 			if (Context->CommonAssetTypeActions.IsValid())
 			{
-				Section->Label = FText::Format(NSLOCTEXT("AssetTools", "AssetSpecificOptionsMenuHeading", "{0} Actions"), Context->CommonAssetTypeActions.Pin()->GetName());
+				Section.Label = FText::Format(NSLOCTEXT("AssetTools", "AssetSpecificOptionsMenuHeading", "{0} Actions"), Context->CommonAssetTypeActions.Pin()->GetName());
 			}
 			else if (Context->CommonClass)
 			{
-				Section->Label = FText::Format(NSLOCTEXT("AssetTools", "AssetSpecificOptionsMenuHeading", "{0} Actions"), FText::FromName(Context->CommonClass->GetFName()));
+				Section.Label = FText::Format(NSLOCTEXT("AssetTools", "AssetSpecificOptionsMenuHeading", "{0} Actions"), FText::FromName(Context->CommonClass->GetFName()));
 			}
-			
+			else
+			{
+				Section.Label = FText::Format(NSLOCTEXT("AssetTools", "AssetSpecificOptionsMenuHeading", "{0} Actions"), FText::FromString(TEXT("Asset")));
+			}
+
 			bAnyTypeOptions = true;
 		}
 	}
