Collection UX improvements

Revamped collection UI:
 - Uses a more modern design.
 - Adds the asset count to the collection view.
 - Adds the collections more prominently in asset tooltips.

Revamped sources panel UI:
 - Added a unified search for all aspects of the sources panel (paths, favorites, and collections).
 - Added a switcher for the paths and collections views, which auto-switches as you browse to paths or collections.
 - Collections are now visible by default, rather than hidden behind a View Option.

Improved Content Browser auto-complete:
 - Now lists collection and meta-data tag names, as well as the old class names.
 - Tokenizes the search terms to allow auto-complete on partial statements.
 - Note: Involves some breaking API changes to SAssetSearchBox to allow more complex filtering and auto-complete rules.

Added scripting API support:
 - Collections can now be created, destroyed, and manipulated via Python or Editor Utilities (editor-only, see the AssetTagsSubsystem API).
 - Collections can be queried both in-editor and at runtime (if cooked).

Added support for optionally "cooking" collections for use at runtime:
 - This only works for static collections, and adds their information to the asset registry meta-data.
 - See the AssetRegistry settings in BaseEngine.ini (bTagAllCollections, CollectionsToIncludeAsTags, and CollectionsToExcludeAsTags).
 - No collections are cooked by default.

diff --git a/Engine/Config/BaseEngine.ini b/Engine/Config/BaseEngine.ini
--- a/Engine/Config/BaseEngine.ini
+++ b/Engine/Config/BaseEngine.ini
@@ -1685,6 +1685,11 @@ bFilterAssetDataWithNoTags=false
 bFilterDependenciesWithNoTags=false
 bFilterSearchableNames=true
 
+; Control which collections are added as asset registry tags during cook (none by default)
+; If bTagAllCollections is true then you can exclude specfic collections via the CollectionsToExcludeAsTags arrray
+; If bTagAllCollections is false then you can include specfic collections via the CollectionsToIncludeAsTags arrray
+bTagAllCollections=false
+
 [AutomationTesting]
 ImportTestPath=../../Content/EditorAutomation/
 ImportTestPackagePath=/Engine/Content/EditorAutomation
diff --git a/Engine/Plugins/Runtime/AssetTags/AssetTags.uplugin b/Engine/Plugins/Runtime/AssetTags/AssetTags.uplugin
new file mode 100644
--- /dev/null
+++ b/Engine/Plugins/Runtime/AssetTags/AssetTags.uplugin
@@ -0,0 +1,25 @@
+{
+	"FileVersion" : 3,
+	"Version" : 1,
+	"VersionName" : "1.0",
+	"FriendlyName" : "Asset Tags",
+	"Description" : "Provides high-level management and access to asset tags and collections for runtime and editor scripting.",
+	"Category" : "Other",
+	"CreatedBy" : "Epic Games, Inc.",
+	"CreatedByURL" : "http://epicgames.com",
+	"DocsURL" : "",
+	"MarketplaceURL" : "",
+	"SupportURL" : "",
+	"EnabledByDefault" : true,
+	"CanContainContent" : false,
+	"IsBetaVersion" : false,
+	"Installed" : false,
+	"Modules" :
+	[
+		{
+			"Name" : "AssetTags",
+			"Type" : "Runtime",
+			"LoadingPhase": "Default"
+ 		}
+	]
+}
diff --git a/Engine/Plugins/Runtime/AssetTags/Source/AssetTags/AssetTags.Build.cs b/Engine/Plugins/Runtime/AssetTags/Source/AssetTags/AssetTags.Build.cs
new file mode 100644
--- /dev/null
+++ b/Engine/Plugins/Runtime/AssetTags/Source/AssetTags/AssetTags.Build.cs
@@ -0,0 +1,40 @@
+// Copyright 1998-2019 Epic Games, Inc. All Rights Reserved.
+
+namespace UnrealBuildTool.Rules
+{
+	public class AssetTags : ModuleRules
+	{
+		public AssetTags(ReadOnlyTargetRules Target) : base(Target)
+		{
+			PublicDependencyModuleNames.AddRange(
+				new string[] {
+					"Core",
+					"CoreUObject",
+					"Engine",
+				}
+			);
+
+			PublicIncludePathModuleNames.AddRange(
+				new string[] {
+					"AssetRegistry",
+				}
+			);
+
+			PrivateDependencyModuleNames.AddRange(
+				new string[] {
+					"AssetRegistry",
+				}
+			);
+
+			if (Target.bBuildEditor)
+			{
+				PrivateDependencyModuleNames.AddRange(
+					new string[] {
+						"AssetTagsEditor",
+						"CollectionManager",
+					}
+				);
+			}
+		}
+	}
+}
diff --git a/Engine/Plugins/Runtime/AssetTags/Source/AssetTags/Private/AssetTagsModule.cpp b/Engine/Plugins/Runtime/AssetTags/Source/AssetTags/Private/AssetTagsModule.cpp
new file mode 100644
--- /dev/null
+++ b/Engine/Plugins/Runtime/AssetTags/Source/AssetTags/Private/AssetTagsModule.cpp
@@ -0,0 +1,5 @@
+// Copyright 1998-2019 Epic Games, Inc. All Rights Reserved.
+
+#include "Modules/ModuleManager.h"
+
+IMPLEMENT_MODULE(FDefaultModuleImpl, AssetTags);
diff --git a/Engine/Plugins/Runtime/AssetTags/Source/AssetTags/Private/AssetTagsSubsystem.cpp b/Engine/Plugins/Runtime/AssetTags/Source/AssetTags/Private/AssetTagsSubsystem.cpp
new file mode 100644
--- /dev/null
+++ b/Engine/Plugins/Runtime/AssetTags/Source/AssetTags/Private/AssetTagsSubsystem.cpp
@@ -0,0 +1,474 @@
+// Copyright 1998-2019 Epic Games, Inc. All Rights Reserved.
+
+#include "AssetTagsSubsystem.h"
+
+#include "IAssetRegistry.h"
+#include "AssetRegistryState.h"
+#include "AssetRegistryModule.h"
+
+#if WITH_EDITOR
+#include "ICollectionManager.h"
+#include "CollectionManagerModule.h"
+#endif	// WITH_EDITOR
+
+DEFINE_LOG_CATEGORY_STATIC(LogAssetTags, Log, All);
+
+#if WITH_EDITOR
+
+namespace AssetTagsSubsystemUtil
+{
+
+void LogLastCollectionManagerError(ICollectionManager& CollectionManager)
+{
+	UE_LOG(LogAssetTags, Warning, TEXT("%s"), *CollectionManager.GetLastError().ToString());
+}
+
+void LogLastCollectionManagerError()
+{
+	ICollectionManager& CollectionManager = FCollectionManagerModule::GetModule().Get();
+	LogLastCollectionManagerError(CollectionManager);
+}
+
+FCollectionNameType FindCollectionByName(ICollectionManager& CollectionManager, const FName Name)
+{
+	TArray<FCollectionNameType> CollectionNamesAndTypes;
+	CollectionManager.GetCollections(Name, CollectionNamesAndTypes);
+	
+	if (CollectionNamesAndTypes.Num() == 0)
+	{
+		UE_LOG(LogAssetTags, Warning, TEXT("No collection found called '%s'"), *Name.ToString());
+		return FCollectionNameType(NAME_None, ECollectionShareType::CST_All);
+	}
+	else if (CollectionNamesAndTypes.Num() > 1)
+	{
+		UE_LOG(LogAssetTags, Warning, TEXT("%d collections found called '%s'; ambiguous result"), CollectionNamesAndTypes.Num(), *Name.ToString());
+		return FCollectionNameType(NAME_None, ECollectionShareType::CST_All);
+	}
+	
+	return CollectionNamesAndTypes[0];
+}
+
+FCollectionNameType FindCollectionByName(const FName Name)
+{
+	ICollectionManager& CollectionManager = FCollectionManagerModule::GetModule().Get();
+	return FindCollectionByName(CollectionManager, Name);
+}
+
+}	// namespace AssetTagsSubsystemUtil
+
+bool UAssetTagsSubsystem::CreateCollection(const FName Name, const ECollectionScriptingShareType ShareType)
+{
+	ICollectionManager& CollectionManager = FCollectionManagerModule::GetModule().Get();
+
+	if (!CollectionManager.IsValidCollectionName(Name.ToString(), ECollectionShareType::CST_All))
+	{
+		AssetTagsSubsystemUtil::LogLastCollectionManagerError(CollectionManager);
+		return false;
+	}
+
+	static_assert((int32)ECollectionShareType::CST_Local == (int32)ECollectionScriptingShareType::Local + 1, "ECollectionShareType::CST_Local is expected to be ECollectionScriptingShareType::Local + 1");
+	static_assert((int32)ECollectionShareType::CST_Private == (int32)ECollectionScriptingShareType::Private + 1, "ECollectionShareType::CST_Private is expected to be ECollectionScriptingShareType::Private + 1");
+	static_assert((int32)ECollectionShareType::CST_Shared == (int32)ECollectionScriptingShareType::Shared + 1, "ECollectionShareType::CST_Shared is expected to be ECollectionScriptingShareType::Shared + 1");
+
+	if (!CollectionManager.CreateCollection(Name, (ECollectionShareType::Type)((int32)ShareType + 1), ECollectionStorageMode::Static))
+	{
+		AssetTagsSubsystemUtil::LogLastCollectionManagerError(CollectionManager);
+		return false;
+	}
+
+	return true;
+}
+
+bool UAssetTagsSubsystem::DestroyCollection(const FName Name)
+{
+	ICollectionManager& CollectionManager = FCollectionManagerModule::GetModule().Get();
+
+	const FCollectionNameType ResolvedNameAndType = AssetTagsSubsystemUtil::FindCollectionByName(CollectionManager, Name);
+	if (ResolvedNameAndType.Name.IsNone())
+	{
+		return false;
+	}
+
+	if (!CollectionManager.DestroyCollection(ResolvedNameAndType.Name, ResolvedNameAndType.Type))
+	{
+		AssetTagsSubsystemUtil::LogLastCollectionManagerError(CollectionManager);
+		return false;
+	}
+
+	return true;
+}
+
+bool UAssetTagsSubsystem::RenameCollection(const FName Name, const FName NewName)
+{
+	ICollectionManager& CollectionManager = FCollectionManagerModule::GetModule().Get();
+
+	const FCollectionNameType ResolvedNameAndType = AssetTagsSubsystemUtil::FindCollectionByName(CollectionManager, Name);
+	if (ResolvedNameAndType.Name.IsNone())
+	{
+		return false;
+	}
+
+	if (!CollectionManager.IsValidCollectionName(NewName.ToString(), ECollectionShareType::CST_All))
+	{
+		AssetTagsSubsystemUtil::LogLastCollectionManagerError(CollectionManager);
+		return false;
+	}
+
+	if (!CollectionManager.RenameCollection(ResolvedNameAndType.Name, ResolvedNameAndType.Type, NewName, ResolvedNameAndType.Type))
+	{
+		AssetTagsSubsystemUtil::LogLastCollectionManagerError(CollectionManager);
+		return false;
+	}
+
+	return true;
+}
+
+bool UAssetTagsSubsystem::ReparentCollection(const FName Name, const FName NewParentName)
+{
+	ICollectionManager& CollectionManager = FCollectionManagerModule::GetModule().Get();
+
+	const FCollectionNameType ResolvedNameAndType = AssetTagsSubsystemUtil::FindCollectionByName(CollectionManager, Name);
+	if (ResolvedNameAndType.Name.IsNone())
+	{
+		return false;
+	}
+
+	const FCollectionNameType ResolvedParentNameAndType = AssetTagsSubsystemUtil::FindCollectionByName(CollectionManager, NewParentName);
+	if (ResolvedParentNameAndType.Name.IsNone())
+	{
+		return false;
+	}
+
+	if (!CollectionManager.ReparentCollection(ResolvedNameAndType.Name, ResolvedNameAndType.Type, ResolvedParentNameAndType.Name, ResolvedParentNameAndType.Type))
+	{
+		AssetTagsSubsystemUtil::LogLastCollectionManagerError(CollectionManager);
+		return false;
+	}
+
+	return true;
+}
+
+bool UAssetTagsSubsystem::EmptyCollection(const FName Name)
+{
+	ICollectionManager& CollectionManager = FCollectionManagerModule::GetModule().Get();
+
+	const FCollectionNameType ResolvedNameAndType = AssetTagsSubsystemUtil::FindCollectionByName(CollectionManager, Name);
+	if (ResolvedNameAndType.Name.IsNone())
+	{
+		return false;
+	}
+
+	if (!CollectionManager.EmptyCollection(ResolvedNameAndType.Name, ResolvedNameAndType.Type))
+	{
+		AssetTagsSubsystemUtil::LogLastCollectionManagerError(CollectionManager);
+		return false;
+	}
+
+	return true;
+}
+
+bool UAssetTagsSubsystem::AddAssetToCollection(const FName Name, const FName AssetPathName)
+{
+	ICollectionManager& CollectionManager = FCollectionManagerModule::GetModule().Get();
+
+	const FCollectionNameType ResolvedNameAndType = AssetTagsSubsystemUtil::FindCollectionByName(CollectionManager, Name);
+	if (ResolvedNameAndType.Name.IsNone())
+	{
+		return false;
+	}
+
+	if (!CollectionManager.AddToCollection(ResolvedNameAndType.Name, ResolvedNameAndType.Type, AssetPathName))
+	{
+		AssetTagsSubsystemUtil::LogLastCollectionManagerError(CollectionManager);
+		return false;
+	}
+
+	return true;
+}
+
+bool UAssetTagsSubsystem::AddAssetDataToCollection(const FName Name, const FAssetData& AssetData)
+{
+	return AddAssetToCollection(Name, AssetData.ObjectPath);
+}
+
+bool UAssetTagsSubsystem::AddAssetPtrToCollection(const FName Name, const UObject* AssetPtr)
+{
+	return AddAssetToCollection(Name, AssetPtr ? FName(*AssetPtr->GetPathName()) : NAME_None);
+}
+
+bool UAssetTagsSubsystem::AddAssetsToCollection(const FName Name, const TArray<FName>& AssetPathNames)
+{
+	ICollectionManager& CollectionManager = FCollectionManagerModule::GetModule().Get();
+
+	const FCollectionNameType ResolvedNameAndType = AssetTagsSubsystemUtil::FindCollectionByName(CollectionManager, Name);
+	if (ResolvedNameAndType.Name.IsNone())
+	{
+		return false;
+	}
+
+	if (!CollectionManager.AddToCollection(ResolvedNameAndType.Name, ResolvedNameAndType.Type, AssetPathNames))
+	{
+		AssetTagsSubsystemUtil::LogLastCollectionManagerError(CollectionManager);
+		return false;
+	}
+
+	return true;
+}
+
+bool UAssetTagsSubsystem::AddAssetDatasToCollection(const FName Name, const TArray<FAssetData>& AssetDatas)
+{
+	TArray<FName> AssetPathNames;
+	AssetPathNames.Reserve(AssetDatas.Num());
+	for (const FAssetData& AssetData : AssetDatas)
+	{
+		AssetPathNames.Add(AssetData.ObjectPath);
+	}
+	return AddAssetsToCollection(Name, AssetPathNames);
+}
+
+bool UAssetTagsSubsystem::AddAssetPtrsToCollection(const FName Name, const TArray<UObject*>& AssetPtrs)
+{
+	TArray<FName> AssetPathNames;
+	AssetPathNames.Reserve(AssetPtrs.Num());
+	for (const UObject* AssetPtr : AssetPtrs)
+	{
+		AssetPathNames.Add(AssetPtr ? FName(*AssetPtr->GetPathName()) : NAME_None);
+	}
+	return AddAssetsToCollection(Name, AssetPathNames);
+}
+
+bool UAssetTagsSubsystem::RemoveAssetFromCollection(const FName Name, const FName AssetPathName)
+{
+	ICollectionManager& CollectionManager = FCollectionManagerModule::GetModule().Get();
+
+	const FCollectionNameType ResolvedNameAndType = AssetTagsSubsystemUtil::FindCollectionByName(CollectionManager, Name);
+	if (ResolvedNameAndType.Name.IsNone())
+	{
+		return false;
+	}
+
+	if (!CollectionManager.RemoveFromCollection(ResolvedNameAndType.Name, ResolvedNameAndType.Type, AssetPathName))
+	{
+		AssetTagsSubsystemUtil::LogLastCollectionManagerError(CollectionManager);
+		return false;
+	}
+
+	return true;
+}
+
+bool UAssetTagsSubsystem::RemoveAssetDataFromCollection(const FName Name, const FAssetData& AssetData)
+{
+	return RemoveAssetFromCollection(Name, AssetData.ObjectPath);
+}
+
+bool UAssetTagsSubsystem::RemoveAssetPtrFromCollection(const FName Name, const UObject* AssetPtr)
+{
+	return RemoveAssetFromCollection(Name, AssetPtr ? FName(*AssetPtr->GetPathName()) : NAME_None);
+}
+
+bool UAssetTagsSubsystem::RemoveAssetsFromCollection(const FName Name, const TArray<FName>& AssetPathNames)
+{
+	ICollectionManager& CollectionManager = FCollectionManagerModule::GetModule().Get();
+
+	const FCollectionNameType ResolvedNameAndType = AssetTagsSubsystemUtil::FindCollectionByName(CollectionManager, Name);
+	if (ResolvedNameAndType.Name.IsNone())
+	{
+		return false;
+	}
+
+	if (!CollectionManager.RemoveFromCollection(ResolvedNameAndType.Name, ResolvedNameAndType.Type, AssetPathNames))
+	{
+		AssetTagsSubsystemUtil::LogLastCollectionManagerError(CollectionManager);
+		return false;
+	}
+
+	return true;
+}
+
+bool UAssetTagsSubsystem::RemoveAssetDatasFromCollection(const FName Name, const TArray<FAssetData>& AssetDatas)
+{
+	TArray<FName> AssetPathNames;
+	AssetPathNames.Reserve(AssetDatas.Num());
+	for (const FAssetData& AssetData : AssetDatas)
+	{
+		AssetPathNames.Add(AssetData.ObjectPath);
+	}
+	return RemoveAssetsFromCollection(Name, AssetPathNames);
+}
+
+bool UAssetTagsSubsystem::RemoveAssetPtrsFromCollection(const FName Name, const TArray<UObject*>& AssetPtrs)
+{
+	TArray<FName> AssetPathNames;
+	AssetPathNames.Reserve(AssetPtrs.Num());
+	for (const UObject* AssetPtr : AssetPtrs)
+	{
+		AssetPathNames.Add(AssetPtr ? FName(*AssetPtr->GetPathName()) : NAME_None);
+	}
+	return RemoveAssetsFromCollection(Name, AssetPathNames);
+}
+
+#endif	// WITH_EDITOR
+
+bool UAssetTagsSubsystem::CollectionExists(const FName Name)
+{
+	bool bExists = false;
+
+#if WITH_EDITOR
+	{
+		ICollectionManager& CollectionManager = FCollectionManagerModule::GetModule().Get();
+		bExists = CollectionManager.CollectionExists(Name, ECollectionShareType::CST_All);
+	}
+#else	// WITH_EDITOR
+	{
+		IAssetRegistry& AssetRegistry = FModuleManager::LoadModuleChecked<FAssetRegistryModule>(AssetRegistryConstants::ModuleName).Get();
+
+		if (const FAssetRegistryState* StatePtr = AssetRegistry.GetAssetRegistryState())
+		{
+			const FString TagNameStr = FString::Printf(TEXT("%s%s"), FAssetData::CollectionTagPrefix, *Name.ToString());
+			bExists = StatePtr->GetTagToAssetDatasMap().Contains(*TagNameStr);
+		}
+	}
+#endif	// WITH_EDITOR
+
+	return bExists;
+}
+
+TArray<FName> UAssetTagsSubsystem::GetCollections()
+{
+	TArray<FName> CollectionNames;
+
+#if WITH_EDITOR
+	{
+		ICollectionManager& CollectionManager = FCollectionManagerModule::GetModule().Get();
+
+		TArray<FCollectionNameType> CollectionNamesAndTypes;
+		CollectionManager.GetCollections(CollectionNamesAndTypes);
+
+		CollectionNames.Reserve(CollectionNamesAndTypes.Num());
+		for (const FCollectionNameType& CollectionNameAndType : CollectionNamesAndTypes)
+		{
+			CollectionNames.AddUnique(CollectionNameAndType.Name);
+		}
+	}
+#else	// WITH_EDITOR
+	{
+		IAssetRegistry& AssetRegistry = FModuleManager::LoadModuleChecked<FAssetRegistryModule>(AssetRegistryConstants::ModuleName).Get();
+
+		if (const FAssetRegistryState* StatePtr = AssetRegistry.GetAssetRegistryState())
+		{
+			for (const auto& TagAndArrayPair : StatePtr->GetTagToAssetDatasMap())
+			{
+				const FString TagNameStr = TagAndArrayPair.Key.ToString();
+				if (TagNameStr.StartsWith(FAssetData::CollectionTagPrefix))
+				{
+					const FString TrimmedTagNameStr = TagNameStr.Mid(ARRAY_COUNT(FAssetData::CollectionTagPrefix) - 1);
+					CollectionNames.Add(*TrimmedTagNameStr);
+				}
+			}
+		}
+	}
+#endif	// WITH_EDITOR
+
+	CollectionNames.Sort(FNameLexicalLess());
+	return CollectionNames;
+}
+
+TArray<FAssetData> UAssetTagsSubsystem::GetAssetsInCollection(const FName Name)
+{
+	TArray<FAssetData> Assets;
+
+	IAssetRegistry& AssetRegistry = FModuleManager::LoadModuleChecked<FAssetRegistryModule>(AssetRegistryConstants::ModuleName).Get();
+#if WITH_EDITOR
+	{
+		ICollectionManager& CollectionManager = FCollectionManagerModule::GetModule().Get();
+
+		const FCollectionNameType ResolvedNameAndType = AssetTagsSubsystemUtil::FindCollectionByName(CollectionManager, Name);
+		if (!ResolvedNameAndType.Name.IsNone())
+		{
+			TArray<FName> AssetPathNames;
+			CollectionManager.GetAssetsInCollection(ResolvedNameAndType.Name, ResolvedNameAndType.Type, AssetPathNames);
+
+			Assets.Reserve(AssetPathNames.Num());
+			for (const FName AssetPathName : AssetPathNames)
+			{
+				FAssetData AssetData = AssetRegistry.GetAssetByObjectPath(AssetPathName);
+				if (AssetData.IsValid())
+				{
+					Assets.Add(MoveTemp(AssetData));
+				}
+			}
+		}
+	}
+#else	// WITH_EDITOR
+	{
+		const FString TagNameStr = FString::Printf(TEXT("%s%s"), FAssetData::CollectionTagPrefix, *Name.ToString());
+
+		FARFilter Filter;
+		Filter.bIncludeOnlyOnDiskAssets = true; // Collection tags are added at cook-time, so we *must* search the on-disk version of the tags (from the asset registry)
+		Filter.TagsAndValues.Add(*TagNameStr);
+
+		AssetRegistry.GetAssets(Filter, Assets);
+	}
+#endif	// WITH_EDITOR
+
+	Assets.Sort();
+	return Assets;
+}
+
+TArray<FName> UAssetTagsSubsystem::GetCollectionsContainingAsset(const FName AssetPathName)
+{
+	TArray<FName> CollectionNames;
+
+#if WITH_EDITOR
+	{
+		ICollectionManager& CollectionManager = FCollectionManagerModule::GetModule().Get();
+
+		TArray<FCollectionNameType> CollectionNamesAndTypes;
+		CollectionManager.GetCollectionsContainingObject(AssetPathName, CollectionNamesAndTypes);
+
+		CollectionNames.Reserve(CollectionNamesAndTypes.Num());
+		for (const FCollectionNameType& CollectionNameAndType : CollectionNamesAndTypes)
+		{
+			CollectionNames.AddUnique(CollectionNameAndType.Name);
+		}
+	}
+#else	// WITH_EDITOR
+	{
+		IAssetRegistry& AssetRegistry = FModuleManager::LoadModuleChecked<FAssetRegistryModule>(AssetRegistryConstants::ModuleName).Get();
+
+		const bool bIncludeOnlyOnDiskAssets = true; // Collection tags are added at cook-time, so we *must* search the on-disk version of the tags (from the asset registry)
+		const FAssetData AssetData = AssetRegistry.GetAssetByObjectPath(AssetPathName, bIncludeOnlyOnDiskAssets);
+		if (AssetData.IsValid())
+		{
+			for (const auto& TagAndValuePair : AssetData.TagsAndValues)
+			{
+				const FString TagNameStr = TagAndValuePair.Key.ToString();
+				if (TagNameStr.StartsWith(FAssetData::CollectionTagPrefix))
+				{
+					const FString TrimmedTagNameStr = TagNameStr.Mid(ARRAY_COUNT(FAssetData::CollectionTagPrefix) - 1);
+					CollectionNames.Add(*TrimmedTagNameStr);
+				}
+			}
+		}
+	}
+#endif	// WITH_EDITOR
+
+	CollectionNames.Sort(FNameLexicalLess());
+	return CollectionNames;
+}
+
+TArray<FName> UAssetTagsSubsystem::GetCollectionsContainingAssetData(const FAssetData& AssetData)
+{
+	// Note: Use the path name version as the common implementation as:
+	//  1) The path name is always required for the collection manager implementation
+	//  2) The FAssetData for the asset registry implementation *must* come from the asset registry (as the tags are added at cook-time, and missing if FAssetData is generated from a UObject* at runtime)
+	return GetCollectionsContainingAsset(AssetData.ObjectPath);
+}
+
+TArray<FName> UAssetTagsSubsystem::GetCollectionsContainingAssetPtr(const UObject* AssetPtr)
+{
+	// Note: Use the path name version as the common implementation as:
+	//  1) The path name is always required for the collection manager implementation
+	//  2) The FAssetData for the asset registry implementation *must* come from the asset registry (as the tags are added at cook-time, and missing if FAssetData is generated from a UObject* at runtime)
+	return GetCollectionsContainingAsset(AssetPtr ? FName(*AssetPtr->GetPathName()) : NAME_None);
+}
diff --git a/Engine/Plugins/Runtime/AssetTags/Source/AssetTags/Public/AssetTagsSubsystem.h b/Engine/Plugins/Runtime/AssetTags/Source/AssetTags/Public/AssetTagsSubsystem.h
new file mode 100644
--- /dev/null
+++ b/Engine/Plugins/Runtime/AssetTags/Source/AssetTags/Public/AssetTagsSubsystem.h
@@ -0,0 +1,269 @@
+// Copyright 1998-2019 Epic Games, Inc. All Rights Reserved.
+
+#include "CoreMinimal.h"
+#include "Subsystems/EngineSubsystem.h"
+#include "AssetData.h"
+#include "AssetTagsSubsystem.generated.h"
+
+UENUM(DisplayName="Collection Share Type", meta=(ScriptName="CollectionShareType"))
+enum class ECollectionScriptingShareType : uint8
+{
+	/** This collection is only visible to you and is not in source control. */
+	Local,
+	/** This collection is only visible to you. */
+	Private,
+	/** This collection is visible to everyone. */
+	Shared,
+};
+
+UCLASS()
+class ASSETTAGS_API UAssetTagsSubsystem : public UEngineSubsystem
+{
+	GENERATED_BODY()
+
+public:
+#if WITH_EDITOR
+	/**
+	 * Create a new collection with the given name and share type.
+	 *
+	 * @param Name Name to give to the collection.
+	 * @param ShareType Whether the collection should be local, private, or shared?
+	 *
+	 * @return True if the collection was created, false otherwise (see the output log for details on error).
+	 */
+	UFUNCTION(BlueprintCallable, Category="AssetTags")
+	bool CreateCollection(const FName Name, const ECollectionScriptingShareType ShareType);
+	
+	/**
+	 * Destroy the given collection.
+	 *
+	 * @param Name Name of the collection to destroy.
+	 *
+	 * @return True if the collection was destroyed, false otherwise (see the output log for details on error).
+	 */
+	UFUNCTION(BlueprintCallable, Category="AssetTags")
+	bool DestroyCollection(const FName Name);
+	
+	/**
+	 * Rename the given collection.
+	 *
+	 * @param Name Name of the collection to rename.
+	 * @param NewName Name to give to the collection.
+	 *
+	 * @return True if the collection was renamed, false otherwise (see the output log for details on error).
+	 */
+	UFUNCTION(BlueprintCallable, Category="AssetTags")
+	bool RenameCollection(const FName Name, const FName NewName);
+	
+	/**
+	 * Re-parent the given collection.
+	 *
+	 * @param Name Name of the collection to re-parent.
+	 * @param NewParentName Name of the new parent collection, or None to have the collection become a root collection.
+	 *
+	 * @return True if the collection was renamed, false otherwise (see the output log for details on error).
+	 */
+	UFUNCTION(BlueprintCallable, Category="AssetTags")
+	bool ReparentCollection(const FName Name, const FName NewParentName);
+
+	/**
+	 * Remove all assets from the given collection.
+	 *
+	 * @param Name Name of the collection to modify.
+	 *
+	 * @return True if the collection was modified, false otherwise (see the output log for details on error).
+	 */
+	UFUNCTION(BlueprintCallable, Category="AssetTags")
+	bool EmptyCollection(const FName Name);
+
+	/**
+	 * Add the given asset to the given collection.
+	 * 
+	 * @param Name Name of the collection to modify.
+	 * @param AssetPathName Asset to add (its path name, eg) /Game/MyFolder/MyAsset.MyAsset).
+	 *
+	 * @return True if the collection was modified, false otherwise (see the output log for details on error).
+	 */
+	UFUNCTION(BlueprintCallable, Category="AssetTags")
+	bool AddAssetToCollection(const FName Name, const FName AssetPathName);
+
+	/**
+	 * Add the given asset to the given collection.
+	 *
+	 * @param Name Name of the collection to modify.
+	 * @param AssetData Asset to add.
+	 *
+	 * @return True if the collection was modified, false otherwise (see the output log for details on error).
+	 */
+	UFUNCTION(BlueprintCallable, Category="AssetTags")
+	bool AddAssetDataToCollection(const FName Name, const FAssetData& AssetData);
+
+	/**
+	 * Add the given asset to the given collection.
+	 *
+	 * @param Name Name of the collection to modify.
+	 * @param AssetPtr Asset to add.
+	 *
+	 * @return True if the collection was modified, false otherwise (see the output log for details on error).
+	 */
+	UFUNCTION(BlueprintCallable, Category="AssetTags")
+	bool AddAssetPtrToCollection(const FName Name, const UObject* AssetPtr);
+	
+	/**
+	 * Add the given assets to the given collection.
+	 *
+	 * @param Name Name of the collection to modify.
+	 * @param AssetPathNames Assets to add (their path names, eg) /Game/MyFolder/MyAsset.MyAsset).
+	 *
+	 * @return True if the collection was modified, false otherwise (see the output log for details on error).
+	 */
+	UFUNCTION(BlueprintCallable, Category="AssetTags")
+	bool AddAssetsToCollection(const FName Name, const TArray<FName>& AssetPathNames);
+
+	/**
+	 * Add the given assets to the given collection.
+	 *
+	 * @param Name Name of the collection to modify.
+	 * @param AssetDatas Assets to add.
+	 *
+	 * @return True if the collection was modified, false otherwise (see the output log for details on error).
+	 */
+	UFUNCTION(BlueprintCallable, Category="AssetTags")
+	bool AddAssetDatasToCollection(const FName Name, const TArray<FAssetData>& AssetDatas);
+
+	/**
+	 * Add the given assets to the given collection.
+	 *
+	 * @param Name Name of the collection to modify.
+	 * @param AssetPtrs Assets to add.
+	 *
+	 * @return True if the collection was modified, false otherwise (see the output log for details on error).
+	 */
+	UFUNCTION(BlueprintCallable, Category="AssetTags")
+	bool AddAssetPtrsToCollection(const FName Name, const TArray<UObject*>& AssetPtrs);
+
+	/**
+	 * Remove the given asset from the given collection.
+	 * 
+	 * @param Name Name of the collection to modify.
+	 * @param AssetPathName Asset to remove (its path name, eg) /Game/MyFolder/MyAsset.MyAsset).
+	 *
+	 * @return True if the collection was modified, false otherwise (see the output log for details on error).
+	 */
+	UFUNCTION(BlueprintCallable, Category="AssetTags")
+	bool RemoveAssetFromCollection(const FName Name, const FName AssetPathName);
+
+	/**
+	 * Remove the given asset from the given collection.
+	 *
+	 * @param Name Name of the collection to modify.
+	 * @param AssetData Asset to remove.
+	 *
+	 * @return True if the collection was modified, false otherwise (see the output log for details on error).
+	 */
+	UFUNCTION(BlueprintCallable, Category="AssetTags")
+	bool RemoveAssetDataFromCollection(const FName Name, const FAssetData& AssetData);
+
+	/**
+	 * Remove the given asset from the given collection.
+	 *
+	 * @param Name Name of the collection to modify.
+	 * @param AssetPtr Asset to remove.
+	 *
+	 * @return True if the collection was modified, false otherwise (see the output log for details on error).
+	 */
+	UFUNCTION(BlueprintCallable, Category="AssetTags")
+	bool RemoveAssetPtrFromCollection(const FName Name, const UObject* AssetPtr);
+	
+	/**
+	 * Remove the given assets from the given collection.
+	 *
+	 * @param Name Name of the collection to modify.
+	 * @param AssetPathNames Assets to remove (their path names, eg) /Game/MyFolder/MyAsset.MyAsset).
+	 *
+	 * @return True if the collection was modified, false otherwise (see the output log for details on error).
+	 */
+	UFUNCTION(BlueprintCallable, Category="AssetTags")
+	bool RemoveAssetsFromCollection(const FName Name, const TArray<FName>& AssetPathNames);
+
+	/**
+	 * Remove the given assets from the given collection.
+	 *
+	 * @param Name Name of the collection to modify.
+	 * @param AssetDatas Assets to remove.
+	 *
+	 * @return True if the collection was modified, false otherwise (see the output log for details on error).
+	 */
+	UFUNCTION(BlueprintCallable, Category="AssetTags")
+	bool RemoveAssetDatasFromCollection(const FName Name, const TArray<FAssetData>& AssetDatas);
+
+	/**
+	 * Remove the given assets from the given collection.
+	 *
+	 * @param Name Name of the collection to modify.
+	 * @param AssetPtrs Assets to remove.
+	 *
+	 * @return True if the collection was modified, false otherwise (see the output log for details on error).
+	 */
+	UFUNCTION(BlueprintCallable, Category="AssetTags")
+	bool RemoveAssetPtrsFromCollection(const FName Name, const TArray<UObject*>& AssetPtrs);
+#endif 	// WITH_EDITOR
+
+	/**
+	 * Check whether the given collection exists.
+	 *
+	 * @param Name Name of the collection to test.
+	 *
+	 * @return True if the collection exists, false otherwise.
+	 */
+	UFUNCTION(BlueprintCallable, Category="AssetTags")
+	bool CollectionExists(const FName Name);
+
+	/**
+	 * Get the names of all available collections.
+	 *
+	 * @return Names of all available collections.
+	 */
+	UFUNCTION(BlueprintCallable, Category="AssetTags")
+	TArray<FName> GetCollections();
+
+	/**
+	 * Get the assets in the given collection.
+	 * 
+	 * @param Name Name of the collection to test.
+	 *
+	 * @return Assets in the given collection.
+	 */
+	UFUNCTION(BlueprintCallable, Category="AssetTags")
+	TArray<FAssetData> GetAssetsInCollection(const FName Name);
+
+	/**
+	 * Get the names of the collections that contain the given asset.
+	 * 
+	 * @param AssetPathName Asset to test (its path name, eg) /Game/MyFolder/MyAsset.MyAsset).
+	 *
+	 * @return Names of the collections that contain the asset.
+	 */
+	UFUNCTION(BlueprintCallable, Category="AssetTags")
+	TArray<FName> GetCollectionsContainingAsset(const FName AssetPathName);
+
+	/**
+	 * Get the names of the collections that contain the given asset.
+	 * 
+	 * @param AssetData Asset to test.
+	 *
+	 * @return Names of the collections that contain the asset.
+	 */
+	UFUNCTION(BlueprintCallable, Category="AssetTags")
+	TArray<FName> GetCollectionsContainingAssetData(const FAssetData& AssetData);
+
+	/**
+	 * Get the names of the collections that contain the given asset.
+	 * 
+	 * @param AssetPtr Asset to test.
+	 *
+	 * @return Names of the collections that contain the asset.
+	 */
+	UFUNCTION(BlueprintCallable, Category="AssetTags")
+	TArray<FName> GetCollectionsContainingAssetPtr(const UObject* AssetPtr);
+};
diff --git a/Engine/Plugins/VirtualProduction/Takes/Source/TakeRecorder/Private/TakeRecorderModule.cpp b/Engine/Plugins/VirtualProduction/Takes/Source/TakeRecorder/Private/TakeRecorderModule.cpp
--- a/Engine/Plugins/VirtualProduction/Takes/Source/TakeRecorder/Private/TakeRecorderModule.cpp
+++ b/Engine/Plugins/VirtualProduction/Takes/Source/TakeRecorder/Private/TakeRecorderModule.cpp
@@ -85,8 +85,6 @@ namespace
 			ContentBrowserConfig.bCanShowClasses = true;
 			ContentBrowserConfig.bUseSourcesView = true;
 			ContentBrowserConfig.bExpandSourcesView = true;
-			ContentBrowserConfig.bShowAssetPathTree = true;
-			ContentBrowserConfig.bAlwaysShowCollections = false;
 			ContentBrowserConfig.bUsePathPicker = true;
 			ContentBrowserConfig.bCanShowFilters = true;
 			ContentBrowserConfig.bCanShowAssetSearch = true;
diff --git a/Engine/Source/Developer/CollectionManager/Private/Collection.cpp b/Engine/Source/Developer/CollectionManager/Private/Collection.cpp
--- a/Engine/Source/Developer/CollectionManager/Private/Collection.cpp
+++ b/Engine/Source/Developer/CollectionManager/Private/Collection.cpp
@@ -398,6 +398,11 @@ void FCollection::Empty()
 
 bool FCollection::AddObjectToCollection(FName ObjectPath)
 {
+	if (ObjectPath.IsNone())
+	{
+		return false;
+	}
+
 	if (StorageMode == ECollectionStorageMode::Static)
 	{
 		bool bAlreadyInSet = false;
@@ -410,6 +415,11 @@ bool FCollection::AddObjectToCollection(FName ObjectPath)
 
 bool FCollection::RemoveObjectFromCollection(FName ObjectPath)
 {
+	if (ObjectPath.IsNone())
+	{
+		return false;
+	}
+
 	if (StorageMode == ECollectionStorageMode::Static)
 	{
 		return ObjectSet.Remove(ObjectPath) > 0;
@@ -542,6 +552,11 @@ bool FCollection::IsDirty() const
 		return true;
 	}
 	
+	if (CollectionColor != DiskSnapshot.CollectionColor)
+	{
+		return true;
+	}
+
 	bool bHasChanges = false;
 
 	if (StorageMode == ECollectionStorageMode::Static)
@@ -603,6 +618,10 @@ void FCollection::SaveHeaderPairs(TMap<FString,FString>& OutHeaderPairs) const
 	OutHeaderPairs.Add(TEXT("Type"), ECollectionStorageMode::ToString(StorageMode));
 	OutHeaderPairs.Add(TEXT("Guid"), CollectionGuid.ToString(EGuidFormats::DigitsWithHyphens));
 	OutHeaderPairs.Add(TEXT("ParentGuid"), ParentCollectionGuid.ToString(EGuidFormats::DigitsWithHyphens));
+	if (CollectionColor)
+	{
+		OutHeaderPairs.Add(TEXT("Color"), CollectionColor->ToString());
+	}
 }
 
 bool FCollection::LoadHeaderPairs(const TMap<FString,FString>& InHeaderPairs)
@@ -643,6 +662,17 @@ bool FCollection::LoadHeaderPairs(const TMap<FString,FString>& InHeaderPairs)
 		}
 	}
 
+	// Load the optional color
+	CollectionColor.Reset();
+	if (const FString* ColorStr = InHeaderPairs.Find(TEXT("Color")))
+	{
+		FLinearColor NewColor;
+		if (NewColor.InitFromString(*ColorStr))
+		{
+			CollectionColor = MoveTemp(NewColor);
+		}
+	}
+
 	return FileVersion > 0 && FileVersion <= ECollectionVersion::CurrentVersion;
 }
 
@@ -960,6 +990,12 @@ bool FCollection::CheckinCollection(const TArray<FText>& AdditionalChangelistTex
 			ChangelistDescBuilder.AppendLineFormat(LOCTEXT("CollectionChangedParentDesc", "Changed the parent of collection '{0}'"), CollectionNameText);
 		}
 
+		// Color change?
+		if (DiskSnapshot.CollectionColor != CollectionColor)
+		{
+			ChangelistDescBuilder.AppendLineFormat(LOCTEXT("CollectionChangedColorDesc", "Changed the color of collection '{0}'"), CollectionNameText);
+		}
+
 		// Version bump?
 		if (FileVersion < ECollectionVersion::CurrentVersion)
 		{
diff --git a/Engine/Source/Developer/CollectionManager/Private/Collection.h b/Engine/Source/Developer/CollectionManager/Private/Collection.h
--- a/Engine/Source/Developer/CollectionManager/Private/Collection.h
+++ b/Engine/Source/Developer/CollectionManager/Private/Collection.h
@@ -106,6 +106,12 @@ public:
 	/** Set the GUID of the collection we are parented under */
 	FORCEINLINE void SetParentCollectionGuid(const FGuid& NewGuid) { ParentCollectionGuid = NewGuid; }
 
+	/** Get the color of the collection (if any) */
+	FORCEINLINE TOptional<FLinearColor> GetCollectionColor() const { return CollectionColor; }
+
+	/** Set the color of the collection (if any) */
+	FORCEINLINE void SetCollectionColor(const TOptional<FLinearColor>& NewColor) { CollectionColor = NewColor; }
+
 	/** Returns the file version of the collection */
 	FORCEINLINE ECollectionVersion::Type GetCollectionVersion() const { return FileVersion; }
 
@@ -149,6 +155,7 @@ private:
 		void TakeSnapshot(const FCollection& InCollection)
 		{
 			ParentCollectionGuid = InCollection.ParentCollectionGuid;
+			CollectionColor = InCollection.CollectionColor;
 			ObjectSet = InCollection.ObjectSet;
 			DynamicQueryText = InCollection.DynamicQueryText;
 		}
@@ -156,6 +163,9 @@ private:
 		/** The GUID of the collection we are parented under */
 		FGuid ParentCollectionGuid;
 
+		/** The color of the collection (if set) */
+		TOptional<FLinearColor> CollectionColor;
+
 		/** The set of objects in the collection. Takes the form PackageName.AssetName. Static collections only. */
 		TSet<FName> ObjectSet;
 
@@ -172,6 +182,9 @@ private:
 	/** The GUID of the collection we are parented under */
 	FGuid ParentCollectionGuid;
 
+	/** The color of the collection (if any) */
+	TOptional<FLinearColor> CollectionColor;
+
 	/** Source control is used if true */
 	bool bUseSCC;
 
diff --git a/Engine/Source/Developer/CollectionManager/Private/CollectionManager.cpp b/Engine/Source/Developer/CollectionManager/Private/CollectionManager.cpp
--- a/Engine/Source/Developer/CollectionManager/Private/CollectionManager.cpp
+++ b/Engine/Source/Developer/CollectionManager/Private/CollectionManager.cpp
@@ -16,6 +16,7 @@ FCollectionManagerCache::FCollectionManagerCache(FAvailableCollectionsMap& InAva
 	bIsCachedCollectionNamesFromGuidsDirty = true;
 	bIsCachedObjectsDirty = true;
 	bIsCachedHierarchyDirty = true;
+	bIsCachedColorsDirty = true;
 }
 
 void FCollectionManagerCache::HandleCollectionAdded()
@@ -28,12 +29,14 @@ void FCollectionManagerCache::HandleCollectionRemoved()
 	bIsCachedCollectionNamesFromGuidsDirty = true;
 	bIsCachedObjectsDirty = true;
 	bIsCachedHierarchyDirty = true;
+	bIsCachedColorsDirty = true;
 }
 
 void FCollectionManagerCache::HandleCollectionChanged()
 {
 	bIsCachedObjectsDirty = true;
 	bIsCachedHierarchyDirty = true;
+	bIsCachedColorsDirty = true;
 }
 
 const FGuidToCollectionNamesMap& FCollectionManagerCache::GetCachedCollectionNamesFromGuids() const
@@ -53,7 +56,7 @@ const FGuidToCollectionNamesMap& FCollectionManagerCache::GetCachedCollectionNam
 			CachedCollectionNamesFromGuids_Internal.Add(Collection->GetCollectionGuid(), CollectionKey);
 		}
 
-		UE_LOG(LogCollectionManager, Log, TEXT("Rebuilt the GUID cache for %d collections in %0.6f seconds"), AvailableCollections.Num(), FPlatformTime::Seconds() - CacheStartTime);
+		UE_LOG(LogCollectionManager, Verbose, TEXT("Rebuilt the GUID cache for %d collections in %0.6f seconds"), AvailableCollections.Num(), FPlatformTime::Seconds() - CacheStartTime);
 	}
 
 	return CachedCollectionNamesFromGuids_Internal;
@@ -118,7 +121,7 @@ const FCollectionObjectsMap& FCollectionManagerCache::GetCachedObjects() const
 			}
 		}
 
-		UE_LOG(LogCollectionManager, Log, TEXT("Rebuilt the object cache for %d collections in %0.6f seconds (found %d objects)"), AvailableCollections.Num(), FPlatformTime::Seconds() - CacheStartTime, CachedObjects_Internal.Num());
+		UE_LOG(LogCollectionManager, Verbose, TEXT("Rebuilt the object cache for %d collections in %0.6f seconds (found %d objects)"), AvailableCollections.Num(), FPlatformTime::Seconds() - CacheStartTime, CachedObjects_Internal.Num());
 	}
 
 	return CachedObjects_Internal;
@@ -149,12 +152,42 @@ const FCollectionHierarchyMap& FCollectionManagerCache::GetCachedHierarchy() con
 			}
 		}
 
-		UE_LOG(LogCollectionManager, Log, TEXT("Rebuilt the hierarchy cache for %d collections in %0.6f seconds"), AvailableCollections.Num(), FPlatformTime::Seconds() - CacheStartTime);
+		UE_LOG(LogCollectionManager, Verbose, TEXT("Rebuilt the hierarchy cache for %d collections in %0.6f seconds"), AvailableCollections.Num(), FPlatformTime::Seconds() - CacheStartTime);
 	}
 
 	return CachedHierarchy_Internal;
 }
 
+const FCollectionColorArray& FCollectionManagerCache::GetCachedColors() const
+{
+	if (bIsCachedColorsDirty)
+	{
+		CachedColors_Internal.Reset();
+		bIsCachedColorsDirty = false;
+
+		const double CacheStartTime = FPlatformTime::Seconds();
+
+		for (const auto& AvailableCollection : AvailableCollections)
+		{
+			const TSharedRef<FCollection>& Collection = AvailableCollection.Value;
+
+			if (const TOptional<FLinearColor> CollectionColor = Collection->GetCollectionColor())
+			{
+				// Only add if not already present (ignores near matches too)
+				const bool bExists = CachedColors_Internal.FindByPredicate([CurrentColor = CollectionColor.GetValue()](const FLinearColor& Color) { return CurrentColor.Equals(Color); });
+				if (!bExists)
+				{
+					CachedColors_Internal.Add(CollectionColor.GetValue());
+				}
+			}
+		}
+
+		UE_LOG(LogCollectionManager, Verbose, TEXT("Rebuilt the color cache for %d collections in %0.6f seconds"), AvailableCollections.Num(), FPlatformTime::Seconds() - CacheStartTime);
+	}
+
+	return CachedColors_Internal;
+}
+
 void FCollectionManagerCache::RecursionHelper_DoWork(const FCollectionNameType& InCollectionKey, const ECollectionRecursionFlags::Flags InRecursionMode, FRecursiveWorkerFunc InWorkerFunc) const
 {
 	if ((InRecursionMode & ECollectionRecursionFlags::Self) && InWorkerFunc(InCollectionKey, ECollectionRecursionFlags::Self) == ERecursiveWorkerFlowControl::Stop)
@@ -281,6 +314,17 @@ void FCollectionManager::GetCollections(TArray<FCollectionNameType>& OutCollecti
 	}
 }
 
+void FCollectionManager::GetCollections(FName CollectionName, TArray<FCollectionNameType>& OutCollections) const
+{
+	for (int32 CacheIdx = 0; CacheIdx < ECollectionShareType::CST_All; ++CacheIdx)
+	{
+		if (AvailableCollections.Contains(FCollectionNameType(CollectionName, ECollectionShareType::Type(CacheIdx))))
+		{
+			OutCollections.Add(FCollectionNameType(CollectionName, ECollectionShareType::Type(CacheIdx)));
+		}
+	}
+}
+
 void FCollectionManager::GetCollectionNames(ECollectionShareType::Type ShareType, TArray<FName>& CollectionNames) const
 {
 	for (const auto& AvailableCollection : AvailableCollections)
@@ -644,11 +688,6 @@ FString FCollectionManager::GetCollectionsStringForObject(FName ObjectPath, ECol
 
 void FCollectionManager::CreateUniqueCollectionName(const FName& BaseName, ECollectionShareType::Type ShareType, FName& OutCollectionName) const
 {
-	if (!ensure(ShareType != ECollectionShareType::CST_All))
-	{
-		return;
-	}
-
 	int32 IntSuffix = 1;
 	bool CollectionAlreadyExists = false;
 	do
@@ -662,7 +701,7 @@ void FCollectionManager::CreateUniqueCollectionName(const FName& BaseName, EColl
 			OutCollectionName = *FString::Printf(TEXT("%s%d"), *BaseName.ToString(), IntSuffix);
 		}
 
-		CollectionAlreadyExists = AvailableCollections.Contains(FCollectionNameType(OutCollectionName, ShareType));
+		CollectionAlreadyExists = CollectionExists(OutCollectionName, ShareType);
 		++IntSuffix;
 	}
 	while (CollectionAlreadyExists);
@@ -1326,6 +1365,72 @@ bool FCollectionManager::GetCollectionStatusInfo(FName CollectionName, ECollecti
 	return false;
 }
 
+bool FCollectionManager::HasCollectionColors(TArray<FLinearColor>* OutColors) const
+{
+	const FCollectionColorArray& CollectionColors = CollectionCache.GetCachedColors();
+	if (OutColors)
+	{
+		*OutColors = CollectionColors;
+	}
+	return CollectionColors.Num() > 0;
+}
+
+bool FCollectionManager::GetCollectionColor(FName CollectionName, ECollectionShareType::Type ShareType, TOptional<FLinearColor>& OutColor) const
+{
+	if (!ensure(ShareType < ECollectionShareType::CST_All))
+	{
+		// Bad share type
+		LastError = LOCTEXT("Error_Internal", "There was an internal error.");
+		return true;
+	}
+
+	const FCollectionNameType CollectionKey(CollectionName, ShareType);
+	const TSharedRef<FCollection>* const CollectionRefPtr = AvailableCollections.Find(CollectionKey);
+	if (CollectionRefPtr)
+	{
+		OutColor = (*CollectionRefPtr)->GetCollectionColor();
+		return true;
+	}
+	else
+	{
+		LastError = LOCTEXT("Error_DoesntExist", "The collection doesn't exist.");
+	}
+
+	return false;
+}
+
+bool FCollectionManager::SetCollectionColor(FName CollectionName, ECollectionShareType::Type ShareType, const TOptional<FLinearColor>& NewColor)
+{
+	if (!ensure(ShareType < ECollectionShareType::CST_All))
+	{
+		// Bad share type
+		LastError = LOCTEXT("Error_Internal", "There was an internal error.");
+		return true;
+	}
+
+	const FCollectionNameType CollectionKey(CollectionName, ShareType);
+	const TSharedRef<FCollection>* const CollectionRefPtr = AvailableCollections.Find(CollectionKey);
+	if (CollectionRefPtr)
+	{
+		(*CollectionRefPtr)->SetCollectionColor(NewColor);
+
+		if (InternalSaveCollection(*CollectionRefPtr, LastError))
+		{
+			CollectionFileCaches[ShareType]->IgnoreFileModification((*CollectionRefPtr)->GetSourceFilename());
+			
+			CollectionCache.HandleCollectionChanged();
+			CollectionUpdatedEvent.Broadcast(CollectionKey);
+			return true;
+		}
+	}
+	else
+	{
+		LastError = LOCTEXT("Error_DoesntExist", "The collection doesn't exist.");
+	}
+
+	return false;
+}
+
 bool FCollectionManager::GetCollectionStorageMode(FName CollectionName, ECollectionShareType::Type ShareType, ECollectionStorageMode::Type& OutStorageMode) const
 {
 	if (!ensure(ShareType < ECollectionShareType::CST_All))
diff --git a/Engine/Source/Developer/CollectionManager/Private/CollectionManager.h b/Engine/Source/Developer/CollectionManager/Private/CollectionManager.h
--- a/Engine/Source/Developer/CollectionManager/Private/CollectionManager.h
+++ b/Engine/Source/Developer/CollectionManager/Private/CollectionManager.h
@@ -37,6 +37,7 @@ typedef TMap<FCollectionNameType, TSharedRef<FCollection>> FAvailableCollections
 typedef TMap<FGuid, FCollectionNameType> FGuidToCollectionNamesMap;
 typedef TMap<FName, TArray<FObjectCollectionInfo>> FCollectionObjectsMap;
 typedef TMap<FGuid, TArray<FGuid>> FCollectionHierarchyMap;
+typedef TArray<FLinearColor> FCollectionColorArray;
 
 /** Wraps up the lazy caching of the collection manager */
 class FCollectionManagerCache
@@ -62,6 +63,9 @@ public:
 	/** Access the CachedHierarchy map, ensuring that it is up-to-date */
 	const FCollectionHierarchyMap& GetCachedHierarchy() const;
 
+	/** Access the CachedColors array, ensuring that it is up-to-date */
+	const FCollectionColorArray& GetCachedColors() const;
+
 	enum class ERecursiveWorkerFlowControl : uint8
 	{
 		Stop,
@@ -87,6 +91,9 @@ private:
 	/** A map of parent collection GUIDs to their child collection GUIDs - only collections that have children will appear in here */
 	mutable FCollectionHierarchyMap CachedHierarchy_Internal;
 
+	/** An array of all unique colors currently used by collections */
+	mutable FCollectionColorArray CachedColors_Internal;
+
 	/** Flag to say whether the CachedCollectionNamesFromGuids map is dirty */
 	mutable bool bIsCachedCollectionNamesFromGuidsDirty : 1;
 
@@ -95,6 +102,9 @@ private:
 
 	/** Flag to say whether the CachedHierarchy map is dirty */
 	mutable bool bIsCachedHierarchyDirty : 1;
+
+	/** Flag to say whether the CachedColors array is dirty */
+	mutable bool bIsCachedColorsDirty : 1;
 };
 
 class FCollectionManager : public ICollectionManager
@@ -106,6 +116,7 @@ public:
 	// ICollectionManager implementation
 	virtual bool HasCollections() const override;
 	virtual void GetCollections(TArray<FCollectionNameType>& OutCollections) const override;
+	virtual void GetCollections(FName CollectionName, TArray<FCollectionNameType>& OutCollections) const override;
 	virtual void GetCollectionNames(ECollectionShareType::Type ShareType, TArray<FName>& CollectionNames) const override;
 	virtual void GetRootCollections(TArray<FCollectionNameType>& OutCollections) const override;
 	virtual void GetRootCollectionNames(ECollectionShareType::Type ShareType, TArray<FName>& CollectionNames) const override;
@@ -137,6 +148,9 @@ public:
 	virtual bool SaveCollection(FName CollectionName, ECollectionShareType::Type ShareType) override;
 	virtual bool UpdateCollection(FName CollectionName, ECollectionShareType::Type ShareType) override;
 	virtual bool GetCollectionStatusInfo(FName CollectionName, ECollectionShareType::Type ShareType, FCollectionStatusInfo& OutStatusInfo) const override;
+	virtual bool HasCollectionColors(TArray<FLinearColor>* OutColors = nullptr) const override;
+	virtual bool GetCollectionColor(FName CollectionName, ECollectionShareType::Type ShareType, TOptional<FLinearColor>& OutColor) const override;
+	virtual bool SetCollectionColor(FName CollectionName, ECollectionShareType::Type ShareType, const TOptional<FLinearColor>& NewColor) override;
 	virtual bool GetCollectionStorageMode(FName CollectionName, ECollectionShareType::Type ShareType, ECollectionStorageMode::Type& OutStorageMode) const override;
 	virtual bool IsObjectInCollection(FName ObjectPath, FName CollectionName, ECollectionShareType::Type ShareType, ECollectionRecursionFlags::Flags RecursionMode = ECollectionRecursionFlags::Self) const override;
 	virtual bool IsValidParentCollection(FName CollectionName, ECollectionShareType::Type ShareType, FName ParentCollectionName, ECollectionShareType::Type ParentShareType) const override;
diff --git a/Engine/Source/Developer/CollectionManager/Public/ICollectionManager.h b/Engine/Source/Developer/CollectionManager/Public/ICollectionManager.h
--- a/Engine/Source/Developer/CollectionManager/Public/ICollectionManager.h
+++ b/Engine/Source/Developer/CollectionManager/Public/ICollectionManager.h
@@ -19,6 +19,9 @@ public:
 	/** Returns the list of collections */
 	virtual void GetCollections(TArray<FCollectionNameType>& OutCollections) const = 0;
 
+	/** Returns the list of collections with the given name */
+	virtual void GetCollections(FName CollectionName, TArray<FCollectionNameType>& OutCollections) const = 0;
+
 	/** Returns the list of collection names of the specified share type */
 	virtual void GetCollectionNames(ECollectionShareType::Type ShareType, TArray<FName>& CollectionNames) const = 0;
 
@@ -206,6 +209,34 @@ public:
 	  */
 	virtual bool GetCollectionStatusInfo(FName CollectionName, ECollectionShareType::Type ShareType, FCollectionStatusInfo& OutStatusInfo) const = 0;
 
+	/**
+	 * Checks whether any collections have a custom color set, optionally retrieving the list of in-use colors
+	 *
+	 * @param OutColors The optional color array to populate.
+	 * @return true if there are any custom colors in use, false otherwise.
+	 */
+	virtual bool HasCollectionColors(TArray<FLinearColor>* OutColors = nullptr) const = 0;
+
+	/**
+	 * Gets the optional color for the specified collection
+	 *
+	 * @param CollectionName The collection to get the color for
+	 * @param ShareType The way the collection is shared.
+	 * @param OutColor The optional color to populate.
+	 * @return true if the collection was found, and the color was optionally filled in. If false, GetLastError will return a human readable string description of the error.
+	 */
+	virtual bool GetCollectionColor(FName CollectionName, ECollectionShareType::Type ShareType, TOptional<FLinearColor>& OutColor) const = 0;
+
+	/**
+	 * Sets the optional color for the specified collection
+	 *
+	 * @param CollectionName The collection to set the color for
+	 * @param ShareType The way the collection is shared.
+	 * @param NewColor The optional color to set.
+	 * @return true if the collection was found, and the color was set. If false, GetLastError will return a human readable string description of the error.
+	 */
+	virtual bool SetCollectionColor(FName CollectionName, ECollectionShareType::Type ShareType, const TOptional<FLinearColor>& NewColor) = 0;
+
 	/**
 	 * Gets the method by which the specified collection stores its objects (static or dynamic)
 	 *
diff --git a/Engine/Source/Editor/AssetTagsEditor/AssetTagsEditor.Build.cs b/Engine/Source/Editor/AssetTagsEditor/AssetTagsEditor.Build.cs
new file mode 100644
--- /dev/null
+++ b/Engine/Source/Editor/AssetTagsEditor/AssetTagsEditor.Build.cs
@@ -0,0 +1,19 @@
+// Copyright 1998-2019 Epic Games, Inc. All Rights Reserved.
+
+using UnrealBuildTool;
+
+public class AssetTagsEditor : ModuleRules
+{
+	public AssetTagsEditor(ReadOnlyTargetRules Target) : base(Target)
+	{
+		PublicDependencyModuleNames.AddRange(
+            new string[] {
+                "Core",
+				"CoreUObject",
+				"Slate",
+				"SlateCore",
+				"EditorStyle",
+			}
+        );
+	}
+}
diff --git a/Engine/Source/Editor/AssetTagsEditor/Private/AssetTagsEditorModule.cpp b/Engine/Source/Editor/AssetTagsEditor/Private/AssetTagsEditorModule.cpp
new file mode 100644
--- /dev/null
+++ b/Engine/Source/Editor/AssetTagsEditor/Private/AssetTagsEditorModule.cpp
@@ -0,0 +1,5 @@
+// Copyright 1998-2019 Epic Games, Inc. All Rights Reserved.
+
+#include "Modules/ModuleManager.h"
+
+IMPLEMENT_MODULE(FDefaultModuleImpl, AssetTagsEditor);
diff --git a/Engine/Source/Editor/AssetTagsEditor/Private/SAssetTagItem.cpp b/Engine/Source/Editor/AssetTagsEditor/Private/SAssetTagItem.cpp
new file mode 100644
--- /dev/null
+++ b/Engine/Source/Editor/AssetTagsEditor/Private/SAssetTagItem.cpp
@@ -0,0 +1,257 @@
+// Copyright 1998-2019 Epic Games, Inc. All Rights Reserved.
+
+#include "SAssetTagItem.h"
+#include "Widgets/SToolTip.h"
+#include "Widgets/Images/SImage.h"
+#include "Widgets/Text/STextBlock.h"
+
+#define LOCTEXT_NAMESPACE "AssetTagsEditor"
+
+class SAssetTagItemToolTip : public SToolTip
+{
+public:
+	SLATE_BEGIN_ARGS(SAssetTagItemToolTip)
+	{}
+
+		/** Binding to get the display name of this asset tag item (must be set) */
+		SLATE_ATTRIBUTE(FText, DisplayName)
+
+		/** Callback used to build the tooltip info box for this asset tag item */
+		SLATE_EVENT(FOnBuildAssetTagItemToolTipInfo, OnBuildToolTipInfo)
+
+	SLATE_END_ARGS()
+
+	/** Constructs this widget with InArgs */
+	void Construct(const FArguments& InArgs)
+	{
+		DisplayName = InArgs._DisplayName;
+		OnBuildToolTipInfo = InArgs._OnBuildToolTipInfo;
+
+		SToolTip::Construct(
+			SToolTip::FArguments()
+			.TextMargin(1.0f)
+			.BorderImage(FEditorStyle::GetBrush("ContentBrowser.TileViewTooltip.ToolTipBorder"))
+			);
+	}
+
+	virtual bool IsEmpty() const override
+	{
+		return false;
+	}
+
+	virtual void OnOpening() override
+	{
+		SetContentWidget(CreateToolTipWidget());
+	}
+
+private:
+	TSharedRef<SWidget> CreateToolTipWidget()
+	{
+		TSharedRef<SVerticalBox> OverallTooltipVBox = SNew(SVerticalBox);
+
+		// Create a box to hold every line of info in the body of the tooltip
+		TSharedRef<SVerticalBox> InfoBox = SNew(SVerticalBox);
+		if (OnBuildToolTipInfo.IsBound())
+		{
+			OnBuildToolTipInfo.Execute([&InfoBox](const FText& Key, const FText& Value)
+			{
+				AddToToolTipInfoBox(InfoBox, Key, Value);
+			});
+		}
+
+		// Top section (asset tag item name)
+		OverallTooltipVBox->AddSlot()
+			.AutoHeight()
+			.Padding(0)
+			[
+				SNew(SBorder)
+				.Padding(6)
+				.BorderImage(FEditorStyle::GetBrush("ContentBrowser.TileViewTooltip.ContentBorder"))
+				[
+					SNew(STextBlock)
+					.Text(DisplayName)
+					.Font(FEditorStyle::GetFontStyle("ContentBrowser.TileViewTooltip.NameFont"))
+				]
+			];
+
+		// Bottom section (asset tag item details, if provided)
+		if (InfoBox->NumSlots() > 0)
+		{
+			OverallTooltipVBox->AddSlot()
+				.AutoHeight()
+				.Padding(0, 4, 0, 0)
+				[
+					SNew(SBorder)
+					.Padding(6)
+					.BorderImage(FEditorStyle::GetBrush("ContentBrowser.TileViewTooltip.ContentBorder"))
+					[
+						InfoBox
+					]
+				];
+		}
+
+		return SNew(SBorder)
+			.Padding(6)
+			.BorderImage(FEditorStyle::GetBrush("ContentBrowser.TileViewTooltip.NonContentBorder"))
+			[
+				OverallTooltipVBox
+			];
+	}
+
+	static void AddToToolTipInfoBox(const TSharedRef<SVerticalBox>& InfoBox, const FText& Key, const FText& Value)
+	{
+		InfoBox->AddSlot()
+		.AutoHeight()
+		.Padding(0, 1)
+		[
+			SNew(SHorizontalBox)
+
+			+SHorizontalBox::Slot()
+			.AutoWidth()
+			.Padding(0, 0, 4, 0)
+			[
+				SNew(STextBlock)
+				.Text(FText::Format(LOCTEXT("AssetTagTooltipKeyFormat", "{0}:"), Key))
+				.ColorAndOpacity(FSlateColor::UseSubduedForeground())
+			]
+
+			+SHorizontalBox::Slot()
+			.AutoWidth()
+			[
+				SNew(STextBlock)
+				.Text(Value)
+				.ColorAndOpacity(FSlateColor::UseForeground())
+				.WrapTextAt(700.0f)
+			]
+		];
+	}
+
+	/** Binding to get the display name of this asset tag item (must be set) */
+	TAttribute<FText> DisplayName;
+
+	/** Callback used to build the tooltip info box for this asset tag item */
+	FOnBuildAssetTagItemToolTipInfo OnBuildToolTipInfo;
+};
+
+void SAssetTagItem::Construct(const FArguments& InArgs)
+{
+	BaseColor = InArgs._BaseColor;
+	WarningText = InArgs._WarningText;
+	IsSelectedCallback = InArgs._IsSelected;
+
+	checkf(InArgs._DisplayName.IsSet(), TEXT("SAssetTagItem DisplayName must be set!"));
+
+	TSharedRef<SAssetTagItemToolTip> AssetTagToolTip = SNew(SAssetTagItemToolTip)
+		.DisplayName(InArgs._DisplayName)
+		.OnBuildToolTipInfo(InArgs._OnBuildToolTipInfo);
+
+	ANSICHAR* StyleSpecifier = nullptr;
+	if (InArgs._ViewMode == EAssetTagItemViewMode::Compact)
+	{
+		StyleSpecifier = ".Compact";
+	}
+
+	TAttribute<bool> IsCheckBoxEnabled = InArgs._IsCheckBoxEnabled;
+	if (!InArgs._IsChecked.IsSet() || !InArgs._OnCheckStateChanged.IsBound())
+	{
+		IsCheckBoxEnabled = false;
+	}
+
+	TSharedPtr<SWidget> NameWidget;
+	if (InArgs._OnNameCommitted.IsBound())
+	{
+		NameWidget = SAssignNew(InlineRenameWidget, SInlineEditableTextBlock)
+			.Font(FEditorStyle::GetFontStyle("ContentBrowser.AssetTagNameFont", StyleSpecifier))
+			.ShadowOffset(FVector2D(1.0f, 1.0f))
+			.Text(InArgs._DisplayName)
+			.HighlightText(InArgs._HighlightText)
+			.OnBeginTextEdit(InArgs._OnBeginNameEdit)
+			.OnTextCommitted(InArgs._OnNameCommitted)
+			.OnVerifyTextChanged(InArgs._OnVerifyName)
+			.IsSelected(IsSelectedCallback)
+			.IsReadOnly(InArgs._IsNameReadOnly);
+	}
+	else
+	{
+		NameWidget = SNew(STextBlock)
+			.Font(FEditorStyle::GetFontStyle("ContentBrowser.AssetTagNameFont", StyleSpecifier))
+			.ShadowOffset(FVector2D(1.0f, 1.0f))
+			.Text(InArgs._DisplayName)
+			.HighlightText(InArgs._HighlightText);
+	}
+
+	TSharedPtr<SHorizontalBox> HBox;
+
+	ChildSlot
+	[
+		SNew(SBorder)
+		.Padding(0.0f)
+		.BorderBackgroundColor(FLinearColor(0.2f, 0.2f, 0.2f, 0.2f))
+		.BorderImage(FEditorStyle::GetBrush("ContentBrowser.AssetTagBackground"))
+		.ToolTip(AssetTagToolTip)
+		[
+			SAssignNew(HBox, SHorizontalBox)
+
+			+SHorizontalBox::Slot()
+			.AutoWidth()
+			[
+				SAssignNew(CheckBox, SCheckBox)
+				.Style(FEditorStyle::Get(), "ContentBrowser.AssetTagButton", StyleSpecifier)
+				.IsEnabled(IsCheckBoxEnabled)
+				.IsChecked(InArgs._IsChecked)
+				.OnCheckStateChanged(InArgs._OnCheckStateChanged)
+				.ForegroundColor(this, &SAssetTagItem::GetCheckBoxForegroundColor)
+				.ToolTipText(this, &SAssetTagItem::GetCheckBoxTooltipText)
+			]
+
+			+SHorizontalBox::Slot()
+			.Padding(FEditorStyle::GetMargin("ContentBrowser.AssetTagNamePadding", StyleSpecifier))
+			.VAlign(VAlign_Center)
+			[
+				NameWidget.ToSharedRef()
+			]
+		]
+	];
+
+	if (WarningText.IsSet())
+	{
+		HBox->AddSlot()
+			.AutoWidth()
+			.VAlign(VAlign_Center)
+			[
+				SNew(SImage)
+				.Visibility(this, &SAssetTagItem::GetWarningIconVisibility)
+				.Image(FEditorStyle::GetBrush("Icons.Warning"))
+				.ToolTipText(WarningText)
+			];
+	}
+
+	if (InArgs._CountText.IsSet())
+	{
+		HBox->AddSlot()
+			.AutoWidth()
+			.Padding(FEditorStyle::GetMargin("ContentBrowser.AssetTagCountPadding", StyleSpecifier))
+			[
+				SNew(SBorder)
+				.Padding(0.0f)
+				.VAlign(VAlign_Center)
+				.BorderBackgroundColor(this, &SAssetTagItem::GetCountBackgroundColor)
+				.BorderImage(FEditorStyle::GetBrush("ContentBrowser.AssetTagBackground"))
+				[
+					SNew(SBox)
+					.VAlign(VAlign_Center)
+					.MinDesiredWidth(30.0f)
+					[
+						SNew(STextBlock)
+						.Font(FEditorStyle::GetFontStyle("ContentBrowser.AssetTagCountFont", StyleSpecifier))
+						.ColorAndOpacity(FLinearColor::White)
+						.ShadowOffset(FVector2D(1.0f, 1.0f))
+						.Text(InArgs._CountText)
+						.Justification(ETextJustify::Center)
+					]
+				]
+			];
+	}
+}
+
+#undef LOCTEXT_NAMESPACE
diff --git a/Engine/Source/Editor/AssetTagsEditor/Public/AssetTagItemTypes.h b/Engine/Source/Editor/AssetTagsEditor/Public/AssetTagItemTypes.h
new file mode 100644
--- /dev/null
+++ b/Engine/Source/Editor/AssetTagsEditor/Public/AssetTagItemTypes.h
@@ -0,0 +1,15 @@
+// Copyright 1998-2019 Epic Games, Inc. All Rights Reserved.
+
+#pragma once
+
+#include "CoreMinimal.h"
+#include "Templates/Function.h"
+
+enum class EAssetTagItemViewMode : uint8
+{
+	Standard,
+	Compact,
+};
+
+typedef TFunctionRef<void(const FText&, const FText&)> FOnBuildAssetTagItemToolTipInfoEntry;
+DECLARE_DELEGATE_OneParam(FOnBuildAssetTagItemToolTipInfo, const FOnBuildAssetTagItemToolTipInfoEntry&);
diff --git a/Engine/Source/Editor/AssetTagsEditor/Public/SAssetTagItem.h b/Engine/Source/Editor/AssetTagsEditor/Public/SAssetTagItem.h
new file mode 100644
--- /dev/null
+++ b/Engine/Source/Editor/AssetTagsEditor/Public/SAssetTagItem.h
@@ -0,0 +1,215 @@
+// Copyright 1998-2019 Epic Games, Inc. All Rights Reserved.
+
+#pragma once
+
+#include "CoreMinimal.h"
+#include "AssetTagItemTypes.h"
+#include "EditorStyleSet.h"
+#include "Widgets/DeclarativeSyntaxSupport.h"
+#include "Widgets/SCompoundWidget.h"
+#include "Widgets/Views/STableRow.h"
+#include "Widgets/Input/SCheckBox.h"
+#include "Widgets/Text/SInlineEditableTextBlock.h"
+
+/** Custom table row used to contain asset tag items within list and tree views */
+template <typename ItemType>
+class SAssetTagItemTableRow : public STableRow<ItemType>
+{
+public:
+	void Construct(const FArguments& InArgs, const TSharedRef<STableViewBase>& InOwnerTableView)
+	{
+		ConstructInternal(InArgs, InOwnerTableView);
+
+		Style = &FEditorStyle::GetWidgetStyle<FTableRowStyle>("ContentBrowser.AssetTagTableRow");
+		BorderImage = FEditorStyle::GetBrush("NoBorder");
+
+		ConstructChildren(
+			InOwnerTableView->TableViewMode,
+			ApplySelectionBorderOuterPadding(InArgs._Padding),
+			WrapContentInSelectionBorder(InArgs._Content.Widget)
+			);
+	}
+
+	void SetIsDropTarget(const TAttribute<bool>& InIsDropTarget)
+	{
+		IsDropTarget = InIsDropTarget;
+	}
+
+	virtual void SetContent(TSharedRef<SWidget> InContent) override
+	{
+		STableRow::SetContent(WrapContentInSelectionBorder(InContent));
+	}
+
+	virtual const FSlateBrush* GetBorder() const override
+	{
+		if (IsDropTarget.Get(false))
+		{
+			return &Style->InactiveHoveredBrush;
+		}
+
+		return STableRow::GetBorder();
+	}
+
+private:
+	TSharedRef<SWidget> WrapContentInSelectionBorder(const TSharedRef<SWidget> InContent) const
+	{
+		return
+			SNew(SBorder)
+			.Padding(2.0f)
+			.BorderImage(this, &SAssetTagItemTableRow::GetBorder)
+			[
+				InContent
+			];
+	}
+
+	TAttribute<FMargin> ApplySelectionBorderOuterPadding(const TAttribute<FMargin>& InPadding) const
+	{
+		const FMargin ExtraPadding = FMargin(0.0f, 0.0f, 0.0f, 1.0f);
+
+		if (InPadding.IsBound())
+		{
+			return MakeAttributeLambda([InPadding, ExtraPadding]()
+			{
+				return InPadding.Get() + ExtraPadding;
+			});
+		}
+		else if (InPadding.IsSet())
+		{
+			return InPadding.Get() + ExtraPadding;
+		}
+
+		return ExtraPadding;
+	}
+
+	TAttribute<bool> IsDropTarget;
+};
+
+/** A single asset tag item in the tags and collections tree. */
+class ASSETTAGSEDITOR_API SAssetTagItem : public SCompoundWidget
+{
+public:
+	SLATE_BEGIN_ARGS(SAssetTagItem)
+		: _ViewMode(EAssetTagItemViewMode::Standard)
+		, _BaseColor(FLinearColor::White)
+	{}
+
+		/** Should this asset tag item use the standard or compact view? */
+		SLATE_ARGUMENT(EAssetTagItemViewMode, ViewMode)
+
+		/** Binding to get the base color of this asset tag item */
+		SLATE_ATTRIBUTE(FLinearColor, BaseColor)
+
+		/** Binding to get the display name of this asset tag item (must be set) */
+		SLATE_ATTRIBUTE(FText, DisplayName)
+
+		/** Binding to get the count text of this asset tag item (unset to omit the count UI) */
+		SLATE_ATTRIBUTE(FText, CountText)
+
+		/** Binding to get the warning text of this asset tag item (unset to omit the warning UI) */
+		SLATE_ATTRIBUTE(FText, WarningText)
+
+		/** Binding to get the highlight text of the asset tag view */
+		SLATE_ATTRIBUTE(FText, HighlightText)
+
+		/** Binding to check whether the asset tag item name is read-only (disables the name edit UI) */
+		SLATE_ATTRIBUTE(bool, IsNameReadOnly)
+
+		/** Binding to check whether the check box of the asset tag item is enabled */
+		SLATE_ATTRIBUTE(bool, IsCheckBoxEnabled)
+
+		/** Binding to check whether the check box of the asset tag item is currently in a checked state (unset to always disable the check box) */
+		SLATE_ATTRIBUTE(ECheckBoxState, IsChecked)
+
+		/** Callback when the checked state of the asset tag item check box is checked or unchecked (unset to always disable the check box) */
+		SLATE_EVENT(FOnCheckStateChanged, OnCheckStateChanged)
+
+		/** Callback when the asset tag item name starts to be edited */
+		SLATE_EVENT(FOnBeginTextEdit, OnBeginNameEdit)
+
+		/** Callback when the asset tag item name is committed (unset to omit the name edit UI) */
+		SLATE_EVENT(FOnTextCommitted, OnNameCommitted)
+
+		/** Called to validate the asset tag item name (called during changes and during commit) */
+		SLATE_EVENT(FOnVerifyTextChanged, OnVerifyName)
+
+		/** Callback to check if the asset tag item is selected (should only be hooked up if a parent widget is handling selection or focus) */
+		SLATE_EVENT(FIsSelected, IsSelected)
+
+		/** Callback used to build the tooltip info box for this asset tag item */
+		SLATE_EVENT(FOnBuildAssetTagItemToolTipInfo, OnBuildToolTipInfo)
+
+	SLATE_END_ARGS()
+
+	/** Constructs this widget with InArgs */
+	void Construct(const FArguments& InArgs);
+
+	void RequestRename()
+	{
+		if (InlineRenameWidget)
+		{
+			InlineRenameWidget->EnterEditingMode();
+		}
+	}
+
+private:
+	FLinearColor GetAssetTagBrightColor() const
+	{
+		return BaseColor.Get().CopyWithNewOpacity(0.3f);
+	}
+
+	FLinearColor GetAssetTagDullColor() const
+	{
+		return BaseColor.Get().CopyWithNewOpacity(0.1f);
+	}
+
+	FLinearColor GetAssetTagDisabledColor() const
+	{
+		return BaseColor.Get().CopyWithNewOpacity(0.04f);
+	}
+
+	FSlateColor GetCheckBoxForegroundColor() const
+	{
+		return CheckBox && CheckBox->IsEnabled()
+			? CheckBox->IsChecked() 
+				? GetAssetTagBrightColor()
+				: GetAssetTagDullColor()
+			: GetAssetTagDisabledColor();
+	}
+
+	FText GetCheckBoxTooltipText() const
+	{
+		return CheckBox && CheckBox->IsEnabled()
+			? CheckBox->IsChecked() 
+				? NSLOCTEXT("AssetTagsEditor", "RemoveAssetSelectionFromAssetTagItem", "Untag the current asset selection")
+				: NSLOCTEXT("AssetTagsEditor", "AddAssetSelectionToAssetTagItem", "Tag the current asset selection")
+			: FText::GetEmpty();
+	}
+
+	FSlateColor GetCountBackgroundColor() const
+	{
+		return GetAssetTagBrightColor();
+	}
+
+	EVisibility GetWarningIconVisibility() const
+	{
+		FText WarningTextValue = WarningText.Get(FText::GetEmpty());
+		return WarningTextValue.IsEmpty()
+			? EVisibility::Collapsed
+			: EVisibility::Visible;
+	}
+
+	/** Binding to get the base color of this asset tag item */
+	TAttribute<FLinearColor> BaseColor;
+
+	/** Binding to get the warning text of this asset tag item (unset to omit the warning UI) */
+	TAttribute<FText> WarningText;
+
+	/** Callback to check if the asset tag item is selected (should only be hooked up if a parent widget is handling selection or focus) */
+	FIsSelected IsSelectedCallback;
+
+	/** The check box for enabling/disabling the asset tag item */
+	TSharedPtr<SCheckBox> CheckBox;
+
+	/** The widget to display the name of the asset tag item and allows for renaming */
+	TSharedPtr<SInlineEditableTextBlock> InlineRenameWidget;
+};
diff --git a/Engine/Source/Editor/ContentBrowser/ContentBrowser.Build.cs b/Engine/Source/Editor/ContentBrowser/ContentBrowser.Build.cs
--- a/Engine/Source/Editor/ContentBrowser/ContentBrowser.Build.cs
+++ b/Engine/Source/Editor/ContentBrowser/ContentBrowser.Build.cs
@@ -40,6 +40,7 @@ public class ContentBrowser : ModuleRules
 				"AddContentDialog",
 				"DesktopPlatform",
 				"AssetRegistry",
+				"AssetTagsEditor",
 				"ToolMenus",
 			}
 		);
diff --git a/Engine/Source/Editor/ContentBrowser/Private/AssetViewWidgets.cpp b/Engine/Source/Editor/ContentBrowser/Private/AssetViewWidgets.cpp
--- a/Engine/Source/Editor/ContentBrowser/Private/AssetViewWidgets.cpp
+++ b/Engine/Source/Editor/ContentBrowser/Private/AssetViewWidgets.cpp
@@ -4,6 +4,7 @@
 #include "AssetViewWidgets.h"
 #include "UObject/UnrealType.h"
 #include "Widgets/SOverlay.h"
+#include "Widgets/Layout/SWrapBox.h"
 #include "Engine/GameViewportClient.h"
 #include "SlateOptMacros.h"
 #include "Framework/Application/SlateApplication.h"
@@ -24,6 +25,7 @@
 #include "AssetViewTypes.h"
 #include "SThumbnailEditModeTools.h"
 #include "AutoReimport/AssetSourceFilenameCache.h"
+#include "SAssetTagItem.h"
 #include "CollectionViewUtils.h"
 #include "DragAndDrop/AssetDragDropOp.h"
 #include "DragDropHandler.h"
@@ -801,16 +803,6 @@ TSharedRef<SWidget> SAssetViewItem::CreateToolTipWidget() const
 				UE_LOG(LogContentBrowser, Error, TEXT("AssetData for '%s' is invalid"), *AssetData.PackagePath.ToString());
 			}
 
-			// Add Collections
-			{
-				FCollectionManagerModule& CollectionManagerModule = FCollectionManagerModule::GetModule();
-				const FString CollectionNames = CollectionManagerModule.Get().GetCollectionsStringForObject(AssetData.ObjectPath, ECollectionShareType::CST_All, ECollectionRecursionFlags::Self, GetDefault<UContentBrowserSettings>()->bShowFullCollectionNameInToolTip);
-				if (!CollectionNames.IsEmpty())
-				{
-					AddToToolTipInfoBox(InfoBox, LOCTEXT("AssetToolTipKey_Collections", "Collections"), FText::FromString(CollectionNames), false);
-				}
-			}
-
 			// Add tags
 			for (const auto& DisplayTagItem : CachedDisplayTags)
 			{
@@ -864,7 +856,8 @@ TSharedRef<SWidget> SAssetViewItem::CreateToolTipWidget() const
 							.AutoWidth()
 							.VAlign(VAlign_Center)
 							[
-								SNew(STextBlock).Text(ClassText)
+								SNew(STextBlock)
+								.Text(ClassText)
 								.HighlightText(HighlightText)
 							]
 						]
@@ -877,6 +870,7 @@ TSharedRef<SWidget> SAssetViewItem::CreateToolTipWidget() const
 							.Text(this, &SAssetViewItem::GetCheckedOutByOtherText)
 							.ColorAndOpacity(FLinearColor(0.1f, 0.5f, 1.f, 1.f))
 						]
+
 						+ SVerticalBox::Slot()
 						.AutoHeight()
 						[
@@ -917,6 +911,49 @@ TSharedRef<SWidget> SAssetViewItem::CreateToolTipWidget() const
 					]
 				];
 
+			// Final section (collection pips)
+			{
+				ICollectionManager& CollectionManager = FCollectionManagerModule::GetModule().Get();
+
+				TArray<FCollectionNameType> CollectionsContainingObject;
+				CollectionManager.GetCollectionsContainingObject(AssetData.ObjectPath, CollectionsContainingObject);
+
+				if (CollectionsContainingObject.Num() > 0)
+				{
+					TSharedRef<SWrapBox> CollectionPipsWrapBox = SNew(SWrapBox)
+						.PreferredWidth(700.0f);
+
+					for (const FCollectionNameType& CollectionContainingObject : CollectionsContainingObject)
+					{
+						FCollectionStatusInfo CollectionStatusInfo;
+						if (CollectionManager.GetCollectionStatusInfo(CollectionContainingObject.Name, CollectionContainingObject.Type, CollectionStatusInfo))
+						{
+							CollectionPipsWrapBox->AddSlot()
+							.Padding(0, 4, 4, 0)
+							[
+								// TODO: Honor or remove GetDefault<UContentBrowserSettings>()->bShowFullCollectionNameInToolTip
+								SNew(SAssetTagItem)
+								.ViewMode(EAssetTagItemViewMode::Compact)
+								.BaseColor(CollectionViewUtils::ResolveColor(CollectionContainingObject.Name, CollectionContainingObject.Type))
+								.DisplayName(FText::FromName(CollectionContainingObject.Name))
+								.CountText(FText::AsNumber(CollectionStatusInfo.NumObjects))
+							];
+						}
+					}
+
+					OverallTooltipVBox->AddSlot()
+						.AutoHeight()
+						.Padding(0, 4, 0, 0)
+						[
+							SNew(SBorder)
+							.Padding(FMargin(6, 2, 6, 6))
+							.BorderImage(FEditorStyle::GetBrush("ContentBrowser.TileViewTooltip.ContentBorder"))
+							[
+								CollectionPipsWrapBox
+							]
+						];
+				}
+			}
 
 			return SNew(SBorder)
 				.Padding(6)
@@ -1458,23 +1495,23 @@ FSlateColor SAssetViewItem::GetAssetColor() const
 		{
 			TSharedPtr<FAssetViewFolder> AssetFolderItem = StaticCastSharedPtr<FAssetViewFolder>(AssetItem);
 
-			TSharedPtr<FLinearColor> Color;
 			if (AssetFolderItem->bCollectionFolder)
 			{
 				FName CollectionName;
 				ECollectionShareType::Type CollectionFolderShareType = ECollectionShareType::CST_All;
 				ContentBrowserUtils::IsCollectionPath(AssetFolderItem->FolderPath, &CollectionName, &CollectionFolderShareType);
 
-				Color = CollectionViewUtils::LoadColor( CollectionName.ToString(), CollectionFolderShareType );
+				if (TOptional<FLinearColor> Color = CollectionViewUtils::GetCustomColor(CollectionName, CollectionFolderShareType))
+				{
+					return Color.GetValue();
+				}
 			}
 			else
 			{
-				Color = ContentBrowserUtils::LoadColor( AssetFolderItem->FolderPath );
-			}
-
-			if ( Color.IsValid() )
-			{
-				return *Color.Get();
+				if (TSharedPtr<FLinearColor> Color = ContentBrowserUtils::LoadColor(AssetFolderItem->FolderPath))
+				{
+					return *Color;
+				}
 			}
 		}
 		else if(AssetTypeActions.IsValid())
diff --git a/Engine/Source/Editor/ContentBrowser/Private/CollectionContextMenu.cpp b/Engine/Source/Editor/ContentBrowser/Private/CollectionContextMenu.cpp
--- a/Engine/Source/Editor/ContentBrowser/Private/CollectionContextMenu.cpp
+++ b/Engine/Source/Editor/ContentBrowser/Private/CollectionContextMenu.cpp
@@ -120,7 +120,7 @@ TSharedPtr<SWidget> FCollectionContextMenu::MakeCollectionTreeContextMenu(TShare
 			);
 
 		// If any colors have already been set, display color options as a sub menu
-		if ( CollectionViewUtils::HasCustomColors() )
+		if ( CollectionViewUtils::HasCustomColors() && CanExecuteColorChange() )
 		{
 			// Set Color (submenu)
 			MenuBuilder.AddSubMenu(
@@ -136,7 +136,10 @@ TSharedPtr<SWidget> FCollectionContextMenu::MakeCollectionTreeContextMenu(TShare
 				LOCTEXT("SetColor", "Set Color"),
 				LOCTEXT("SetCollectionColorTooltip", "Sets the color this collection should appear as."),
 				FSlateIcon(),
-				FUIAction( FExecuteAction::CreateSP( this, &FCollectionContextMenu::ExecutePickColor ) )
+				FUIAction(
+					FExecuteAction::CreateSP( this, &FCollectionContextMenu::ExecutePickColor ), 
+					FCanExecuteAction::CreateSP( this, &FCollectionContextMenu::CanExecuteColorChange )
+					)
 				);
 		}
 	}
@@ -423,12 +426,6 @@ void FCollectionContextMenu::ExecuteNewCollection(ECollectionShareType::Type Col
 		return;
 	}
 
-	if (!GetDefault<UContentBrowserSettings>()->GetDisplayCollections())
-	{
-		GetMutableDefault<UContentBrowserSettings>()->SetDisplayCollections(true);
-		GetMutableDefault<UContentBrowserSettings>()->PostEditChange();
-	}
-
 	CollectionView.Pin()->CreateCollectionItem(CollectionType, StorageMode, InCreationPayload);
 }
 
@@ -455,12 +452,6 @@ void FCollectionContextMenu::ExecuteSaveDynamicCollection(FCollectionNameType In
 {
 	FCollectionManagerModule& CollectionManagerModule = FCollectionManagerModule::GetModule();
 
-	if (!GetDefault<UContentBrowserSettings>()->GetDisplayCollections())
-	{
-		GetMutableDefault<UContentBrowserSettings>()->SetDisplayCollections(true);
-		GetMutableDefault<UContentBrowserSettings>()->PostEditChange();
-	}
-
 	CollectionManagerModule.Get().SetDynamicQueryText(InCollection.Name, InCollection.Type, InSearchQuery.ToString());
 }
 
@@ -586,34 +577,18 @@ void FCollectionContextMenu::ExecutePickColor()
 	TSharedPtr<SCollectionView> CollectionViewPtr = CollectionView.Pin();
 	TArray<FCollectionNameType> SelectedCollections = CollectionViewPtr->GetSelectedCollections();
 
-	// Spawn a color picker, so the user can select which color they want
-	TArray<FLinearColor*> LinearColorArray;
-	FColorPickerArgs PickerArgs;
-	PickerArgs.bIsModal = false;
-	PickerArgs.ParentWidget = CollectionViewPtr;
-	if ( SelectedCollections.Num() > 0 )
+	FLinearColor InitialColor = FLinearColor::White;
+	if (SelectedCollections.Num() > 0)
 	{
-		// Make sure an color entry exists for all the collections, otherwise they won't update in realtime with the widget color
-		for (int32 PathIdx = SelectedCollections.Num() - 1; PathIdx >= 0; --PathIdx)
-		{
-			const FCollectionNameType& SelectedCollection = SelectedCollections[PathIdx];
-			TSharedPtr<FLinearColor> Color = CollectionViewUtils::LoadColor( SelectedCollection.Name.ToString(), SelectedCollection.Type );
-			if ( !Color.IsValid() )
-			{
-				Color = MakeShareable( new FLinearColor( CollectionViewUtils::GetDefaultColor() ) );
-				CollectionViewUtils::SaveColor( SelectedCollection.Name.ToString(), SelectedCollection.Type, Color, true );
-			}
-			else
-			{
-				// Default the color to the first valid entry
-				PickerArgs.InitialColorOverride = *Color.Get();
-			}
-			LinearColorArray.Add( Color.Get() );
-		}	
-		PickerArgs.LinearColorArray = &LinearColorArray;
+		const FCollectionNameType& SelectedCollection = SelectedCollections[0];
+		InitialColor = CollectionViewUtils::ResolveColor(SelectedCollection.Name, SelectedCollection.Type);
 	}
-		
-	PickerArgs.OnColorPickerWindowClosed = FOnWindowClosed::CreateSP(this, &FCollectionContextMenu::NewColorComplete);
+
+	FColorPickerArgs PickerArgs;
+	PickerArgs.bIsModal = true; // TODO: Allow live color updates via a proxy?
+	PickerArgs.ParentWidget = CollectionViewPtr;
+	PickerArgs.InitialColorOverride = InitialColor;
+	PickerArgs.OnColorCommitted.BindSP(this, &FCollectionContextMenu::OnColorCommitted);
 
 	OpenColorPicker(PickerArgs);
 }
@@ -685,8 +660,8 @@ bool FCollectionContextMenu::SelectedHasCustomColors() const
 	for(const FCollectionNameType& SelectedCollection : SelectedCollections)
 	{
 		// Ignore any that are the default color
-		const TSharedPtr<FLinearColor> Color = CollectionViewUtils::LoadColor( SelectedCollection.Name.ToString(), SelectedCollection.Type );
-		if ( Color.IsValid() && !Color->Equals( CollectionViewUtils::GetDefaultColor() ) )
+		const TOptional<FLinearColor> Color = CollectionViewUtils::GetCustomColor(SelectedCollection.Name, SelectedCollection.Type);
+		if (Color)
 		{
 			return true;
 		}
@@ -694,41 +669,39 @@ bool FCollectionContextMenu::SelectedHasCustomColors() const
 	return false;
 }
 
-void FCollectionContextMenu::NewColorComplete(const TSharedRef<SWindow>& Window)
+bool FCollectionContextMenu::CanExecuteColorChange() const
 {
-	TSharedPtr<SCollectionView> CollectionViewPtr = CollectionView.Pin();
-	TArray<FCollectionNameType> SelectedCollections = CollectionViewPtr->GetSelectedCollections();
+	TArray<TSharedPtr<FCollectionItem>> SelectedCollections = CollectionView.Pin()->CollectionTreePtr->GetSelectedItems();
+	const bool bIsSourceControlValid = bProjectUnderSourceControl && ISourceControlModule::Get().IsEnabled() && ISourceControlModule::Get().GetProvider().IsAvailable();
 
-	// Save the colors back in the config (ptr should have already updated by the widget)
-	for(const FCollectionNameType& SelectedCollection : SelectedCollections)
+	bool bCanChangeColor = false;
+	for (TSharedPtr<FCollectionItem> SelectedCollection : SelectedCollections)
 	{
-		const TSharedPtr<FLinearColor> Color = CollectionViewUtils::LoadColor( SelectedCollection.Name.ToString(), SelectedCollection.Type );
-		check( Color.IsValid() );
-		CollectionViewUtils::SaveColor( SelectedCollection.Name.ToString(), SelectedCollection.Type, Color );
+		bCanChangeColor |= (SelectedCollection->CollectionType == ECollectionShareType::CST_Local || bIsSourceControlValid);
 	}
+	return bCanChangeColor;
 }
 
 FReply FCollectionContextMenu::OnColorClicked( const FLinearColor InColor )
+{
+	OnColorCommitted(InColor);
+
+	// Dismiss the menu here, as we can't make the 'clear' option appear if a folder has just had a color set for the first time
+	FSlateApplication::Get().DismissAllMenus();
+
+	return FReply::Handled();
+}
+
+void FCollectionContextMenu::OnColorCommitted( const FLinearColor InColor )
 {
 	TSharedPtr<SCollectionView> CollectionViewPtr = CollectionView.Pin();
 	TArray<FCollectionNameType> SelectedCollections = CollectionViewPtr->GetSelectedCollections();
 
 	// Make sure an color entry exists for all the collections, otherwise it can't save correctly
-	for(const FCollectionNameType& SelectedCollection : SelectedCollections)
+	for (const FCollectionNameType& SelectedCollection : SelectedCollections)
 	{
-		TSharedPtr<FLinearColor> Color = CollectionViewUtils::LoadColor( SelectedCollection.Name.ToString(), SelectedCollection.Type );
-		if ( !Color.IsValid() )
-		{
-			Color = MakeShareable( new FLinearColor() );
-		}
-		*Color.Get() = InColor;
-		CollectionViewUtils::SaveColor( SelectedCollection.Name.ToString(), SelectedCollection.Type, Color );
+		CollectionViewUtils::SetCustomColor(SelectedCollection.Name, SelectedCollection.Type, InColor);
 	}
-
-	// Dismiss the menu here, as we can't make the 'clear' option appear if a folder has just had a color set for the first time
-	FSlateApplication::Get().DismissAllMenus();
-
-	return FReply::Handled();
 }
 
 void FCollectionContextMenu::ResetColors()
@@ -739,7 +712,7 @@ void FCollectionContextMenu::ResetColors()
 	// Clear the custom colors for all the selected collections
 	for(const FCollectionNameType& SelectedCollection : SelectedCollections)
 	{
-		CollectionViewUtils::SaveColor( SelectedCollection.Name.ToString(), SelectedCollection.Type, nullptr );
+		CollectionViewUtils::SetCustomColor(SelectedCollection.Name, SelectedCollection.Type, TOptional<FLinearColor>());
 	}
 }
 
diff --git a/Engine/Source/Editor/ContentBrowser/Private/CollectionContextMenu.h b/Engine/Source/Editor/ContentBrowser/Private/CollectionContextMenu.h
--- a/Engine/Source/Editor/ContentBrowser/Private/CollectionContextMenu.h
+++ b/Engine/Source/Editor/ContentBrowser/Private/CollectionContextMenu.h
@@ -97,12 +97,15 @@ private:
 	/** Checks to see if any of the selected collections use custom colors */
 	bool SelectedHasCustomColors() const;
 
-	/** Callback when the color picker dialog has been closed */
-	void NewColorComplete(const TSharedRef<SWindow>& Window);
+	/** Handler to check to see if "Set/Clear Color" can be executed */
+	bool CanExecuteColorChange() const;
 
 	/** Callback when the color is picked from the set color submenu */
 	FReply OnColorClicked( const FLinearColor InColor );
 
+	/** Callback when the color is picked from the set color submenu or picker */
+	void OnColorCommitted( const FLinearColor InColor );
+
 	/** Resets the colors of the selected collections */
 	void ResetColors();
 
diff --git a/Engine/Source/Editor/ContentBrowser/Private/CollectionViewTypes.h b/Engine/Source/Editor/ContentBrowser/Private/CollectionViewTypes.h
--- a/Engine/Source/Editor/ContentBrowser/Private/CollectionViewTypes.h
+++ b/Engine/Source/Editor/ContentBrowser/Private/CollectionViewTypes.h
@@ -62,9 +62,15 @@ struct FCollectionItem
 	/** If true, this item will be created the next time it is renamed */
 	bool bNewCollection;
 
+	/** The number of objects in this collection (static collections only) */
+	int32 NumObjects;
+
 	/** Current status of this collection item */
 	ECollectionItemStatus CurrentStatus;
 
+	/** The color of the collection */
+	FLinearColor CollectionColor;
+
 	/** Called once after the collection is created (see bNewCollection) */
 	DECLARE_DELEGATE_OneParam( FCollectionCreatedEvent, FCollectionNameType )
 
@@ -86,6 +92,8 @@ struct FCollectionItem
 		, ChildCollections()
 		, bRenaming(false)
 		, bNewCollection(false)
+		, NumObjects(0)
 		, CurrentStatus(ECollectionItemStatus::IsUpToDateAndEmpty)
+		, CollectionColor(FLinearColor::White)
 	{}
 };
diff --git a/Engine/Source/Editor/ContentBrowser/Private/CollectionViewUtils.cpp b/Engine/Source/Editor/ContentBrowser/Private/CollectionViewUtils.cpp
--- a/Engine/Source/Editor/ContentBrowser/Private/CollectionViewUtils.cpp
+++ b/Engine/Source/Editor/ContentBrowser/Private/CollectionViewUtils.cpp
@@ -6,88 +6,95 @@
 #include "Misc/ConfigCacheIni.h"
 #include "ICollectionManager.h"
 #include "CollectionManagerModule.h"
+#include "EditorStyleSet.h"
 
 #define LOCTEXT_NAMESPACE "CollectionView"
 
 namespace CollectionViewUtils
 {
-	/** Create a string of the form "CollectionName:CollectionType */
-	FString ToConfigKey(const FString& InCollectionName, const ECollectionShareType::Type& InCollectionType)
-	{
-		static_assert(ECollectionShareType::CST_All == 4, "Update CollectionViewUtils::ToConfigKey for the updated ECollectionShareType values");
 
-		check(InCollectionType != ECollectionShareType::CST_All);
+// Keep a map of all the collections that have custom colors, so updating the color in one location updates them all
+static TMap<FString, TOptional<FLinearColor>> LegacyCollectionColors;
 
-		FString CollectionTypeStr;
-		switch(InCollectionType)
-		{
-		case ECollectionShareType::CST_System:
-			CollectionTypeStr = "System";
-			break;
-		case ECollectionShareType::CST_Local:
-			CollectionTypeStr = "Local";
-			break;
-		case ECollectionShareType::CST_Private:
-			CollectionTypeStr = "Private";
-			break;
-		case ECollectionShareType::CST_Shared:
-			CollectionTypeStr = "Shared";
-			break;
-		default:
-			break;
-		}
+FLinearColor GetLegacyDefaultColor()
+{
+	// The default tint the folder should appear as
+	return FLinearColor::Gray;
+}
 
-		return InCollectionName + ":" + CollectionTypeStr;
-	}
+/** Create a string of the form "CollectionName:CollectionType */
+FString ToConfigKey(const FString& InCollectionName, const ECollectionShareType::Type& InCollectionType)
+{
+	static_assert(ECollectionShareType::CST_All == 4, "Update CollectionViewUtils::ToConfigKey for the updated ECollectionShareType values");
+
+	check(InCollectionType != ECollectionShareType::CST_All);
 
-	/** Convert a string of the form "CollectionName:CollectionType back into its individual elements */
-	bool FromConfigKey(const FString& InKey, FString& OutCollectionName, ECollectionShareType::Type& OutCollectionType)
+	FString CollectionTypeStr;
+	switch(InCollectionType)
 	{
-		static_assert(ECollectionShareType::CST_All == 4, "Update CollectionViewUtils::FromConfigKey for the updated ECollectionShareType values");
+	case ECollectionShareType::CST_System:
+		CollectionTypeStr = "System";
+		break;
+	case ECollectionShareType::CST_Local:
+		CollectionTypeStr = "Local";
+		break;
+	case ECollectionShareType::CST_Private:
+		CollectionTypeStr = "Private";
+		break;
+	case ECollectionShareType::CST_Shared:
+		CollectionTypeStr = "Shared";
+		break;
+	default:
+		break;
+	}
 
-		FString CollectionTypeStr;
-		if(InKey.Split(":", &OutCollectionName, &CollectionTypeStr, ESearchCase::IgnoreCase, ESearchDir::FromEnd))
-		{
-			if(CollectionTypeStr == "System")
-			{
-				OutCollectionType = ECollectionShareType::CST_System;
-			}
-			else if(CollectionTypeStr == "Local")
-			{
-				OutCollectionType = ECollectionShareType::CST_Local;
-			}
-			else if(CollectionTypeStr == "Private")
-			{
-				OutCollectionType = ECollectionShareType::CST_Private;
-			}
-			else if(CollectionTypeStr == "Shared")
-			{
-				OutCollectionType = ECollectionShareType::CST_Shared;
-			}
-			else
-			{
-				return false;
-			}
+	return InCollectionName + ":" + CollectionTypeStr;
+}
 
-			return !OutCollectionName.IsEmpty();
+/** Convert a string of the form "CollectionName:CollectionType back into its individual elements */
+bool FromConfigKey(const FString& InKey, FString& OutCollectionName, ECollectionShareType::Type& OutCollectionType)
+{
+	static_assert(ECollectionShareType::CST_All == 4, "Update CollectionViewUtils::FromConfigKey for the updated ECollectionShareType values");
+
+	FString CollectionTypeStr;
+	if(InKey.Split(":", &OutCollectionName, &CollectionTypeStr, ESearchCase::IgnoreCase, ESearchDir::FromEnd))
+	{
+		if(CollectionTypeStr == "System")
+		{
+			OutCollectionType = ECollectionShareType::CST_System;
+		}
+		else if(CollectionTypeStr == "Local")
+		{
+			OutCollectionType = ECollectionShareType::CST_Local;
+		}
+		else if(CollectionTypeStr == "Private")
+		{
+			OutCollectionType = ECollectionShareType::CST_Private;
+		}
+		else if(CollectionTypeStr == "Shared")
+		{
+			OutCollectionType = ECollectionShareType::CST_Shared;
+		}
+		else
+		{
+			return false;
 		}
 
-		return false;
+		return !OutCollectionName.IsEmpty();
 	}
 
-	// Keep a map of all the collections that have custom colors, so updating the color in one location updates them all
-	static TMap<FString, TSharedPtr<FLinearColor>> CollectionColors;
+	return false;
 }
 
-const TSharedPtr<FLinearColor> CollectionViewUtils::LoadColor(const FString& InCollectionName, const ECollectionShareType::Type& InCollectionType)
+const TOptional<FLinearColor> LoadLegacyColor(const FName InCollectionName, const ECollectionShareType::Type& InCollectionType)
 {
 	check(InCollectionType != ECollectionShareType::CST_All);
 
-	const FString ColorKeyStr = ToConfigKey(InCollectionName, InCollectionType);
+	const FString ColorKeyStr = ToConfigKey(InCollectionName.ToString(), InCollectionType);
 
 	// See if we have a value cached first
 	{
-		TSharedPtr<FLinearColor>* const CachedColor = CollectionColors.Find(ColorKeyStr);
+		TOptional<FLinearColor>* const CachedColor = LegacyCollectionColors.Find(ColorKeyStr);
 		if(CachedColor)
 		{
 			return *CachedColor;
@@ -102,55 +109,35 @@ const TSharedPtr<FLinearColor> CollectionViewUtils::LoadColor(const FString& InC
 		if(GConfig->GetString(TEXT("CollectionColor"), *ColorKeyStr, ColorStr, GEditorPerProjectIni))
 		{
 			FLinearColor Color;
-			if(Color.InitFromString(ColorStr) && !Color.Equals(CollectionViewUtils::GetDefaultColor()))
+			if(Color.InitFromString(ColorStr) && !Color.Equals(CollectionViewUtils::GetLegacyDefaultColor()))
 			{
-				return CollectionColors.Add(ColorKeyStr, MakeShareable(new FLinearColor(Color)));
+				return LegacyCollectionColors.Add(ColorKeyStr, Color);
 			}
 		}
-		else
-		{
-			return CollectionColors.Add(ColorKeyStr, MakeShareable(new FLinearColor(CollectionViewUtils::GetDefaultColor())));
-		}
 	}
 
-	return nullptr;
+	// Cache an empty color to avoid hitting the disk again
+	return LegacyCollectionColors.Add(ColorKeyStr, TOptional<FLinearColor>());
 }
 
-void CollectionViewUtils::SaveColor(const FString& InCollectionName, const ECollectionShareType::Type& InCollectionType, const TSharedPtr<FLinearColor> CollectionColor, const bool bForceAdd)
+void ClearLegacyColor(const FName InCollectionName, const ECollectionShareType::Type& InCollectionType)
 {
 	check(InCollectionType != ECollectionShareType::CST_All);
 
-	const FString ColorKeyStr = ToConfigKey(InCollectionName, InCollectionType);
-
-	// Remove the color if it's invalid or default
-	const bool bRemove = !CollectionColor.IsValid() || (!bForceAdd && CollectionColor->Equals(CollectionViewUtils::GetDefaultColor()));
+	const FString ColorKeyStr = ToConfigKey(InCollectionName.ToString(), InCollectionType);
 
 	// Saves the color of the collection to the config
 	if(FPaths::FileExists(GEditorPerProjectIni))
 	{
 		// If this is no longer custom, remove it
-		if(bRemove)
-		{
-			GConfig->RemoveKey(TEXT("CollectionColor"), *ColorKeyStr, GEditorPerProjectIni);
-		}
-		else
-		{
-			GConfig->SetString(TEXT("CollectionColor"), *ColorKeyStr, *CollectionColor->ToString(), GEditorPerProjectIni);
-		}
+		GConfig->RemoveKey(TEXT("CollectionColor"), *ColorKeyStr, GEditorPerProjectIni);
 	}
 
 	// Update the map too
-	if(bRemove)
-	{
-		CollectionColors.Remove(ColorKeyStr);
-	}
-	else
-	{
-		CollectionColors.Add(ColorKeyStr, CollectionColor);
-	}
+	LegacyCollectionColors.Remove(ColorKeyStr);
 }
 
-bool CollectionViewUtils::HasCustomColors( TArray< FLinearColor >* OutColors )
+bool HasLegacyCustomColors(TArray<FLinearColor>* OutColors)
 {
 	if(!FPaths::FileExists(GEditorPerProjectIni))
 	{
@@ -178,7 +165,7 @@ bool CollectionViewUtils::HasCustomColors( TArray< FLinearColor >* OutColors )
 
 		// Ignore any that have invalid or default colors
 		FLinearColor CurrentColor;
-		if(!CurrentColor.InitFromString(ColorStr) || CurrentColor.Equals(CollectionViewUtils::GetDefaultColor()))
+		if(!CurrentColor.InitFromString(ColorStr) || CurrentColor.Equals(CollectionViewUtils::GetLegacyDefaultColor()))
 		{
 			continue;
 		}
@@ -218,10 +205,70 @@ bool CollectionViewUtils::HasCustomColors( TArray< FLinearColor >* OutColors )
 	return bHasCustom;
 }
 
-FLinearColor CollectionViewUtils::GetDefaultColor()
+FLinearColor ResolveColor(const FName InCollectionName, const ECollectionShareType::Type& InCollectionType)
 {
-	// The default tint the folder should appear as
-	return FLinearColor::Gray;
+	TOptional<FLinearColor> CollectionColor = GetCustomColor(InCollectionName, InCollectionType);
+	if (!CollectionColor)
+	{
+		CollectionColor = GetDefaultColor();
+	}
+	check(CollectionColor);
+	return CollectionColor.GetValue();
 }
 
+TOptional<FLinearColor> GetCustomColor(const FName InCollectionName, const ECollectionShareType::Type& InCollectionType)
+{
+	const FCollectionManagerModule& CollectionManagerModule = FCollectionManagerModule::GetModule();
+	const ICollectionManager& CollectionManager = CollectionManagerModule.Get();
+
+	// First try and use the color set on the collection itself
+	TOptional<FLinearColor> CollectionColor;
+	CollectionManager.GetCollectionColor(InCollectionName, InCollectionType, CollectionColor);
+
+	// Failing that, try and use the legacy local color set for this collection
+	if (!CollectionColor)
+	{
+		CollectionColor = LoadLegacyColor(InCollectionName, InCollectionType);
+	}
+
+	return CollectionColor;
+}
+
+void SetCustomColor(const FName InCollectionName, const ECollectionShareType::Type& InCollectionType, const TOptional<FLinearColor>& CollectionColor)
+{
+	const FCollectionManagerModule& CollectionManagerModule = FCollectionManagerModule::GetModule();
+	ICollectionManager& CollectionManager = CollectionManagerModule.Get();
+
+	if (CollectionManager.SetCollectionColor(InCollectionName, InCollectionType, CollectionColor))
+	{
+		// Set correctly, so clear any legacy color
+		ClearLegacyColor(InCollectionName, InCollectionType);
+	}
+}
+
+bool HasCustomColors(TArray<FLinearColor>* OutColors)
+{
+	const FCollectionManagerModule& CollectionManagerModule = FCollectionManagerModule::GetModule();
+	const ICollectionManager& CollectionManager = CollectionManagerModule.Get();
+
+	bool bHasColors = false;
+
+	// First get the colors from the collection manager
+	bHasColors |= CollectionManager.HasCollectionColors(OutColors);
+
+	// Then add in any legacy colors
+	bHasColors |= HasLegacyCustomColors(OutColors);
+
+	return bHasColors;
+}
+
+FLinearColor GetDefaultColor()
+{
+	// Use the selection accent color as the default
+	const FSlateColor NewSlateColor = FEditorStyle::GetSlateColor("SelectionColor");
+	return NewSlateColor.IsColorSpecified() ? NewSlateColor.GetSpecifiedColor() : FLinearColor::White;
+}
+
+} // namespace CollectionViewUtils
+
 #undef LOCTEXT_NAMESPACE
diff --git a/Engine/Source/Editor/ContentBrowser/Private/CollectionViewUtils.h b/Engine/Source/Editor/ContentBrowser/Private/CollectionViewUtils.h
--- a/Engine/Source/Editor/ContentBrowser/Private/CollectionViewUtils.h
+++ b/Engine/Source/Editor/ContentBrowser/Private/CollectionViewUtils.h
@@ -8,23 +8,31 @@
 namespace CollectionViewUtils
 {
 	/**
-	 * Loads the color of this collection from the config
+	 * Resolve the color of the collection, falling back to the default if no custom color is set
+	 *
+	 * @param InCollectionName - The name of the collection to resolve the color for
+	 * @param InCollectionType - The type of the collection to resolve the color for
+	 * @return The color the collection should appear as
+	 */
+	FLinearColor ResolveColor(const FName InCollectionName, const ECollectionShareType::Type& InCollectionType);
+
+	/**
+	 * Gets the color of the collection
 	 *
 	 * @param InCollectionName - The name of the collection to get the color for
 	 * @param InCollectionType - The type of the collection to get the color for
-	 * @return The color the collection should appear as, will be null if not customized
+	 * @return The color the collection should appear as, will be empty if not customized
 	 */
-	const TSharedPtr<FLinearColor> LoadColor(const FString& InCollectionName, const ECollectionShareType::Type& InCollectionType);
+	TOptional<FLinearColor> GetCustomColor(const FName InCollectionName, const ECollectionShareType::Type& InCollectionType);
 
 	/**
-	 * Saves the color of the collection to the config
+	 * Sets the color of the collection
 	 *
 	 * @param InCollectionName - The name of the collection to set the color for
 	 * @param InCollectionType - The type of the collection to set the color for
-	 * @param CollectionColor - The color the collection should appear as
-	 * @param bForceAdd - If true, force the color to be added even if it's the default color
+	 * @param CollectionColor - The color the collection should appear as, or an empty color to clear it
 	 */
-	void SaveColor(const FString& InCollectionName, const ECollectionShareType::Type& InCollectionType, const TSharedPtr<FLinearColor> CollectionColor, const bool bForceAdd = false);
+	void SetCustomColor(const FName InCollectionName, const ECollectionShareType::Type& InCollectionType, const TOptional<FLinearColor>& CollectionColor);
 
 	/**
 	 * Checks to see if any collection has a custom color, optionally outputs them to a list
@@ -32,7 +40,7 @@ namespace CollectionViewUtils
 	 * @param OutColors - If specified, returns all the custom colors being used
 	 * @return true, if custom colors are present
 	 */
-	bool HasCustomColors( TArray< FLinearColor >* OutColors = nullptr );
+	bool HasCustomColors(TArray<FLinearColor>* OutColors = nullptr);
 
 	/** Gets the default color the collection should appear as */
 	FLinearColor GetDefaultColor();
diff --git a/Engine/Source/Editor/ContentBrowser/Private/SAssetPicker.cpp b/Engine/Source/Editor/ContentBrowser/Private/SAssetPicker.cpp
--- a/Engine/Source/Editor/ContentBrowser/Private/SAssetPicker.cpp
+++ b/Engine/Source/Editor/ContentBrowser/Private/SAssetPicker.cpp
@@ -292,7 +292,6 @@ void SAssetPicker::Construct( const FArguments& InArgs )
 		.FilterRecursivelyWithBackendFilter( false )
 		.CanShowRealTimeThumbnails( InArgs._AssetPickerConfig.bCanShowRealTimeThumbnails )
 		.CanShowDevelopersFolder( InArgs._AssetPickerConfig.bCanShowDevelopersFolder )
-		.CanShowCollections( false )
 		.PreloadAssetsForContextMenu( InArgs._AssetPickerConfig.bPreloadAssetsForContextMenu )
 		.HighlightedText( HighlightText )
 		.ThumbnailLabel( ThumbnailLabel )
diff --git a/Engine/Source/Editor/ContentBrowser/Private/SAssetView.cpp b/Engine/Source/Editor/ContentBrowser/Private/SAssetView.cpp
--- a/Engine/Source/Editor/ContentBrowser/Private/SAssetView.cpp
+++ b/Engine/Source/Editor/ContentBrowser/Private/SAssetView.cpp
@@ -419,7 +419,6 @@ void SAssetView::Construct( const FArguments& InArgs )
 
 	bCanShowDevelopersFolder = InArgs._CanShowDevelopersFolder;
 
-	bCanShowCollections = InArgs._CanShowCollections;
 	bCanShowFavorites = InArgs._CanShowFavorites;
 	bPreloadAssetsForContextMenu = InArgs._PreloadAssetsForContextMenu;
 
@@ -3160,19 +3159,6 @@ TSharedRef<SWidget> SAssetView::GetViewButtonContent()
 			EUserInterfaceActionType::ToggleButton
 		);
 
-		MenuBuilder.AddMenuEntry(
-			LOCTEXT("ShowCollectionOption", "Show Collections"),
-			LOCTEXT("ShowCollectionOptionToolTip", "Show the collections list in the view?"),
-			FSlateIcon(),
-			FUIAction(
-				FExecuteAction::CreateSP( this, &SAssetView::ToggleShowCollections ),
-				FCanExecuteAction::CreateSP( this, &SAssetView::IsToggleShowCollectionsAllowed ),
-				FIsActionChecked::CreateSP( this, &SAssetView::IsShowingCollections )
-			),
-			NAME_None,
-			EUserInterfaceActionType::ToggleButton
-		);
-
 		MenuBuilder.AddMenuEntry(
 			LOCTEXT("ShowFavoriteOptions", "Show Favorites"),
 			LOCTEXT("ShowFavoriteOptionToolTip", "Show the favorite folders in the view?"),
@@ -3531,24 +3517,6 @@ bool SAssetView::IsShowingLocalizedContent() const
 	return IsToggleShowLocalizedContentAllowed() && GetDefault<UContentBrowserSettings>()->GetDisplayL10NFolder();
 }
 
-void SAssetView::ToggleShowCollections()
-{
-	const bool bDisplayCollections = GetDefault<UContentBrowserSettings>()->GetDisplayCollections();
-	GetMutableDefault<UContentBrowserSettings>()->SetDisplayCollections( !bDisplayCollections );
-	GetMutableDefault<UContentBrowserSettings>()->PostEditChange();
-}
-
-bool SAssetView::IsToggleShowCollectionsAllowed() const
-{
-	return bCanShowCollections;
-}
-
-bool SAssetView::IsShowingCollections() const
-{
-	return IsToggleShowCollectionsAllowed() && GetDefault<UContentBrowserSettings>()->GetDisplayCollections();
-}
-
-
 void SAssetView::ToggleShowFavorites()
 {
 	const bool bShowingFavorites = GetDefault<UContentBrowserSettings>()->GetDisplayFavorites();
@@ -4842,7 +4810,14 @@ FText SAssetView::GetAssetShowWarningText() const
 
 	if ( SourcesData.HasCollections() && !SourcesData.IsDynamicCollection() )
 	{
-		DropText = LOCTEXT( "DragAssetsHere", "Drag and drop assets here to add them to the collection." );
+		if (SourcesData.Collections[0].Name.IsNone())
+		{
+			DropText = LOCTEXT("NoCollectionSelected", "No collection selected.");
+		}
+		else
+		{
+			DropText = LOCTEXT("DragAssetsHere", "Drag and drop assets here to add them to the collection.");
+		}
 	}
 	else if ( OnGetAssetContextMenu.IsBound() )
 	{
diff --git a/Engine/Source/Editor/ContentBrowser/Private/SAssetView.h b/Engine/Source/Editor/ContentBrowser/Private/SAssetView.h
--- a/Engine/Source/Editor/ContentBrowser/Private/SAssetView.h
+++ b/Engine/Source/Editor/ContentBrowser/Private/SAssetView.h
@@ -57,7 +57,6 @@ public:
 		, _FilterRecursivelyWithBackendFilter(true)
 		, _CanShowRealTimeThumbnails(false)
 		, _CanShowDevelopersFolder(false)
-		, _CanShowCollections(false)
 		, _CanShowFavorites(false)
 		, _PreloadAssetsForContextMenu(true)
 		, _SelectionMode( ESelectionMode::Multi )
@@ -162,9 +161,6 @@ public:
 		/** Indicates if the 'Show Developers' option should be enabled or disabled */
 		SLATE_ARGUMENT( bool, CanShowDevelopersFolder )
 
-		/** Indicates if the 'Show Collections' option should be enabled or disabled */
-		SLATE_ARGUMENT( bool, CanShowCollections )
-
 		/** Indicates if the 'Show Favorites' option should be enabled or disabled */
 		SLATE_ARGUMENT(bool, CanShowFavorites)
 
@@ -512,15 +508,6 @@ private:
 	/** @return true when we are showing the developers content */
 	bool IsShowingDevelopersContent() const;
 
-	/** Toggle whether collections should be shown or not */
-	void ToggleShowCollections();
-
-	/** Whether or not it's possible to toggle collections */
-	bool IsToggleShowCollectionsAllowed() const;
-
-	/** @return true when we are showing collections */
-	bool IsShowingCollections() const;
-
 	/** Toggle whether favorites should be shown or not */
 	void ToggleShowFavorites();
 
@@ -983,9 +970,6 @@ private:
 	/** Indicates if the 'Show Developers' option should be enabled or disabled */
 	bool bCanShowDevelopersFolder;
 
-	/** Indicates if the 'Show Collections' option should be enabled or disabled */
-	bool bCanShowCollections;
-
 	/** Indicates if the 'Show Favorites' option should be enabled or disabled */
 	bool bCanShowFavorites;
 
diff --git a/Engine/Source/Editor/ContentBrowser/Private/SCollectionView.cpp b/Engine/Source/Editor/ContentBrowser/Private/SCollectionView.cpp
--- a/Engine/Source/Editor/ContentBrowser/Private/SCollectionView.cpp
+++ b/Engine/Source/Editor/ContentBrowser/Private/SCollectionView.cpp
@@ -22,8 +22,10 @@
 
 #include "CollectionAssetManagement.h"
 #include "CollectionContextMenu.h"
+#include "CollectionViewUtils.h"
 #include "DragAndDrop/AssetDragDropOp.h"
 #include "DragAndDrop/CollectionDragDropOp.h"
+#include "SourcesSearch.h"
 #include "SourcesViewWidgets.h"
 #include "ContentBrowserModule.h"
 #include "Widgets/Layout/SExpandableArea.h"
@@ -119,6 +121,47 @@ void SCollectionView::Construct( const FArguments& InArgs )
 		CollectionListContextMenuOpening = FOnContextMenuOpening::CreateSP( this, &SCollectionView::MakeCollectionTreeContextMenu );
 	}
 
+	SearchPtr = InArgs._ExternalSearch;
+	if (!SearchPtr)
+	{
+		SearchPtr = MakeShared<FSourcesSearch>();
+		SearchPtr->Initialize();
+		SearchPtr->SetHintText(LOCTEXT("CollectionsViewSearchBoxHint", "Search Collections"));
+	}
+	SearchPtr->OnSearchChanged().AddSP(this, &SCollectionView::SetCollectionsSearchFilterText);
+
+	TSharedRef<SHorizontalBox> TitleContent = SNew(SHorizontalBox);
+	if (InArgs._ExternalSearch)
+	{
+		// If using an external search then just show the title and don't hide it when collapsing
+		TitleContent->AddSlot()
+		[
+			SNew(STextBlock)
+			.Font(FEditorStyle::GetFontStyle("ContentBrowser.SourceTitleFont"))
+			.Text(LOCTEXT("CollectionsListTitle", "Collections"))
+		];
+	}
+	else
+	{
+		// If using an internal search then show the title or search box depending on whether we're collapsed or not
+		TitleContent->AddSlot()
+		[
+			SNew(STextBlock)
+			.Font(FEditorStyle::GetFontStyle("ContentBrowser.SourceTitleFont"))
+			.Text(LOCTEXT("CollectionsListTitle", "Collections"))
+			.Visibility(this, &SCollectionView::GetCollectionsTitleTextVisibility)
+		];
+
+		TitleContent->AddSlot()
+		[
+			SNew(SBox)
+			.Visibility(this, &SCollectionView::GetCollectionsSearchBoxVisibility)
+			[
+				SearchPtr->GetWidget()
+			]
+		];
+	}
+
 	PreventSelectionChangedDelegateCount = 0;
 
 	TSharedRef< SWidget > HeaderContent = SNew(SHorizontalBox)
@@ -126,22 +169,24 @@ void SCollectionView::Construct( const FArguments& InArgs )
 			.FillWidth(1.0f)
 			.Padding(0.0f)
 			[
-				SNew(SHorizontalBox)
-
-				+ SHorizontalBox::Slot()
-				[
-					SNew(STextBlock)
-					.Font( FEditorStyle::GetFontStyle("ContentBrowser.SourceTitleFont") )
-					.Text( LOCTEXT("CollectionsListTitle", "Collections") )
-					.Visibility( this, &SCollectionView::GetCollectionsTitleTextVisibility )
-				]
+				TitleContent
+			]
 
-				+ SHorizontalBox::Slot()
+			+SHorizontalBox::Slot()
+			.AutoWidth()
+			.VAlign(VAlign_Center)
+			.Padding(2.0f, 0.0f, 0.0f, 0.0f)
+			[
+				SNew(SButton)
+				.ButtonStyle(FEditorStyle::Get(), "FlatButton")
+				.ToolTipText(this, &SCollectionView::GetSwitchCollectionViewModeToolTipText)
+				.OnClicked(this, &SCollectionView::OnSwitchCollectionViewMode)
+				.ContentPadding(FMargin(2, 2))
+				.ForegroundColor(FEditorStyle::GetSlateColor("DefaultForeground"))
+				.Visibility(this, &SCollectionView::GetCollectionButtonsVisibility)
 				[
-					SAssignNew(SearchBoxPtr, SSearchBox)
-					.HintText( LOCTEXT( "CollectionsViewSearchBoxHint", "Search Collections" ) )
-					.OnTextChanged( this, &SCollectionView::SetCollectionsSearchFilterText )
-					.Visibility( this, &SCollectionView::GetCollectionsSearchBoxVisibility )
+					SNew(SImage)
+					.Image(this, &SCollectionView::GetSwitchCollectionViewModeIcon)
 				]
 			]
 
@@ -155,9 +200,10 @@ void SCollectionView::Construct( const FArguments& InArgs )
 				.ToolTipText(LOCTEXT("AddCollectionButtonTooltip", "Add a collection."))
 				.OnClicked(this, &SCollectionView::MakeAddCollectionMenu)
 				.ContentPadding( FMargin(2, 2) )
-				.Visibility(this, &SCollectionView::GetAddCollectionButtonVisibility)
+				.Visibility(this, &SCollectionView::GetCollectionButtonsVisibility)
 				[
-					SNew(SImage) .Image( FEditorStyle::GetBrush("ContentBrowser.AddCollectionButtonIcon") )
+					SNew(SImage)
+					.Image( FEditorStyle::GetBrush("ContentBrowser.AddCollectionButtonIcon") )
 				]
 			];
 
@@ -165,8 +211,10 @@ void SCollectionView::Construct( const FArguments& InArgs )
 			// Separator
 			+SVerticalBox::Slot()
 			.AutoHeight()
+			.Padding(0, 0, 0, 1)
 			[
 				SNew(SSeparator)
+				.Visibility(InArgs._ShowSeparator ? EVisibility::Visible : EVisibility::Collapsed)
 			]
 
 			// Collections tree
@@ -211,6 +259,7 @@ void SCollectionView::Construct( const FArguments& InArgs )
 		Content = SNew( SVerticalBox )
 		+SVerticalBox::Slot()
 		.AutoHeight()
+		.Padding(FMargin(12.0f, 0.0f, 0.0f, 0.0f))
 		[
 			HeaderContent
 		]
@@ -285,6 +334,7 @@ void SCollectionView::HandleCollectionUpdated( const FCollectionNameType& Collec
 	if (CollectionItemToUpdate.IsValid())
 	{
 		bQueueSCCRefresh = true;
+		CollectionItemToUpdate->CollectionColor = CollectionViewUtils::ResolveColor(Collection.Name, Collection.Type);
 		UpdateCollectionItemStatus(CollectionItemToUpdate.ToSharedRef());
 	}
 }
@@ -319,6 +369,7 @@ void SCollectionView::HandleSourceControlStateChanged()
 
 void SCollectionView::UpdateCollectionItemStatus( const TSharedRef<FCollectionItem>& CollectionItem )
 {
+	int32 NewObjectCount = 0;
 	TOptional<ECollectionItemStatus> NewStatus;
 
 	// Check IsModuleAvailable as we might be in the process of shutting down, and were notified due to the SCC provider being nulled out...
@@ -329,6 +380,8 @@ void SCollectionView::UpdateCollectionItemStatus( const TSharedRef<FCollectionIt
 		FCollectionStatusInfo StatusInfo;
 		if (CollectionManagerModule.Get().GetCollectionStatusInfo(CollectionItem->CollectionName, CollectionItem->CollectionType, StatusInfo))
 		{
+			NewObjectCount = StatusInfo.NumObjects;
+
 			// Test the SCC state first as this should take priority when reporting the status back to the user
 			if (StatusInfo.bUseSCC)
 			{
@@ -376,6 +429,7 @@ void SCollectionView::UpdateCollectionItemStatus( const TSharedRef<FCollectionIt
 		}
 	}
 
+	CollectionItem->NumObjects = NewObjectCount;
 	CollectionItem->CurrentStatus = NewStatus.Get(ECollectionItemStatus::IsUpToDateAndEmpty);
 }
 
@@ -420,6 +474,7 @@ void SCollectionView::UpdateCollectionItems()
 				OutAvailableCollections.Add(Collection, CollectionItem);
 
 				CollectionManagerModule.Get().GetCollectionStorageMode(Collection.Name, Collection.Type, CollectionItem->StorageMode);
+				CollectionItem->CollectionColor = CollectionViewUtils::ResolveColor(Collection.Name, Collection.Type);
 
 				SCollectionView::UpdateCollectionItemStatus(CollectionItem);
 
@@ -532,10 +587,15 @@ void SCollectionView::UpdateFilteredCollectionItems()
 	CollectionTreePtr->RequestTreeRefresh();
 }
 
-void SCollectionView::SetCollectionsSearchFilterText( const FText& InSearchText )
+void SCollectionView::SetCollectionsSearchFilterText( const FText& InSearchText, TArray<FText>& OutErrors )
 {
 	CollectionItemTextFilter->SetRawFilterText( InSearchText );
-	SearchBoxPtr->SetError( CollectionItemTextFilter->GetFilterErrorText() );
+	
+	const FText ErrorText = CollectionItemTextFilter->GetFilterErrorText();
+	if (!ErrorText.IsEmpty())
+	{
+		OutErrors.Add(ErrorText);
+	}
 }
 
 FText SCollectionView::GetCollectionsSearchFilterText() const
@@ -573,7 +633,7 @@ void SCollectionView::SetSelectedCollections(const TArray<FCollectionNameType>&
 			// If the selected collection doesn't pass our current filter, we need to clear it
 			if (bEnsureVisible && !CollectionItemTextFilter->PassesFilter(*CollectionItemToSelect))
 			{
-				SearchBoxPtr->SetText(FText::GetEmpty());
+				SearchPtr->ClearSearch();
 			}
 		}
 	}
@@ -672,6 +732,8 @@ void SCollectionView::SaveSettings(const FString& IniFilename, const FString& In
 		CollectionTreePtr->GetExpandedItems(ExpandedCollectionItems);
 		SaveCollectionsArrayToIni(TEXT(".ExpandedCollections"), ExpandedCollectionItems.Array());
 	}
+
+	GConfig->SetInt(*IniSection, *(SettingsString + TEXT(".ViewMode")), (int32)CollectionViewMode, IniFilename);
 }
 
 void SCollectionView::LoadSettings(const FString& IniFilename, const FString& IniSection, const FString& SettingsString)
@@ -730,6 +792,15 @@ void SCollectionView::LoadSettings(const FString& IniFilename, const FString& In
 	{
 		SetExpandedCollections(NewExpandedCollections);
 	}
+
+	// View Mode
+	{
+		int32 CollectionViewModeInt = 0;
+		GConfig->GetInt(*IniSection, *(SettingsString + TEXT(".ViewMode")), CollectionViewModeInt, IniFilename);
+		CollectionViewMode = (EAssetTagItemViewMode)CollectionViewModeInt;
+
+		CollectionTreePtr->RebuildList();
+	}
 }
 
 void SCollectionView::Tick( const FGeometry& AllottedGeometry, const double InCurrentTime, const float InDeltaTime )
@@ -896,11 +967,66 @@ EVisibility SCollectionView::GetCollectionsSearchBoxVisibility() const
 	return (!CollectionsExpandableAreaPtr.IsValid() || CollectionsExpandableAreaPtr->IsExpanded()) ? EVisibility::Visible : EVisibility::Collapsed;
 }
 
-EVisibility SCollectionView::GetAddCollectionButtonVisibility() const
+EVisibility SCollectionView::GetCollectionButtonsVisibility() const
 {
 	return (bAllowCollectionButtons && ( !CollectionsExpandableAreaPtr.IsValid() || CollectionsExpandableAreaPtr->IsExpanded() ) ) ? EVisibility::Visible : EVisibility::Collapsed;
 }
 
+const FSlateBrush* SCollectionView::GetSwitchCollectionViewModeIcon() const
+{
+	switch (CollectionViewMode)
+	{
+	case EAssetTagItemViewMode::Standard:
+		return FEditorStyle::GetBrush("ContentBrowser.Sources.Collections.Compact");
+
+	case EAssetTagItemViewMode::Compact:
+		return FEditorStyle::GetBrush("ContentBrowser.Sources.Collections");
+
+	default:
+		break;
+	}
+	check(false);
+	return nullptr;
+}
+
+FText SCollectionView::GetSwitchCollectionViewModeToolTipText() const
+{
+	switch (CollectionViewMode)
+	{
+	case EAssetTagItemViewMode::Standard:
+		return LOCTEXT("SwitchToCompactView_ToolTip", "Switch to compact view");
+
+	case EAssetTagItemViewMode::Compact:
+		return LOCTEXT("SwitchToStanardView_ToolTip", "Switch to standard view");
+
+	default:
+		break;
+	}
+	check(false);
+	return FText();
+}
+
+FReply SCollectionView::OnSwitchCollectionViewMode()
+{
+	switch (CollectionViewMode)
+	{
+	case EAssetTagItemViewMode::Standard:
+		CollectionViewMode = EAssetTagItemViewMode::Compact;
+		break;
+
+	case EAssetTagItemViewMode::Compact:
+		CollectionViewMode = EAssetTagItemViewMode::Standard;
+		break;
+
+	default:
+		check(false);
+		break;
+	}
+
+	CollectionTreePtr->RebuildList();
+	return FReply::Handled();
+}
+
 void SCollectionView::CreateCollectionItem( ECollectionShareType::Type CollectionType, ECollectionStorageMode::Type StorageMode, const FCreateCollectionPayload& InCreationPayload )
 {
 	if ( ensure(CollectionType != ECollectionShareType::CST_All) )
@@ -909,12 +1035,12 @@ void SCollectionView::CreateCollectionItem( ECollectionShareType::Type Collectio
 
 		const FName BaseCollectionName = *LOCTEXT("NewCollectionName", "NewCollection").ToString();
 		FName CollectionName;
-		CollectionManagerModule.Get().CreateUniqueCollectionName(BaseCollectionName, CollectionType, CollectionName);
+		CollectionManagerModule.Get().CreateUniqueCollectionName(BaseCollectionName, ECollectionShareType::CST_All, CollectionName);
 		TSharedPtr<FCollectionItem> NewItem = MakeShareable(new FCollectionItem(CollectionName, CollectionType));
 		NewItem->StorageMode = StorageMode;
 
 		// Adding a new collection now, so clear any filter we may have applied
-		SearchBoxPtr->SetText(FText::GetEmpty());
+		SearchPtr->ClearSearch();
 
 		if ( InCreationPayload.ParentCollection.IsSet() )
 		{
@@ -1115,14 +1241,17 @@ TSharedRef<ITableRow> SCollectionView::GenerateCollectionRow( TSharedPtr<FCollec
 		}
 	}
 
-	TSharedPtr< STableRow< TSharedPtr<FCollectionItem> > > TableRow = SNew( STableRow< TSharedPtr<FCollectionItem> >, OwnerTable )
+	TSharedPtr< SAssetTagItemTableRow< TSharedPtr<FCollectionItem> > > TableRow = SNew( SAssetTagItemTableRow< TSharedPtr<FCollectionItem> >, OwnerTable )
 		.OnDragDetected(this, &SCollectionView::OnCollectionDragDetected);
 
+	TSharedPtr<SCollectionTreeItem> CollectionTreeItem;
+
 	TableRow->SetContent
 		(
-			SNew(SCollectionTreeItem)
+			SAssignNew(CollectionTreeItem, SCollectionTreeItem)
 			.ParentWidget(SharedThis(this))
 			.CollectionItem(CollectionItem)
+			.ViewMode(CollectionViewMode)
 			.OnNameChangeCommit(this, &SCollectionView::CollectionNameChangeCommit)
 			.OnVerifyRenameCommit(this, &SCollectionView::CollectionVerifyRenameCommit)
 			.OnValidateDragDrop(this, &SCollectionView::ValidateDragDropOnCollectionItem)
@@ -1135,6 +1264,8 @@ TSharedRef<ITableRow> SCollectionView::GenerateCollectionRow( TSharedPtr<FCollec
 			.OnCollectionCheckStateChanged(OnCollectionCheckStateChangedDelegate)
 		);
 
+	TableRow->SetIsDropTarget(MakeAttributeSP(CollectionTreeItem.Get(), &SCollectionTreeItem::IsDraggedOver));
+
 	return TableRow.ToSharedRef();
 }
 
@@ -1508,7 +1639,7 @@ bool SCollectionView::CollectionNameChangeCommit( const TSharedPtr< FCollectionI
 		// If the new name doesn't pass our current filter, we need to clear it
 		if ( !CollectionItemTextFilter->PassesFilter( FCollectionItem(NewNameFinal, CollectionItem->CollectionType) ) )
 		{
-			SearchBoxPtr->SetText(FText::GetEmpty());
+			SearchPtr->ClearSearch();
 		}
 
 		// Otherwise perform the rename
@@ -1552,7 +1683,7 @@ bool SCollectionView::CollectionVerifyRenameCommit(const TSharedPtr< FCollection
 
 	FCollectionManagerModule& CollectionManagerModule = FCollectionManagerModule::GetModule();
 
-	if (!CollectionManagerModule.Get().IsValidCollectionName(NewName, ECollectionShareType::CST_Shared))
+	if (!CollectionManagerModule.Get().IsValidCollectionName(NewName, ECollectionShareType::CST_All))
 	{
 		OutErrorMessage = CollectionManagerModule.Get().GetLastError();
 		return false;
diff --git a/Engine/Source/Editor/ContentBrowser/Private/SCollectionView.h b/Engine/Source/Editor/ContentBrowser/Private/SCollectionView.h
--- a/Engine/Source/Editor/ContentBrowser/Private/SCollectionView.h
+++ b/Engine/Source/Editor/ContentBrowser/Private/SCollectionView.h
@@ -9,6 +9,7 @@
 #include "Input/Reply.h"
 #include "Widgets/SCompoundWidget.h"
 #include "AssetData.h"
+#include "AssetTagItemTypes.h"
 #include "CollectionManagerTypes.h"
 #include "Widgets/Views/STableViewBase.h"
 #include "Widgets/Views/STableRow.h"
@@ -19,6 +20,7 @@
 class FCollectionAssetManagement;
 class FCollectionContextMenu;
 class FCollectionDragDropOp;
+class FSourcesSearch;
 class FUICommandList;
 struct FHistoryData;
 
@@ -38,6 +40,7 @@ public:
 		, _AllowContextMenu(true)
 		, _AllowCollectionDrag(false)
 		, _AllowQuickAssetManagement(false)
+		, _ShowSeparator(true)
 		{}
 
 		/** Called when a collection was selected */
@@ -55,6 +58,12 @@ public:
 		/** If true, check boxes that let you quickly add/remove the current selection from a collection will be displayed */
 		SLATE_ARGUMENT( bool, AllowQuickAssetManagement )
 
+		/** If true, The tree search bar separator be displayed */
+		SLATE_ARGUMENT( bool, ShowSeparator )
+
+		/** Optional external search. Will hide and replace our internal search UI */
+		SLATE_ARGUMENT( TSharedPtr<FSourcesSearch>, ExternalSearch )
+
 	SLATE_END_ARGS()
 
 	/** Constructs this widget with InArgs */
@@ -142,9 +151,18 @@ private:
 	/** Gets the visibility of the collections search box */
 	EVisibility GetCollectionsSearchBoxVisibility() const;
 
-	/** Gets the visibility of the add collection button */
-	EVisibility GetAddCollectionButtonVisibility() const;
+	/** Gets the visibility of the collection buttons */
+	EVisibility GetCollectionButtonsVisibility() const;
 	
+	/** Get the icon used on the collection view mode switcher button */
+	const FSlateBrush* GetSwitchCollectionViewModeIcon() const;
+
+	/** Get the tooltip text used on the collection view mode switcher button */
+	FText GetSwitchCollectionViewModeToolTipText() const;
+
+	/** Handle the collection view mode switcher button being clicked */
+	FReply OnSwitchCollectionViewMode();
+
 	/** Sets up an inline creation process for a new collection of the specified type */
 	void CreateCollectionItem( ECollectionShareType::Type CollectionType, ECollectionStorageMode::Type StorageMode, const FCreateCollectionPayload& InCreationPayload );
 
@@ -248,7 +266,7 @@ private:
 	void UpdateFilteredCollectionItems();
 
 	/** Set the active filter text */
-	void SetCollectionsSearchFilterText( const FText& InSearchText );
+	void SetCollectionsSearchFilterText( const FText& InSearchText, TArray<FText>& OutErrors );
 
 	/** Get the active filter text */
 	FText GetCollectionsSearchFilterText() const;
@@ -275,8 +293,8 @@ private:
 		TSharedRef<SCollectionView> CollectionView;
 	};
 
-	/** The collection list search box */
-	TSharedPtr< SSearchBox > SearchBoxPtr;
+	/** The collection search interface */
+	TSharedPtr<FSourcesSearch> SearchPtr;
 
 	/** The collection tree widget */
 	TSharedPtr< STreeView< TSharedPtr<FCollectionItem> > > CollectionTreePtr;
@@ -316,6 +334,9 @@ private:
 	/** True when a drag is over this view with a valid operation for drop */
 	bool bDraggedOver;
 
+	/** The view mode that collections within this view should use */
+	EAssetTagItemViewMode CollectionViewMode = EAssetTagItemViewMode::Standard;
+
 	/** If > 0, the selection changed delegate will not be called. Used to update the tree from an external source or in certain bulk operations. */
 	int32 PreventSelectionChangedDelegateCount;
 
diff --git a/Engine/Source/Editor/ContentBrowser/Private/SContentBrowser.cpp b/Engine/Source/Editor/ContentBrowser/Private/SContentBrowser.cpp
--- a/Engine/Source/Editor/ContentBrowser/Private/SContentBrowser.cpp
+++ b/Engine/Source/Editor/ContentBrowser/Private/SContentBrowser.cpp
@@ -19,10 +19,13 @@
 #include "Widgets/Images/SImage.h"
 #include "Widgets/Text/STextBlock.h"
 #include "Widgets/Layout/SBox.h"
+#include "Widgets/Layout/SExpandableArea.h"
 #include "Framework/MultiBox/MultiBoxBuilder.h"
 #include "Widgets/Input/SButton.h"
 #include "Widgets/Input/SComboButton.h"
+#include "Widgets/Input/SSearchBox.h"
 #include "Widgets/Layout/SSplitter.h"
+#include "Widgets/Layout/SWidgetSwitcher.h"
 #include "Framework/Docking/TabManager.h"
 #include "EditorStyleSet.h"
 #include "EditorFontGlyphs.h"
@@ -31,13 +34,14 @@
 #include "Editor.h"
 #include "FileHelpers.h"
 #include "AssetRegistryModule.h"
+#include "AssetRegistryState.h"
 #include "AssetToolsModule.h"
 #include "Widgets/Navigation/SBreadcrumbTrail.h"
 #include "ContentBrowserLog.h"
 #include "FrontendFilters.h"
 #include "ContentBrowserSingleton.h"
 #include "ContentBrowserUtils.h"
-#include "SAssetSearchBox.h"
+#include "SourcesSearch.h"
 #include "SFilterList.h"
 #include "SPathView.h"
 #include "SCollectionView.h"
@@ -61,6 +65,12 @@
 
 const FString SContentBrowser::SettingsIniSection = TEXT("ContentBrowser");
 
+namespace ContentBrowserSourcesWidgetSwitcherIndex
+{
+	static const int32 PathView = 0;
+	static const int32 CollectionsView = 1;
+}
+
 SContentBrowser::~SContentBrowser()
 {
 	// Remove the listener for when view settings are changed
@@ -95,7 +105,6 @@ void SContentBrowser::Construct( const FArguments& InArgs, const FName& InInstan
 	}
 	
 	bIsLocked = InArgs._InitiallyLocked;
-	bAlwaysShowCollections = Config != nullptr ? Config->bAlwaysShowCollections : false;
 	bCanSetAsPrimaryBrowser = Config != nullptr ? Config->bCanSetAsPrimaryBrowser : true;
 
 	HistoryManager.SetOnApplyHistoryData(FOnApplyHistoryData::CreateSP(this, &SContentBrowser::OnApplyHistoryData));
@@ -111,6 +120,10 @@ void SContentBrowser::Construct( const FArguments& InArgs, const FName& InInstan
 	FrontendFilters = MakeShareable(new FAssetFilterCollectionType());
 	TextFilter = MakeShareable( new FFrontendFilter_Text() );
 
+	SourcesSearch = MakeShared<FSourcesSearch>();
+	SourcesSearch->Initialize();
+	SourcesSearch->SetHintText(MakeAttributeSP(this, &SContentBrowser::GetSourcesSearchHintText));
+
 	static const FName DefaultForegroundName("DefaultForeground");
 
 	BindCommands();
@@ -436,114 +449,157 @@ void SContentBrowser::Construct( const FArguments& InArgs, const FName& InInstan
 			+ SSplitter::Slot()
 			.Value(0.25f)
 			[
-				SAssignNew(PathFavoriteSplitterPtr, SSplitter)
-				.Orientation(EOrientation::Orient_Vertical)
-				.MinimumSlotHeight(70.0f)
-				.Visibility( this, &SContentBrowser::GetSourcesViewVisibility )
-				+ SSplitter::Slot()
-				.Value(.2f)
+				SNew(SBorder)
+				.Padding(FMargin(3))
+				.BorderImage(FEditorStyle::GetBrush("ToolPanel.GroupBorder"))
+				.Visibility(this, &SContentBrowser::GetSourcesViewVisibility)
 				[
-					SNew(SBorder)
+					SNew(SVerticalBox)
 
-					.Visibility(this, &SContentBrowser::GetFavoriteFolderVisibility)
-					.Padding(FMargin(3))
-					.BorderImage(FEditorStyle::GetBrush("ToolPanel.GroupBorder"))
+					// Paths expansion/search
+					+SVerticalBox::Slot()
+					.AutoHeight()
 					[
-						SAssignNew(FavoritePathViewPtr, SFavoritePathView)
-						.OnPathSelected(this, &SContentBrowser::FavoritePathSelected)
-						.OnGetFolderContextMenu(this, &SContentBrowser::GetFolderContextMenu, true)
-						.OnGetPathContextMenuExtender(this, &SContentBrowser::GetPathContextMenuExtender)
-						.FocusSearchBoxWhenOpened(false)
-						.ShowTreeTitle(true)
-						.ShowSeparator(false)
-						.AllowClassesFolder(true)
-						.SearchContent()
+						SNew(SHorizontalBox)
+						.AddMetaData<FTagMetaData>(FTagMetaData(TEXT("ContentBrowserSourcesToggle1")))
+					
+						+SHorizontalBox::Slot()
+						.AutoWidth()
+						.Padding(0, 0, 2, 0)
 						[
-							SNew(SVerticalBox)
-							.AddMetaData<FTagMetaData>(FTagMetaData(TEXT("ContentBrowserSourcesToggle1")))
-							+ SVerticalBox::Slot()
-							.FillHeight(1.0f)
-							.Padding(0, 0, 2, 0)
+							SNew(SButton)
+							.VAlign(EVerticalAlignment::VAlign_Center)
+							.ButtonStyle(FEditorStyle::Get(), "ToggleButton")
+							.ToolTipText(LOCTEXT("SourcesTreeToggleTooltip", "Show or hide the sources panel"))
+							.ContentPadding(FMargin(1, 0))
+							.ForegroundColor(FEditorStyle::GetSlateColor(DefaultForegroundName))
+							.OnClicked(this, &SContentBrowser::SourcesViewExpandClicked)
 							[
-								SNew(SButton)
-								.VAlign(EVerticalAlignment::VAlign_Center)
-								.ButtonStyle(FEditorStyle::Get(), "ToggleButton")
-								.ToolTipText(LOCTEXT("SourcesTreeToggleTooltip", "Show or hide the sources panel"))
-								.ContentPadding(FMargin(1, 0))
-								.ForegroundColor(FEditorStyle::GetSlateColor(DefaultForegroundName))
-								.OnClicked(this, &SContentBrowser::SourcesViewExpandClicked)
-								[
-									SNew(SImage)
-									.Image(this, &SContentBrowser::GetSourcesToggleImage)
-								]
+								SNew(SImage)
+								.Image(this, &SContentBrowser::GetSourcesToggleImage)
+							]
+						]
+
+						+SHorizontalBox::Slot()
+						[
+							SourcesSearch->GetWidget()
+						]
+
+						+SHorizontalBox::Slot()
+						.AutoWidth()
+						.Padding(2, 0, 0, 0)
+						[
+							SNew(SButton)
+							.VAlign(EVerticalAlignment::VAlign_Center)
+							.ButtonStyle(FEditorStyle::Get(), "ToggleButton")
+							.ToolTipText(this, &SContentBrowser::GetSourcesSwitcherToolTipText)
+							.ContentPadding(FMargin(1, 0))
+							.ForegroundColor(FEditorStyle::GetSlateColor(DefaultForegroundName))
+							.OnClicked(this, &SContentBrowser::OnSourcesSwitcherClicked)
+							[
+								SNew(SImage)
+								.Image(this, &SContentBrowser::GetSourcesSwitcherIcon)
 							]
 						]
 					]
-				]
-				+ SSplitter::Slot()
-				.Value(0.8f)
-				[
-					SAssignNew(PathCollectionSplitterPtr, SSplitter)
-					.Style( FEditorStyle::Get(), "ContentBrowser.Splitter" )
-					.Orientation( Orient_Vertical )
-					// Path View
-					+ SSplitter::Slot()
-					.Value(0.9f)
+
+					+SVerticalBox::Slot()
+					.FillHeight(1.0f)
 					[
-						SNew(SBorder)
-						.Visibility( ( Config != nullptr ? Config->bShowAssetPathTree : true ) ? EVisibility::Visible : EVisibility::Collapsed )
-						.Padding(FMargin(3))
-						.BorderImage( FEditorStyle::GetBrush("ToolPanel.GroupBorder") )
+						// Note: If adding more widgets here, fix ContentBrowserSourcesWidgetSwitcherIndex and the code that uses it!
+						SAssignNew(SourcesWidgetSwitcher, SWidgetSwitcher)
+
+						// Paths View
+						+SWidgetSwitcher::Slot()
 						[
-							SAssignNew( PathViewPtr, SPathView )
-							.OnPathSelected( this, &SContentBrowser::PathSelected )
-							.OnGetFolderContextMenu( this, &SContentBrowser::GetFolderContextMenu, true )
-							.OnGetPathContextMenuExtender( this, &SContentBrowser::GetPathContextMenuExtender )
-							.FocusSearchBoxWhenOpened( false )
-							.SearchBarVisibility(this, &SContentBrowser::GetAlternateSearchBarVisibility)
-							.ShowTreeTitle( false )
-							.ShowSeparator( false )
-							.AllowClassesFolder( true )
-							.AddMetaData<FTagMetaData>(FTagMetaData(TEXT("ContentBrowserSources")))
-							.SearchContent()
+							SAssignNew(PathFavoriteSplitterPtr, SSplitter)
+							.Orientation(EOrientation::Orient_Vertical)
+							.MinimumSlotHeight(70.0f)
+							.Visibility( this, &SContentBrowser::GetSourcesViewVisibility )
+							+ SSplitter::Slot()
+							.Value(.2f)
 							[
-								SNew(SVerticalBox)
-								.AddMetaData<FTagMetaData>(FTagMetaData(TEXT("ContentBrowserSourcesToggle1")))
-								.Visibility(this, &SContentBrowser::GetAlternateSearchBarVisibility)
-								+ SVerticalBox::Slot()
-								.FillHeight(1.0f)
-								.Padding(0, 0, 2, 0)
+								SNew(SBox)
+								.Visibility(this, &SContentBrowser::GetFavoriteFolderVisibility)
 								[
-									SNew(SButton)
-									.VAlign(EVerticalAlignment::VAlign_Center)
-									.ButtonStyle(FEditorStyle::Get(), "ToggleButton")
-									.ToolTipText(LOCTEXT("SourcesTreeToggleTooltip", "Show or hide the sources panel"))
-									.ContentPadding(FMargin(1, 0))
-									.ForegroundColor(FEditorStyle::GetSlateColor(DefaultForegroundName))
-									.OnClicked(this, &SContentBrowser::SourcesViewExpandClicked)
+									SNew(SExpandableArea)
+									.BorderImage(FEditorStyle::GetBrush("NoBorder"))
+									.HeaderPadding(FMargin(0.0f, 3.0f, 0.0f, 0.0f))
+									.HeaderContent()
+									[
+										SNew(SHorizontalBox)
+
+										+SHorizontalBox::Slot()
+										.AutoWidth()
+										.Padding(0, 0, 2, 0)
+										.VAlign(VAlign_Center)
+										[
+											SNew(SImage) 
+											.Image(FEditorStyle::GetBrush("PropertyWindow.Favorites_Enabled"))
+										]
+
+										+SHorizontalBox::Slot()
+										.AutoWidth()
+										.VAlign(VAlign_Center)
+										[
+											SNew(STextBlock)
+											.Text(LOCTEXT("Favorites", "Favorites"))
+											.Font(FEditorStyle::GetFontStyle("ContentBrowser.SourceTreeRootItemFont"))
+										]
+									]
+									.BodyContent()
 									[
-										SNew(SImage)
-										.Image(this, &SContentBrowser::GetSourcesToggleImage)
+										SNew(SBox)
+										.Padding(FMargin(9, 0, 0, 0))
+										[
+											SAssignNew(FavoritePathViewPtr, SFavoritePathView)
+											.OnPathSelected(this, &SContentBrowser::FavoritePathSelected)
+											.OnGetFolderContextMenu(this, &SContentBrowser::GetFolderContextMenu, true)
+											.OnGetPathContextMenuExtender(this, &SContentBrowser::GetPathContextMenuExtender)
+											.FocusSearchBoxWhenOpened(false)
+											.ShowTreeTitle(false)
+											.ShowSeparator(false)
+											.AllowClassesFolder(true)
+											.AddMetaData<FTagMetaData>(FTagMetaData(TEXT("ContentBrowserFavorites")))
+											.ExternalSearch(SourcesSearch)
+										]
 									]
 								]
 							]
+							+ SSplitter::Slot()
+							.Value(0.8f)
+							[
+								SNew(SBox)
+								.Padding(FMargin(0.0f, 1.0f, 0.0f, 0.0f))
+								[
+									SAssignNew( PathViewPtr, SPathView )
+									.OnPathSelected( this, &SContentBrowser::PathSelected )
+									.OnGetFolderContextMenu( this, &SContentBrowser::GetFolderContextMenu, true )
+									.OnGetPathContextMenuExtender( this, &SContentBrowser::GetPathContextMenuExtender )
+									.FocusSearchBoxWhenOpened( false )
+									.ShowTreeTitle( false )
+									.ShowSeparator( false )
+									.AllowClassesFolder( true )
+									.AddMetaData<FTagMetaData>(FTagMetaData(TEXT("ContentBrowserSources")))
+									.ExternalSearch(SourcesSearch)
+								]
+							]
 						]
-					]
 
-					// Collection View
-					+ SSplitter::Slot()
-					.Value(0.9f)
-					[
-						SNew(SBorder)
-						.Visibility( this, &SContentBrowser::GetCollectionViewVisibility )
-						.Padding(FMargin(3))
-						.BorderImage( FEditorStyle::GetBrush("ToolPanel.GroupBorder") )
+						// Collections View
+						+SWidgetSwitcher::Slot()
 						[
-							SAssignNew(CollectionViewPtr, SCollectionView)
-							.OnCollectionSelected(this, &SContentBrowser::CollectionSelected)
-							.AddMetaData<FTagMetaData>(FTagMetaData(TEXT("ContentBrowserCollections")))
-							.AllowCollectionDrag(true)
-							.AllowQuickAssetManagement(true)
+							SNew(SBox)
+							.Padding(FMargin(0.0f, 4.0f, 0.0f, 0.0f))
+							[
+								SAssignNew(CollectionViewPtr, SCollectionView)
+								.OnCollectionSelected(this, &SContentBrowser::CollectionSelected)
+								.AddMetaData<FTagMetaData>(FTagMetaData(TEXT("ContentBrowserCollections")))
+								.AllowCollapsing(false)
+								.AllowCollectionDrag(true)
+								.AllowQuickAssetManagement(true)
+								.ExternalSearch(SourcesSearch)
+							]
 						]
 					]
 				]
@@ -639,7 +695,8 @@ void SContentBrowser::Construct( const FArguments& InArgs, const FName& InInstan
 							.HintText( this, &SContentBrowser::GetSearchAssetsHintText )
 							.OnTextChanged( this, &SContentBrowser::OnSearchBoxChanged )
 							.OnTextCommitted( this, &SContentBrowser::OnSearchBoxCommitted )
-							.PossibleSuggestions( this, &SContentBrowser::GetAssetSearchSuggestions )
+							.OnAssetSearchBoxSuggestionFilter( this, &SContentBrowser::OnAssetSearchSuggestionFilter )
+							.OnAssetSearchBoxSuggestionChosen( this, &SContentBrowser::OnAssetSearchSuggestionChosen )
 							.DelayChangeNotificationsWhileTyping( true )
 							.Visibility( ( Config != nullptr ? Config->bCanShowAssetSearch : true ) ? EVisibility::Visible : EVisibility::Collapsed )
 							.AddMetaData<FTagMetaData>(FTagMetaData(TEXT("ContentBrowserSearchAssets")))
@@ -706,7 +763,6 @@ void SContentBrowser::Construct( const FArguments& InArgs, const FName& InInstan
 						.CanShowClasses(Config != nullptr ? Config->bCanShowClasses : true)
 						.CanShowRealTimeThumbnails( Config != nullptr ? Config->bCanShowRealTimeThumbnails : true)
 						.CanShowDevelopersFolder( Config != nullptr ? Config->bCanShowDevelopersFolder : true)
-						.CanShowCollections(true)
 						.CanShowFavorites(true)
 						.AddMetaData<FTagMetaData>(FTagMetaData(TEXT("ContentBrowserAssets")))
 						.OnSearchOptionsChanged(this, &SContentBrowser::HandleAssetViewSearchOptionsChanged)
@@ -743,10 +799,6 @@ void SContentBrowser::Construct( const FArguments& InArgs, const FName& InInstan
 		FavoritePathViewPtr->SetSelectedPaths(SelectedFavoritePaths);
 	}
 
-	//Bind the path view filtering to the favorite path view search bar
-	FavoritePathViewPtr->OnFavoriteSearchChanged.BindSP(PathViewPtr.Get(), &SPathView::OnAssetTreeSearchBoxChanged);
-	FavoritePathViewPtr->OnFavoriteSearchCommitted.BindSP(PathViewPtr.Get(), &SPathView::OnAssetTreeSearchBoxCommitted);
-
 	// Bind the favorites menu to update after folder changes in the path or asset view
 	PathViewPtr->OnFolderPathChanged.BindSP(FavoritePathViewPtr.Get(), &SFavoritePathView::FixupFavoritesFromExternalChange);
 	AssetViewPtr->OnFolderPathChanged.BindSP(FavoritePathViewPtr.Get(), &SFavoritePathView::FixupFavoritesFromExternalChange);
@@ -849,21 +901,11 @@ void SContentBrowser::BindCommands()
 	}
 }
 
-EVisibility SContentBrowser::GetCollectionViewVisibility() const
-{
-	return bAlwaysShowCollections ? EVisibility::Visible : ( GetDefault<UContentBrowserSettings>()->GetDisplayCollections() ? EVisibility::Visible : EVisibility::Collapsed );
-}
-
 EVisibility SContentBrowser::GetFavoriteFolderVisibility() const
 {
 	return GetDefault<UContentBrowserSettings>()->GetDisplayFavorites() ? EVisibility::Visible : EVisibility::Collapsed;
 }
 
-EVisibility SContentBrowser::GetAlternateSearchBarVisibility() const
-{
-	return GetDefault<UContentBrowserSettings>()->GetDisplayFavorites() ? EVisibility::Collapsed : EVisibility::Visible;
-}
-
 void SContentBrowser::ToggleFolderFavorite(const TArray<FString>& FolderPaths)
 {
 	bool bAddedFavorite = false;
@@ -1177,12 +1219,6 @@ void SContentBrowser::SaveSettings() const
 		float SplitterSize = PathAssetSplitterPtr->SlotAt(SlotIndex).SizeValue.Get();
 		GConfig->SetFloat(*SettingsIniSection, *(SettingsString + FString::Printf(TEXT(".VerticalSplitter.SlotSize%d"), SlotIndex)), SplitterSize, GEditorPerProjectIni);
 	}
-	
-	for(int32 SlotIndex = 0; SlotIndex < PathCollectionSplitterPtr->GetChildren()->Num(); SlotIndex++)
-	{
-		float SplitterSize = PathCollectionSplitterPtr->SlotAt(SlotIndex).SizeValue.Get();
-		GConfig->SetFloat(*SettingsIniSection, *(SettingsString + FString::Printf(TEXT(".HorizontalSplitter.SlotSize%d"), SlotIndex)), SplitterSize, GEditorPerProjectIni);
-	}
 
 	for (int32 SlotIndex = 0; SlotIndex < PathFavoriteSplitterPtr->GetChildren()->Num(); SlotIndex++)
 	{
@@ -1316,13 +1352,6 @@ void SContentBrowser::LoadSettings(const FName& InInstanceName)
 		GConfig->GetFloat(*SettingsIniSection, *(SettingsString + FString::Printf(TEXT(".VerticalSplitter.SlotSize%d"), SlotIndex)), SplitterSize, GEditorPerProjectIni);
 		PathAssetSplitterPtr->SlotAt(SlotIndex).SizeValue = SplitterSize;
 	}
-	
-	for(int32 SlotIndex = 0; SlotIndex < PathCollectionSplitterPtr->GetChildren()->Num(); SlotIndex++)
-	{
-		float SplitterSize = PathCollectionSplitterPtr->SlotAt(SlotIndex).SizeValue.Get();
-		GConfig->GetFloat(*SettingsIniSection, *(SettingsString + FString::Printf(TEXT(".HorizontalSplitter.SlotSize%d"), SlotIndex)), SplitterSize, GEditorPerProjectIni);
-		PathCollectionSplitterPtr->SlotAt(SlotIndex).SizeValue = SplitterSize;
-	}
 
 	for (int32 SlotIndex = 0; SlotIndex < PathFavoriteSplitterPtr->GetChildren()->Num(); SlotIndex++)
 	{
@@ -1487,17 +1516,13 @@ void SContentBrowser::CollectionSelected(const FCollectionNameType& SelectedColl
 	TArray<FCollectionNameType> SelectedCollections = CollectionViewPtr->GetSelectedCollections();
 	TArray<FString> SelectedPaths;
 
-	if( SelectedCollections.Num() == 0  )
-	{
-		// just select the game folder
-		SelectedPaths.Add(TEXT("/Game"));
-		SourcesChanged(SelectedPaths, SelectedCollections);
-	}
-	else
+	if (SelectedCollections.Num() == 0)
 	{
-		SourcesChanged(SelectedPaths, SelectedCollections);
+		// Select a dummy "None" collection to avoid the sources view switching to the paths view
+		SelectedCollections.Add(FCollectionNameType(NAME_None, ECollectionShareType::CST_System));
 	}
-
+	
+	SourcesChanged(SelectedPaths, SelectedCollections);
 }
 
 void SContentBrowser::PathPickerPathSelected(const FString& FolderPath)
@@ -1684,6 +1709,11 @@ FReply SContentBrowser::OnSaveSearchButtonClicked()
 	{
 		SourcesViewExpandClicked();
 	}
+	if (ActiveSourcesWidgetIndex != ContentBrowserSourcesWidgetSwitcherIndex::CollectionsView)
+	{
+		ActiveSourcesWidgetIndex = ContentBrowserSourcesWidgetSwitcherIndex::CollectionsView;
+		SourcesWidgetSwitcher->SetActiveWidgetIndex(ActiveSourcesWidgetIndex);
+	}
 
 	// We want to add any currently selected paths to the final saved query so that you get back roughly the same list of objects as what you're currently seeing
 	FString SelectedPathsQuery;
@@ -2286,6 +2316,60 @@ EVisibility SContentBrowser::GetPathExpanderVisibility() const
 	return bSourcesViewExpanded ? EVisibility::Collapsed : EVisibility::Visible;
 }
 
+const FSlateBrush* SContentBrowser::GetSourcesSwitcherIcon() const
+{
+	switch (ActiveSourcesWidgetIndex)
+	{
+	case ContentBrowserSourcesWidgetSwitcherIndex::PathView:
+		return FEditorStyle::GetBrush("ContentBrowser.Sources.Collections");
+	case ContentBrowserSourcesWidgetSwitcherIndex::CollectionsView:
+		return FEditorStyle::GetBrush("ContentBrowser.Sources.Paths");
+	default:
+		break;
+	}
+	check(false);
+	return nullptr;
+}
+
+FText SContentBrowser::GetSourcesSwitcherToolTipText() const
+{
+	switch (ActiveSourcesWidgetIndex)
+	{
+	case ContentBrowserSourcesWidgetSwitcherIndex::PathView:
+		return LOCTEXT("SwitchToCollectionView_ToolTip", "Switch to the Collections view");
+	case ContentBrowserSourcesWidgetSwitcherIndex::CollectionsView:
+		return LOCTEXT("SwitchToPathView_ToolTip", "Switch to the Paths view");
+	default:
+		break;
+	}
+	check(false);
+	return FText();
+}
+
+FReply SContentBrowser::OnSourcesSwitcherClicked()
+{
+	// This only works because we only have two switcher types
+	ActiveSourcesWidgetIndex = !ActiveSourcesWidgetIndex;
+	SourcesWidgetSwitcher->SetActiveWidgetIndex(ActiveSourcesWidgetIndex);
+
+	return FReply::Handled();
+}
+
+FText SContentBrowser::GetSourcesSearchHintText() const
+{
+	switch (ActiveSourcesWidgetIndex)
+	{
+	case ContentBrowserSourcesWidgetSwitcherIndex::PathView:
+		return LOCTEXT("SearchPathsHint", "Search Paths");
+	case ContentBrowserSourcesWidgetSwitcherIndex::CollectionsView:
+		return LOCTEXT("SearchCollectionsHint", "Search Collections");
+	default:
+		break;
+	}
+	check(false);
+	return FText();
+}
+
 FReply SContentBrowser::BackClicked()
 {
 	HistoryManager.GoBack();
@@ -2567,8 +2651,12 @@ void SContentBrowser::UpdatePath()
 
 	PathBreadcrumbTrail->ClearCrumbs();
 
+	int32 NewSourcesWidgetIndex = ActiveSourcesWidgetIndex;
+
 	if ( SourcesData.HasPackagePaths() )
 	{
+		NewSourcesWidgetIndex = ContentBrowserSourcesWidgetSwitcherIndex::PathView;
+
 		TArray<FString> Crumbs;
 		SourcesData.PackagePaths[0].ToString().ParseIntoArray(Crumbs, TEXT("/"), true);
 
@@ -2585,6 +2673,8 @@ void SContentBrowser::UpdatePath()
 	}
 	else if ( SourcesData.HasCollections() )
 	{
+		NewSourcesWidgetIndex = ContentBrowserSourcesWidgetSwitcherIndex::CollectionsView;
+
 		FCollectionManagerModule& CollectionManagerModule = FCollectionManagerModule::GetModule();
 		TArray<FCollectionNameType> CollectionPathItems;
 
@@ -2613,6 +2703,12 @@ void SContentBrowser::UpdatePath()
 	{
 		PathBreadcrumbTrail->PushCrumb(LOCTEXT("AllAssets", "All Assets"), TEXT(""));
 	}
+
+	if (ActiveSourcesWidgetIndex != NewSourcesWidgetIndex)
+	{
+		ActiveSourcesWidgetIndex = NewSourcesWidgetIndex;
+		SourcesWidgetSwitcher->SetActiveWidgetIndex(ActiveSourcesWidgetIndex);
+	}
 }
 
 void SContentBrowser::OnFilterChanged()
@@ -2812,27 +2908,183 @@ FText SContentBrowser::GetSearchAssetsHintText() const
 	return NSLOCTEXT( "ContentBrowser", "SearchBoxHint", "Search Assets" );
 }
 
-TArray<FString> SContentBrowser::GetAssetSearchSuggestions() const
+void ExtractAssetSearchFilterTerms(const FText& SearchText, FString* OutFilterKey, FString* OutFilterValue, int32* OutSuggestionInsertionIndex)
 {
-	TArray<FString> AllSuggestions;
+	const FString SearchString = SearchText.ToString();
+
+	if (OutFilterKey)
+	{
+		OutFilterKey->Reset();
+	}
+	if (OutFilterValue)
+	{
+		OutFilterValue->Reset();
+	}
+	if (OutSuggestionInsertionIndex)
+	{
+		*OutSuggestionInsertionIndex = SearchString.Len();
+	}
 
-	FAssetToolsModule& AssetToolsModule = FModuleManager::LoadModuleChecked<FAssetToolsModule>("AssetTools");
-	TArray< TWeakPtr<IAssetTypeActions> > AssetTypeActionsList;
-	AssetToolsModule.Get().GetAssetTypeActionsList(AssetTypeActionsList);
+	// Build the search filter terms so that we can inspect the tokens
+	FTextFilterExpressionEvaluator LocalFilter(ETextFilterExpressionEvaluatorMode::Complex);
+	LocalFilter.SetFilterText(SearchText);
 
-	for ( auto TypeActionsIt = AssetTypeActionsList.CreateConstIterator(); TypeActionsIt; ++TypeActionsIt )
+	// Inspect the tokens to see what the last part of the search term was
+	// If it was a key->value pair then we'll use that to control what kinds of results we show
+	// For anything else we just use the text from the last token as our filter term to allow incremental auto-complete
+	const TArray<FExpressionToken>& FilterTokens = LocalFilter.GetFilterExpressionTokens();
+	if (FilterTokens.Num() > 0)
 	{
-		if ( (*TypeActionsIt).IsValid() )
+		const FExpressionToken& LastToken = FilterTokens.Last();
+
+		// If the last token is a text token, then consider it as a value and walk back to see if we also have a key
+		if (LastToken.Node.Cast<TextFilterExpressionParser::FTextToken>())
+		{
+			if (OutFilterValue)
+			{
+				*OutFilterValue = LastToken.Context.GetString();
+			}
+			if (OutSuggestionInsertionIndex)
+			{
+				*OutSuggestionInsertionIndex = FMath::Min(*OutSuggestionInsertionIndex, LastToken.Context.GetCharacterIndex());
+			}
+
+			if (FilterTokens.IsValidIndex(FilterTokens.Num() - 2))
+			{
+				const FExpressionToken& ComparisonToken = FilterTokens[FilterTokens.Num() - 2];
+				if (ComparisonToken.Node.Cast<TextFilterExpressionParser::FEqual>())
+				{
+					if (FilterTokens.IsValidIndex(FilterTokens.Num() - 3))
+					{
+						const FExpressionToken& KeyToken = FilterTokens[FilterTokens.Num() - 3];
+						if (KeyToken.Node.Cast<TextFilterExpressionParser::FTextToken>())
+						{
+							if (OutFilterKey)
+							{
+								*OutFilterKey = KeyToken.Context.GetString();
+							}
+							if (OutSuggestionInsertionIndex)
+							{
+								*OutSuggestionInsertionIndex = FMath::Min(*OutSuggestionInsertionIndex, KeyToken.Context.GetCharacterIndex());
+							}
+						}
+					}
+				}
+			}
+		}
+
+		// If the last token is a comparison operator, then walk back and see if we have a key
+		else if (LastToken.Node.Cast<TextFilterExpressionParser::FEqual>())
 		{
-			const TSharedPtr<IAssetTypeActions> TypeActions = (*TypeActionsIt).Pin();
-			if ( TypeActions->GetSupportedClass() )
+			if (FilterTokens.IsValidIndex(FilterTokens.Num() - 2))
 			{
-				AllSuggestions.Add( TypeActions->GetSupportedClass()->GetName() );
+				const FExpressionToken& KeyToken = FilterTokens[FilterTokens.Num() - 2];
+				if (KeyToken.Node.Cast<TextFilterExpressionParser::FTextToken>())
+				{
+					if (OutFilterKey)
+					{
+						*OutFilterKey = KeyToken.Context.GetString();
+					}
+					if (OutSuggestionInsertionIndex)
+					{
+						*OutSuggestionInsertionIndex = FMath::Min(*OutSuggestionInsertionIndex, KeyToken.Context.GetCharacterIndex());
+					}
+				}
 			}
 		}
 	}
+}
+
+void SContentBrowser::OnAssetSearchSuggestionFilter(const FText& SearchText, TArray<FAssetSearchBoxSuggestion>& PossibleSuggestions, FText& SuggestionHighlightText) const
+{
+	// We don't bind the suggestion list, so this list should be empty as we populate it here based on the search term
+	check(PossibleSuggestions.Num() == 0);
+
+	FString FilterKey;
+	FString FilterValue;
+	ExtractAssetSearchFilterTerms(SearchText, &FilterKey, &FilterValue, nullptr);
+
+	auto PassesValueFilter = [&FilterValue](const FString& InOther)
+	{
+		return FilterValue.IsEmpty() || InOther.Contains(FilterValue);
+	};
+
+	if (FilterKey.IsEmpty() || (FilterKey == TEXT("Type") || FilterKey == TEXT("Class")))
+	{
+		FAssetToolsModule& AssetToolsModule = FModuleManager::LoadModuleChecked<FAssetToolsModule>("AssetTools");
+		TArray< TWeakPtr<IAssetTypeActions> > AssetTypeActionsList;
+		AssetToolsModule.Get().GetAssetTypeActionsList(AssetTypeActionsList);
+
+		const FText TypesCategoryName = NSLOCTEXT("ContentBrowser", "TypesCategoryName", "Types");
+		for (auto TypeActionsIt = AssetTypeActionsList.CreateConstIterator(); TypeActionsIt; ++TypeActionsIt)
+		{
+			if ((*TypeActionsIt).IsValid())
+			{
+				const TSharedPtr<IAssetTypeActions> TypeActions = (*TypeActionsIt).Pin();
+				if (TypeActions->GetSupportedClass())
+				{
+					const FString TypeName = TypeActions->GetSupportedClass()->GetName();
+					const FText TypeDisplayName = TypeActions->GetSupportedClass()->GetDisplayNameText();
+					FString TypeSuggestion = FString::Printf(TEXT("Type=%s"), *TypeName);
+					if (PassesValueFilter(TypeSuggestion))
+					{
+						PossibleSuggestions.Add(FAssetSearchBoxSuggestion{ MoveTemp(TypeSuggestion), TypeDisplayName, TypesCategoryName });
+					}
+				}
+			}
+		}
+	}
+
+	if (FilterKey.IsEmpty() || (FilterKey == TEXT("Collection") || FilterKey == TEXT("Tag")))
+	{
+		ICollectionManager& CollectionManager = FCollectionManagerModule::GetModule().Get();
+
+		TArray<FCollectionNameType> AllCollections;
+		CollectionManager.GetCollections(AllCollections);
+
+		const FText CollectionsCategoryName = NSLOCTEXT("ContentBrowser", "CollectionsCategoryName", "Collections");
+		for (const FCollectionNameType& Collection : AllCollections)
+		{
+			FString CollectionName = Collection.Name.ToString();
+			FString CollectionSuggestion = FString::Printf(TEXT("Collection=%s"), *CollectionName);
+			if (PassesValueFilter(CollectionSuggestion))
+			{
+				PossibleSuggestions.Add(FAssetSearchBoxSuggestion{ MoveTemp(CollectionSuggestion), FText::FromString(MoveTemp(CollectionName)), CollectionsCategoryName });
+			}
+		}
+	}
+
+	if (FilterKey.IsEmpty())
+	{
+		IAssetRegistry& AssetRegistry = FModuleManager::LoadModuleChecked<FAssetRegistryModule>(AssetRegistryConstants::ModuleName).Get();
+
+		if (const FAssetRegistryState* StatePtr = AssetRegistry.GetAssetRegistryState())
+		{
+			const FText MetaDataCategoryName = NSLOCTEXT("ContentBrowser", "MetaDataCategoryName", "Meta-Data");
+			for (const auto& TagAndArrayPair : StatePtr->GetTagToAssetDatasMap())
+			{
+				const FString TagNameStr = TagAndArrayPair.Key.ToString();
+				if (PassesValueFilter(TagNameStr))
+				{
+					PossibleSuggestions.Add(FAssetSearchBoxSuggestion{ TagNameStr, FText::FromString(TagNameStr), MetaDataCategoryName });
+				}
+			}
+		}
+	}
+
+	SuggestionHighlightText = FText::FromString(FilterValue);
+}
+
+FText SContentBrowser::OnAssetSearchSuggestionChosen(const FText& SearchText, const FString& Suggestion) const
+{
+	int32 SuggestionInsertionIndex = 0;
+	ExtractAssetSearchFilterTerms(SearchText, nullptr, nullptr, &SuggestionInsertionIndex);
+
+	FString SearchString = SearchText.ToString();
+	SearchString.RemoveAt(SuggestionInsertionIndex, SearchString.Len() - SuggestionInsertionIndex, false);
+	SearchString.Append(Suggestion);
 
-	return AllSuggestions;
+	return FText::FromString(SearchString);
 }
 
 TSharedPtr<SWidget> SContentBrowser::GetFolderContextMenu(const TArray<FString>& SelectedPaths, FContentBrowserMenuExtender_SelectedPaths InMenuExtender, FOnCreateNewFolder InOnCreateNewFolder, bool bPathView)
diff --git a/Engine/Source/Editor/ContentBrowser/Private/SContentBrowser.h b/Engine/Source/Editor/ContentBrowser/Private/SContentBrowser.h
--- a/Engine/Source/Editor/ContentBrowser/Private/SContentBrowser.h
+++ b/Engine/Source/Editor/ContentBrowser/Private/SContentBrowser.h
@@ -9,6 +9,7 @@
 #include "Widgets/SWidget.h"
 #include "Widgets/SCompoundWidget.h"
 #include "AssetData.h"
+#include "SAssetSearchBox.h"
 #include "Framework/MultiBox/MultiBoxExtender.h"
 #include "CollectionManagerTypes.h"
 #include "IContentBrowserSingleton.h"
@@ -16,10 +17,10 @@
 
 class FAssetContextMenu;
 class FFrontendFilter_Text;
+class FSourcesSearch;
 class FPathContextMenu;
 class FTabManager;
 class FUICommandList;
-class SAssetSearchBox;
 class SAssetView;
 class SCollectionView;
 class SComboButton;
@@ -291,6 +292,18 @@ private:
 	/** Gets the visibility of the path expander button */
 	EVisibility GetPathExpanderVisibility() const;
 
+	/** Gets the icon used on the source switch button */
+	const FSlateBrush* GetSourcesSwitcherIcon() const;
+
+	/** Gets the tooltip text used on the source switch button */
+	FText GetSourcesSwitcherToolTipText() const;
+
+	/** Handler for clicking the source switch button */
+	FReply OnSourcesSwitcherClicked();
+
+	/** Gets the source search hint text */
+	FText GetSourcesSearchHintText() const;
+
 	/** Handler for clicking the history back button */
 	FReply BackClicked();
 
@@ -403,7 +416,10 @@ private:
 	void HandlePathRemoved(const FString& Path);
 
 	/** Gets all suggestions for the asset search box */
-	TArray<FString> GetAssetSearchSuggestions() const;
+	void OnAssetSearchSuggestionFilter(const FText& SearchText, TArray<FAssetSearchBoxSuggestion>& PossibleSuggestions, FText& SuggestionHighlightText) const;
+
+	/** Combines the chosen suggestion with the active search text */
+	FText OnAssetSearchSuggestionChosen(const FText& SearchText, const FString& Suggestion) const;
 
 	/** Gets the dynamic hint text for the "Search Assets" search text box */
 	FText GetSearchAssetsHintText() const;
@@ -423,14 +439,8 @@ private:
 	/** Bind our UI commands */
 	void BindCommands();
 
-	/** Gets the visibility of the collection view */
-	EVisibility GetCollectionViewVisibility() const;
-
 	/** Gets the visibility of the favorites view */
 	EVisibility GetFavoriteFolderVisibility() const;
-	
-	/** The visibility of the search bar for the base path view*/
-	EVisibility GetAlternateSearchBarVisibility() const;
 
 	/** Toggles the favorite status of an array of folders*/
 	void ToggleFolderFavorite(const TArray<FString>& FolderPaths);
@@ -452,6 +462,9 @@ private:
 	/** The context menu manager for the path view */
 	TSharedPtr<class FPathContextMenu> PathContextMenu;
 
+	/** The sources search, shared between the paths and collections views */
+	TSharedPtr<FSourcesSearch> SourcesSearch;
+
 	/** The asset tree widget */
 	TSharedPtr<SPathView> PathViewPtr;
 
@@ -476,6 +489,9 @@ private:
 	/** The path picker */
 	TSharedPtr<SComboButton> PathPickerButton;
 
+	/** Index of the active sources widget */
+	int32 ActiveSourcesWidgetIndex = 0;
+
 	/** The expanded state of the asset tree */
 	bool bSourcesViewExpanded;
 
@@ -503,21 +519,18 @@ private:
 	/** Delegate used to create a new folder */
 	FOnCreateNewFolder OnCreateNewFolder;
 
+	/** Switcher between the different sources views */
+	TSharedPtr<SWidgetSwitcher> SourcesWidgetSwitcher;
+
 	/** The splitter between the path & asset view */
 	TSharedPtr<SSplitter> PathAssetSplitterPtr;
 
-	/** The splitter between the path & collection view */
-	TSharedPtr<SSplitter> PathCollectionSplitterPtr;
-
 	/** The splitter between the path & favorite view */
 	TSharedPtr<SSplitter> PathFavoriteSplitterPtr;
 
 	/** When viewing a dynamic collection, the active search query will be stashed in this variable so that it can be restored again later */
 	TOptional<FText> StashedSearchBoxText;
 
-	/** True if we should always show collections in this Content Browser, ignoring the view settings */
-	bool bAlwaysShowCollections;
-
 public: 
 
 	/** The section of EditorPerProjectUserSettings in which to save content browser settings */
diff --git a/Engine/Source/Editor/ContentBrowser/Private/SPathView.cpp b/Engine/Source/Editor/ContentBrowser/Private/SPathView.cpp
--- a/Engine/Source/Editor/ContentBrowser/Private/SPathView.cpp
+++ b/Engine/Source/Editor/ContentBrowser/Private/SPathView.cpp
@@ -19,6 +19,7 @@
 #include "DragDropHandler.h"
 
 #include "PathViewTypes.h"
+#include "SourcesSearch.h"
 #include "SourcesViewWidgets.h"
 #include "Widgets/Input/SSearchBox.h"
 #include "NativeClassHierarchy.h"
@@ -112,32 +113,50 @@ void SPathView::Construct( const FArguments& InArgs )
 			.HighlightParentNodesForSelection(true);
 	}
 
-	ChildSlot
-	[
-		SNew(SVerticalBox)
+	SearchPtr = InArgs._ExternalSearch;
+	if (!SearchPtr)
+	{
+		SearchPtr = MakeShared<FSourcesSearch>();
+		SearchPtr->Initialize();
+		SearchPtr->SetHintText(LOCTEXT("AssetTreeSearchBoxHint", "Search Folders"));
+	}
+	SearchPtr->OnSearchChanged().AddSP(this, &SPathView::SetSearchFilterText);
 
-		// Search
-		+ SVerticalBox::Slot()
-		.AutoHeight()
-		.Padding(0, 1, 0, 3)
-		[
+	TSharedRef<SBox> SearchBox = SNew(SBox);
+	if (!InArgs._ExternalSearch)
+	{
+		SearchBox->SetPadding(FMargin(0, 1, 0, 3));
+
+		SearchBox->SetContent(
 			SNew(SHorizontalBox)
 
-			+ SHorizontalBox::Slot()
+			+SHorizontalBox::Slot()
 			.AutoWidth()
 			[
 				InArgs._SearchContent.Widget
 			]
 
-			+ SHorizontalBox::Slot()
+			+SHorizontalBox::Slot()
 			.FillWidth(1.0f)
 			[
-				SAssignNew(SearchBoxPtr, SSearchBox)
+				SNew(SBox)
 				.Visibility(InArgs._SearchBarVisibility)
-				.HintText( LOCTEXT( "AssetTreeSearchBoxHint", "Search Folders" ) )
-				.OnTextChanged( this, &SPathView::OnAssetTreeSearchBoxChanged )
-				.OnTextCommitted( this, &SPathView::OnAssetTreeSearchBoxCommitted )
+				[
+					SearchPtr->GetWidget()
+				]
 			]
+		);
+	}
+
+	ChildSlot
+	[
+		SNew(SVerticalBox)
+
+		// Search
+		+ SVerticalBox::Slot()
+		.AutoHeight()
+		[
+			SearchBox
 		]
 
 		// Tree title
@@ -194,11 +213,8 @@ void SPathView::SetSelectedPaths(const TArray<FString>& Paths)
 		return;
 	}
 
-	if ( !SearchBoxPtr->GetText().IsEmpty() )
-	{
-		// Clear the search box so the selected paths will be visible
-		SearchBoxPtr->SetText( FText::GetEmpty() );
-	}
+	// Clear the search box so the selected paths will be visible
+	SearchPtr->ClearSearch();
 
 	// Prevent the selection changed delegate since the invoking code requested it
 	FScopedPreventTreeItemChangedDelegate DelegatePrevention( SharedThis(this) );
@@ -506,7 +522,7 @@ void SPathView::SyncToInternal( const TArray<FAssetData>& AssetDataList, const T
 	TArray<TSharedPtr<FTreeItem>> SyncTreeItems;
 
 	// Clear the filter
-	SearchBoxPtr->SetText(FText::GetEmpty());
+	SearchPtr->ClearSearch();
 
 	TSet<FString> PackagePaths = TSet<FString>(FolderPaths);
 	for (const FAssetData& AssetData : AssetDataList)
@@ -721,7 +737,7 @@ void SPathView::LoadSettings(const FString& IniFilename, const FString& IniSecti
 EActiveTimerReturnType SPathView::SetFocusPostConstruct( double InCurrentTime, float InDeltaTime )
 {
 	FWidgetPath WidgetToFocusPath;
-	FSlateApplication::Get().GeneratePathToWidgetUnchecked( SearchBoxPtr.ToSharedRef(), WidgetToFocusPath );
+	FSlateApplication::Get().GeneratePathToWidgetUnchecked( SearchPtr->GetWidget(), WidgetToFocusPath );
 	FSlateApplication::Get().SetKeyboardFocus( WidgetToFocusPath, EFocusCause::SetDirectly );
 
 	return EActiveTimerReturnType::Stop;
@@ -992,28 +1008,22 @@ void SPathView::TreeExpansionChanged( TSharedPtr< FTreeItem > TreeItem, bool bIs
 	}
 }
 
-void SPathView::OnAssetTreeSearchBoxChanged( const FText& InSearchText )
+void SPathView::FilterUpdated()
 {
-	SearchBoxFolderFilter->SetRawFilterText( InSearchText );
-	SearchBoxPtr->SetError( SearchBoxFolderFilter->GetFilterErrorText() );
+	Populate();
 }
 
-void SPathView::OnAssetTreeSearchBoxCommitted(const FText& InSearchText, ETextCommit::Type InCommitType)
+void SPathView::SetSearchFilterText(const FText& InSearchText, TArray<FText>& OutErrors)
 {
-	if (InCommitType == ETextCommit::OnCleared)
+	SearchBoxFolderFilter->SetRawFilterText(InSearchText);
+
+	const FText ErrorText = SearchBoxFolderFilter->GetFilterErrorText();
+	if (!ErrorText.IsEmpty())
 	{
-		// Clear the search box and the filters
-		SearchBoxPtr->SetText(FText::GetEmpty());
-		OnAssetTreeSearchBoxChanged(FText::GetEmpty());
-		FSlateApplication::Get().ClearKeyboardFocus(EFocusCause::Cleared);
+		OutErrors.Add(ErrorText);
 	}
 }
 
-void SPathView::FilterUpdated()
-{
-	Populate();
-}
-
 FText SPathView::GetHighlightText() const
 {
 	return SearchBoxFolderFilter->GetRawFilterText();
@@ -1714,18 +1724,6 @@ void SFavoritePathView::LoadSettings(const FString& IniFilename, const FString&
 	}
 }
 
-void SFavoritePathView::OnAssetTreeSearchBoxChanged(const FText& InSearchText)
-{
-	SPathView::OnAssetTreeSearchBoxChanged(InSearchText);
-	OnFavoriteSearchChanged.ExecuteIfBound(InSearchText);
-}
-
-void SFavoritePathView::OnAssetTreeSearchBoxCommitted(const FText& InSearchText, ETextCommit::Type InCommitType)
-{
-	SPathView::OnAssetTreeSearchBoxCommitted(InSearchText, InCommitType);
-	OnFavoriteSearchCommitted.ExecuteIfBound(InSearchText, InCommitType);
-}
-
 void SFavoritePathView::SetSelectedPaths(const TArray<FString>& Paths)
 {
 	if (!ensure(TreeViewPtr.IsValid()))
@@ -1733,11 +1731,8 @@ void SFavoritePathView::SetSelectedPaths(const TArray<FString>& Paths)
 		return;
 	}
 
-	if (!SearchBoxPtr->GetText().IsEmpty())
-	{
-		// Clear the search box so the selected paths will be visible
-		SearchBoxPtr->SetText(FText::GetEmpty());
-	}
+	// Clear the search box so the selected paths will be visible
+	SearchPtr->ClearSearch();
 
 	// Prevent the selection changed delegate since the invoking code requested it
 	FScopedPreventTreeItemChangedDelegate DelegatePrevention(SharedThis(this));
diff --git a/Engine/Source/Editor/ContentBrowser/Private/SPathView.h b/Engine/Source/Editor/ContentBrowser/Private/SPathView.h
--- a/Engine/Source/Editor/ContentBrowser/Private/SPathView.h
+++ b/Engine/Source/Editor/ContentBrowser/Private/SPathView.h
@@ -16,14 +16,12 @@
 #include "ContentBrowserDelegates.h"
 #include "Delegates/DelegateCombinations.h"
 
+class FSourcesSearch;
 struct FHistoryData;
 struct FTreeItem;
 
 typedef TTextFilter< const FString& > FolderTextFilter;
 
-DECLARE_DELEGATE_OneParam(FOnAssetTreeSearchBoxChanged, const FText& /*SearchText*/);
-DECLARE_DELEGATE_TwoParams(FOnAssetTreeSearchBoxCommitted, const FText& /*InSearchText*/, ETextCommit::Type /*InCommitType*/);
-
 /**
  * The tree view of folders which contain content.
  */
@@ -73,6 +71,9 @@ public:
 		/** The selection mode for the tree view */
 		SLATE_ARGUMENT( ESelectionMode::Type, SelectionMode )
 
+		/** Optional external search. Will hide and replace our internal search UI */
+		SLATE_ARGUMENT( TSharedPtr<FSourcesSearch>, ExternalSearch )
+
 	SLATE_END_ARGS()
 
 	/** Destructor */
@@ -158,12 +159,6 @@ public:
 		return TreeTitle;
 	}
 
-	/** Handler for when search terms change in the asset tree search box */
-	virtual void OnAssetTreeSearchBoxChanged(const FText& InSearchText);
-
-	/** Handler for when search term is committed in the asset tree search box */
-	virtual void OnAssetTreeSearchBoxCommitted(const FText& InSearchText, ETextCommit::Type InCommitType);
-
 public:
 	/** Delegate that handles if any folder paths changed as a result of a move, rename, etc. in the path view*/
 	FOnFolderPathChanged OnFolderPathChanged;
@@ -202,6 +197,9 @@ protected:
 	/** Handler used to verify the name of a new folder */
 	bool VerifyFolderNameChanged(const FString& InName, FText& OutErrorMessage, const FString& InFolderPath) const;
 
+	/** Set the active filter text */
+	void SetSearchFilterText(const FText& InSearchText, TArray<FText>& OutErrors);
+
 	/** Gets the string to highlight in tree items (used in folder searching) */
 	FText GetHighlightText() const;
 
@@ -311,8 +309,8 @@ protected:
 	/** The tree view widget */
 	TSharedPtr< STreeView< TSharedPtr<FTreeItem>> > TreeViewPtr;
 
-	/** The asset tree search box */
-	TSharedPtr< SSearchBox > SearchBoxPtr;
+	/** The path view search interface */
+	TSharedPtr<FSourcesSearch> SearchPtr;
 
 	/** The list of folders in the tree */
 	TArray< TSharedPtr<FTreeItem> > TreeRootItems;
@@ -374,19 +372,9 @@ public:
 	/** Loads any settings to config that should be persistent between editor sessions */
 	virtual void LoadSettings(const FString& IniFilename, const FString& IniSection, const FString& SettingsString) override;
 
-	/** Handler for when search terms change in the asset tree search box */
-	virtual void OnAssetTreeSearchBoxChanged(const FText& InSearchText) override;
-
-	/** Handler for when search term is committed in the asset tree search box */
-	virtual void OnAssetTreeSearchBoxCommitted(const FText& InSearchText, ETextCommit::Type InCommitType) override;
-
 	/** Selects the closest matches to the supplied paths in the tree. "/" delimited */
 	virtual void SetSelectedPaths(const TArray<FString>& Paths) override;
 
-	FOnAssetTreeSearchBoxChanged OnFavoriteSearchChanged;
-
-	FOnAssetTreeSearchBoxCommitted OnFavoriteSearchCommitted;
-
 	/** Adds nodes to the tree in order to construct the specified path. If bUserNamed is true, the user will name the folder and Path includes the default name. */
 	virtual TSharedPtr<FTreeItem> AddPath(const FString& Path, bool bUserNamed = false) override;
 
diff --git a/Engine/Source/Editor/ContentBrowser/Private/SourcesSearch.cpp b/Engine/Source/Editor/ContentBrowser/Private/SourcesSearch.cpp
new file mode 100644
--- /dev/null
+++ b/Engine/Source/Editor/ContentBrowser/Private/SourcesSearch.cpp
@@ -0,0 +1,64 @@
+// Copyright 1998-2019 Epic Games, Inc. All Rights Reserved.
+
+#include "SourcesSearch.h"
+#include "Widgets/Input/SSearchBox.h"
+#include "Framework/Application/SlateApplication.h"
+
+void FSourcesSearch::Initialize()
+{
+	SearchBox = SNew(SSearchBox)
+		.OnTextChanged(this, &FSourcesSearch::OnSearchBoxTextChanged)
+		.OnTextCommitted(this, &FSourcesSearch::OnSearchBoxTextCommitted);
+}
+
+void FSourcesSearch::ClearSearch()
+{
+	if (!SearchBox->GetSearchText().IsEmpty())
+	{
+		SearchBox->SetSearchText(FText::GetEmpty());
+	}
+}
+
+void FSourcesSearch::SetHintText(const TAttribute<FText>& InHintText)
+{
+	SearchBox->SetHintText(InHintText);
+}
+
+TSharedRef<SWidget> FSourcesSearch::GetWidget() const
+{
+	return SearchBox.ToSharedRef();
+}
+
+void FSourcesSearch::OnSearchBoxTextChanged(const FText& InSearchText)
+{
+	TArray<FText> SearchErrors;
+	OnSearchChangedDelegate.Broadcast(InSearchText, SearchErrors);
+
+	if (SearchErrors.Num() == 0)
+	{
+		SearchBox->SetError(FText::GetEmpty());
+	}
+	else if (SearchErrors.Num() == 1)
+	{
+		SearchBox->SetError(SearchErrors[0]);
+	}
+	else
+	{
+		FTextBuilder CombinedError;
+		for (const FText& SearchError : SearchErrors)
+		{
+			CombinedError.AppendLine(SearchError);
+		}
+		SearchBox->SetError(CombinedError.ToText());
+	}
+}
+
+void FSourcesSearch::OnSearchBoxTextCommitted(const FText& InSearchText, ETextCommit::Type InCommitType)
+{
+	if (InCommitType == ETextCommit::OnCleared)
+	{
+		// Clear the search box
+		ClearSearch();
+		FSlateApplication::Get().ClearKeyboardFocus(EFocusCause::Cleared);
+	}
+}
diff --git a/Engine/Source/Editor/ContentBrowser/Private/SourcesSearch.h b/Engine/Source/Editor/ContentBrowser/Private/SourcesSearch.h
new file mode 100644
--- /dev/null
+++ b/Engine/Source/Editor/ContentBrowser/Private/SourcesSearch.h
@@ -0,0 +1,47 @@
+// Copyright 1998-2019 Epic Games, Inc. All Rights Reserved.
+
+#pragma once
+
+#include "CoreMinimal.h"
+#include "Types/SlateEnums.h"
+#include "Delegates/DelegateCombinations.h"
+
+class SWidget;
+class SSearchBox;
+
+/** Implementation for a sources search, used by the path and collections views */
+class FSourcesSearch : public TSharedFromThis<FSourcesSearch>
+{
+public:
+	/** Initialize the search widget */
+	void Initialize();
+
+	/** Delegate invoked when the search text is changed */
+	DECLARE_MULTICAST_DELEGATE_TwoParams(FOnSearchChanged, const FText& /*InSearchText*/, TArray<FText>& /*OutErrors*/);
+	FOnSearchChanged& OnSearchChanged()
+	{
+		return OnSearchChangedDelegate;
+	}
+	
+	/** Clear the current search text */
+	void ClearSearch();
+	
+	/** Set the active hint text for the underlying search widget */
+	void SetHintText(const TAttribute<FText>& InHintText);
+
+	/** Get the underlying search widget */
+	TSharedRef<SWidget> GetWidget() const;
+
+private:
+	/** Handler for when search term changes */
+	void OnSearchBoxTextChanged(const FText& InSearchText);
+
+	/** Handler for when search term is committed */
+	void OnSearchBoxTextCommitted(const FText& InSearchText, ETextCommit::Type InCommitType);
+
+	/** Underlying search widget */
+	TSharedPtr<SSearchBox> SearchBox;
+
+	/** Delegate invoked when the search text is changed */
+	FOnSearchChanged OnSearchChangedDelegate;
+};
diff --git a/Engine/Source/Editor/ContentBrowser/Private/SourcesViewWidgets.cpp b/Engine/Source/Editor/ContentBrowser/Private/SourcesViewWidgets.cpp
--- a/Engine/Source/Editor/ContentBrowser/Private/SourcesViewWidgets.cpp
+++ b/Engine/Source/Editor/ContentBrowser/Private/SourcesViewWidgets.cpp
@@ -12,7 +12,6 @@
 #include "DragDropHandler.h"
 #include "ContentBrowserUtils.h"
 #include "CollectionViewUtils.h"
-#include "Widgets/Text/SInlineEditableTextBlock.h"
 
 #define LOCTEXT_NAMESPACE "ContentBrowser"
 
@@ -319,96 +318,30 @@ void SCollectionTreeItem::Construct( const FArguments& InArgs )
 	OnHandleDragDrop = InArgs._OnHandleDragDrop;
 	bDraggedOver = false;
 
+	TSharedPtr<SAssetTagItem> AssetTagItem;
+
 	ChildSlot
 	[
-		SNew(SBorder)
-		.BorderImage(this, &SCollectionTreeItem::GetBorderImage)
-		.Padding(0)
-		[
-			SNew(SHorizontalBox)
-
-			+SHorizontalBox::Slot()
-			.AutoWidth()
-			.VAlign(VAlign_Center)
-			.Padding(0, 0, 2, 0)
-			[
-				SNew(SCheckBox)
-				.Visibility( InArgs._IsCollectionChecked.IsSet() ? EVisibility::Visible : EVisibility::Collapsed )
-				.IsEnabled( InArgs._IsCheckBoxEnabled )
-				.IsChecked( InArgs._IsCollectionChecked )
-				.OnCheckStateChanged( InArgs._OnCollectionCheckStateChanged )
-			]
-
-			+SHorizontalBox::Slot()
-			.AutoWidth()
-			.VAlign(VAlign_Center)
-			.Padding(0, 0, 2, 0)
-			[
-				// Share Type Icon
-				SNew(SImage)
-				.Image( FEditorStyle::GetBrush( ECollectionShareType::GetIconStyleName(InArgs._CollectionItem->CollectionType) ) )
-				.ColorAndOpacity( this, &SCollectionTreeItem::GetCollectionColor )
-				.ToolTipText(ECollectionShareType::GetDescription(InArgs._CollectionItem->CollectionType))
-			]
-
-			+SHorizontalBox::Slot()
-			[
-				SAssignNew(InlineRenameWidget, SInlineEditableTextBlock)
-				.Text( this, &SCollectionTreeItem::GetNameText )
-				.HighlightText( InArgs._HighlightText )
-				.Font( FEditorStyle::GetFontStyle("ContentBrowser.SourceListItemFont") )
-				.OnBeginTextEdit(this, &SCollectionTreeItem::HandleBeginNameChange)
-				.OnTextCommitted(this, &SCollectionTreeItem::HandleNameCommitted)
-				.OnVerifyTextChanged(this, &SCollectionTreeItem::HandleVerifyNameChanged)
-				.IsSelected( InArgs._IsSelected )
-				.IsReadOnly( InArgs._IsReadOnly )
-			]
-
-			+SHorizontalBox::Slot()
-			.AutoWidth()
-			.Padding(2, 0, 3, 0)
-			[
-				// Storage Mode Icon
-				SNew(SBox)
-				.WidthOverride(16)
-				.HeightOverride(16)
-				.VAlign(VAlign_Center)
-				.HAlign(HAlign_Center)
-				[
-					SNew(STextBlock)
-					.Font(FEditorStyle::Get().GetFontStyle("FontAwesome.10"))
-					.Text(this, &SCollectionTreeItem::GetCollectionStorageModeIconText)
-					.ColorAndOpacity(FLinearColor::Gray)
-					.ToolTipText(this, &SCollectionTreeItem::GetCollectionStorageModeToolTipText)
-				]
-			]
-
-			+SHorizontalBox::Slot()
-			.AutoWidth()
-			.VAlign(VAlign_Center)
-			.Padding(2, 0, 2, 0)
-			[
-				SNew(SImage)
-				.Image(FEditorStyle::GetBrush("ContentBrowser.CollectionStatus"))
-				.ColorAndOpacity(this, &SCollectionTreeItem::GetCollectionStatusColor)
-				.ToolTipText(this, &SCollectionTreeItem::GetCollectionStatusToolTipText)
-			]
-		]
+		SAssignNew(AssetTagItem, SAssetTagItem)
+		.ViewMode(InArgs._ViewMode)
+		.BaseColor(this, &SCollectionTreeItem::GetCollectionColor)
+		.DisplayName(this, &SCollectionTreeItem::GetNameText)
+		.CountText(this, &SCollectionTreeItem::GetCollectionObjectCountText)
+		.WarningText(this, &SCollectionTreeItem::GetCollectionWarningText)
+		.HighlightText(InArgs._HighlightText)
+		.OnBeginNameEdit(this, &SCollectionTreeItem::HandleBeginNameChange)
+		.OnNameCommitted(this, &SCollectionTreeItem::HandleNameCommitted)
+		.OnVerifyName(this, &SCollectionTreeItem::HandleVerifyNameChanged)
+		.IsSelected(InArgs._IsSelected)
+		.IsNameReadOnly(InArgs._IsReadOnly)
+		.IsCheckBoxEnabled(InArgs._IsCheckBoxEnabled)
+		.IsChecked(InArgs._IsCollectionChecked)
+		.OnCheckStateChanged(InArgs._OnCollectionCheckStateChanged)
+		.OnBuildToolTipInfo(this, &SCollectionTreeItem::BuildToolTipInfo)
 	];
 
-	if(InlineRenameWidget.IsValid())
-	{
-		// This is broadcast when the context menu / input binding requests a rename
-		EnterEditingModeDelegateHandle = CollectionItem.Pin()->OnRenamedRequestEvent.AddSP(InlineRenameWidget.Get(), &SInlineEditableTextBlock::EnterEditingMode);
-	}
-}
-
-SCollectionTreeItem::~SCollectionTreeItem()
-{
-	if(InlineRenameWidget.IsValid())
-	{
-		CollectionItem.Pin()->OnRenamedRequestEvent.Remove( EnterEditingModeDelegateHandle );
-	}
+	// This is broadcast when the context menu / input binding requests a rename
+	InArgs._CollectionItem->OnRenamedRequestEvent.AddSP(AssetTagItem.Get(), &SAssetTagItem::RequestRename);
 }
 
 void SCollectionTreeItem::Tick( const FGeometry& AllottedGeometry, const double InCurrentTime, const float InDeltaTime )
@@ -540,129 +473,81 @@ FText SCollectionTreeItem::GetNameText() const
 	}
 }
 
-FSlateColor SCollectionTreeItem::GetCollectionColor() const
+FLinearColor SCollectionTreeItem::GetCollectionColor() const
 {
 	TSharedPtr<FCollectionItem> CollectionItemPtr = CollectionItem.Pin();
 
 	if ( CollectionItemPtr.IsValid() )
 	{
-		const TSharedPtr<FLinearColor> Color = CollectionViewUtils::LoadColor(CollectionItemPtr->CollectionName.ToString(), CollectionItemPtr->CollectionType);
-		if( Color.IsValid() )
-		{
-			return *Color.Get();
-		}
+		return CollectionItemPtr->CollectionColor;
 	}
 	
 	return CollectionViewUtils::GetDefaultColor();
 }
 
-const FSlateBrush* SCollectionTreeItem::GetBorderImage() const
-{
-	return bDraggedOver ? FEditorStyle::GetBrush("Menu.Background") : FEditorStyle::GetBrush("NoBorder");
-}
-
-FText SCollectionTreeItem::GetCollectionStorageModeIconText() const
+FText SCollectionTreeItem::GetCollectionObjectCountText() const
 {
 	TSharedPtr<FCollectionItem> CollectionItemPtr = CollectionItem.Pin();
 
-	if (CollectionItemPtr.IsValid())
+	if (CollectionItemPtr)
 	{
-		switch(CollectionItemPtr->StorageMode)
+		if (CollectionItemPtr->StorageMode == ECollectionStorageMode::Static)
 		{
-		case ECollectionStorageMode::Static:
-			return FEditorFontGlyphs::List_Alt;
-
-		case ECollectionStorageMode::Dynamic:
-			return FEditorFontGlyphs::Bolt;
-
-		default:
-			break;
+			return CollectionItemPtr->NumObjects > 999
+				? NSLOCTEXT("ContentBrowser", "999+", "999+")
+				: FText::AsNumber(CollectionItemPtr->NumObjects);
 		}
-	}
 
-	return FText::GetEmpty();
-}
-
-FText SCollectionTreeItem::GetCollectionStorageModeToolTipText() const
-{
-	TSharedPtr<FCollectionItem> CollectionItemPtr = CollectionItem.Pin();
-
-	if (CollectionItemPtr.IsValid())
-	{
-		return ECollectionStorageMode::GetDescription(CollectionItemPtr->StorageMode);
+		return NSLOCTEXT("ContentBrowser", "InfinitySymbol", "\u221E");
 	}
 
 	return FText::GetEmpty();
 }
 
-FSlateColor SCollectionTreeItem::GetCollectionStatusColor() const
+FText SCollectionTreeItem::GetCollectionWarningText() const
 {
 	TSharedPtr<FCollectionItem> CollectionItemPtr = CollectionItem.Pin();
 
-	if (CollectionItemPtr.IsValid())
+	if (CollectionItemPtr)
 	{
-		switch(CollectionItemPtr->CurrentStatus)
+		switch (CollectionItemPtr->CurrentStatus)
 		{
-		case ECollectionItemStatus::IsUpToDateAndPopulated:
-			return FLinearColor(0.10616, 0.48777, 0.10616); // Green
-
-		case ECollectionItemStatus::IsUpToDateAndEmpty:
-			return FLinearColor::Gray;
-
 		case ECollectionItemStatus::IsOutOfDate:
-			return FLinearColor(0.87514, 0.42591, 0.07383); // Orange
+			return NSLOCTEXT("ContentBrowser", "CollectionStatus_IsOutOfDate", "Collection is not at the latest revision");
 
 		case ECollectionItemStatus::IsCheckedOutByAnotherUser:
+			return NSLOCTEXT("ContentBrowser", "CollectionStatus_IsCheckedOutByAnotherUser", "Collection is checked out by another user");
+
 		case ECollectionItemStatus::IsConflicted:
+			return NSLOCTEXT("ContentBrowser", "CollectionStatus_IsConflicted", "Collection is conflicted - please use your external source control provider to resolve this conflict");
+
 		case ECollectionItemStatus::IsMissingSCCProvider:
-			return FLinearColor(0.70117, 0.08464, 0.07593); // Red
+			return NSLOCTEXT("ContentBrowser", "CollectionStatus_IsMissingSCCProvider", "Collection is missing its source control provider - please check your source control settings");
 
 		case ECollectionItemStatus::HasLocalChanges:
-			return FLinearColor(0.10363, 0.53564, 0.7372); // Blue
+			return NSLOCTEXT("ContentBrowser", "CollectionStatus_HasLocalChanges", "Collection has local unsaved or uncomitted changes");
 
 		default:
 			break;
 		}
 	}
 
-	return FLinearColor::White;
+	return FText::GetEmpty();
 }
 
-FText SCollectionTreeItem::GetCollectionStatusToolTipText() const
+void SCollectionTreeItem::BuildToolTipInfo(const FOnBuildAssetTagItemToolTipInfoEntry& InCallback) const
 {
 	TSharedPtr<FCollectionItem> CollectionItemPtr = CollectionItem.Pin();
 
-	if (CollectionItemPtr.IsValid())
+	if (CollectionItemPtr)
 	{
-		switch(CollectionItemPtr->CurrentStatus)
+		InCallback(LOCTEXT("CollectionShareTypeKey", "Share Type"), ECollectionShareType::ToText(CollectionItemPtr->CollectionType));
+		InCallback(LOCTEXT("CollectionStorageModeKey", "Storage Mode"), ECollectionStorageMode::ToText(CollectionItemPtr->StorageMode));
+		if (CollectionItemPtr->StorageMode == ECollectionStorageMode::Static)
 		{
-		case ECollectionItemStatus::IsUpToDateAndPopulated:
-			return LOCTEXT("CollectionStatus_IsUpToDateAndPopulated", "Collection is up-to-date");
-
-		case ECollectionItemStatus::IsUpToDateAndEmpty:
-			return LOCTEXT("CollectionStatus_IsUpToDateAndEmpty", "Collection is empty");
-
-		case ECollectionItemStatus::IsOutOfDate:
-			return LOCTEXT("CollectionStatus_IsOutOfDate", "Collection is not at the latest revision");
-
-		case ECollectionItemStatus::IsCheckedOutByAnotherUser:
-			return LOCTEXT("CollectionStatus_IsCheckedOutByAnotherUser", "Collection is checked out by another user");
-
-		case ECollectionItemStatus::IsConflicted:
-			return LOCTEXT("CollectionStatus_IsConflicted", "Collection is conflicted - please use your external source control provider to resolve this conflict");
-
-		case ECollectionItemStatus::IsMissingSCCProvider:
-			return LOCTEXT("CollectionStatus_IsMissingSCCProvider", "Collection is missing its source control provider - please check your source control settings");
-
-		case ECollectionItemStatus::HasLocalChanges:
-			return LOCTEXT("CollectionStatus_HasLocalChanges", "Collection has local unsaved or uncomitted changes");
-
-		default:
-			break;
+			InCallback(LOCTEXT("CollectionObjectCountKey", "Object Count"), FText::AsNumber(CollectionItemPtr->NumObjects));
 		}
 	}
-
-	return FText::GetEmpty();
 }
 
 #undef LOCTEXT_NAMESPACE
diff --git a/Engine/Source/Editor/ContentBrowser/Private/SourcesViewWidgets.h b/Engine/Source/Editor/ContentBrowser/Private/SourcesViewWidgets.h
--- a/Engine/Source/Editor/ContentBrowser/Private/SourcesViewWidgets.h
+++ b/Engine/Source/Editor/ContentBrowser/Private/SourcesViewWidgets.h
@@ -14,6 +14,7 @@
 #include "CollectionViewTypes.h"
 #include "Widgets/Input/SCheckBox.h"
 #include "Editor/ContentBrowser/Private/SPathView.h"
+#include "SAssetTagItem.h"
 
 class SEditableTextBox;
 struct FTreeItem;
@@ -175,6 +176,7 @@ public:
 	SLATE_BEGIN_ARGS( SCollectionTreeItem )
 		: _CollectionItem( TSharedPtr<FCollectionItem>() )
 		, _ParentWidget()
+		, _ViewMode(EAssetTagItemViewMode::Standard)
 	{}
 
 		/** Data for the collection this item represents */
@@ -183,6 +185,9 @@ public:
 		/** The parent widget */
 		SLATE_ARGUMENT( TSharedPtr<SWidget>, ParentWidget )
 
+		/** Should this collection item use the standard or compact view? */
+		SLATE_ARGUMENT( EAssetTagItemViewMode, ViewMode )
+
 		/** Delegate for when the user begins to rename the item */
 		SLATE_EVENT( FOnBeginNameChange, OnBeginNameChange )
 
@@ -221,14 +226,18 @@ public:
 	/** Constructs this widget with InArgs */
 	void Construct( const FArguments& InArgs );
 
-	~SCollectionTreeItem();
-
 	virtual void Tick( const FGeometry& AllottedGeometry, const double InCurrentTime, const float InDeltaTime ) override;
 	virtual void OnDragEnter( const FGeometry& MyGeometry, const FDragDropEvent& DragDropEvent ) override;
 	virtual void OnDragLeave( const FDragDropEvent& DragDropEvent ) override;
 	virtual FReply OnDragOver( const FGeometry& MyGeometry, const FDragDropEvent& DragDropEvent ) override;
 	virtual FReply OnDrop( const FGeometry& MyGeometry, const FDragDropEvent& DragDropEvent ) override;
 
+	/** True when a drag is over this item with a drag operation that we know how to handle. The operation itself may not be valid to drop. */
+	bool IsDraggedOver() const
+	{
+		return bDraggedOver;
+	}
+
 private:
 	/** Used by OnDragEnter, OnDragOver, and OnDrop to check and update the validity of the drag operation */
 	bool ValidateDragDrop( const FGeometry& MyGeometry, const FDragDropEvent& DragDropEvent, bool& OutIsKnownDragOperation ) const;
@@ -246,22 +255,16 @@ private:
 	FText GetNameText() const;
 
 	/** Returns the color of the collection name */
-	FSlateColor GetCollectionColor() const;
-
-	/** Returns the image for the border around this item. Used for drag/drop operations */
-	const FSlateBrush* GetBorderImage() const;
-
-	/** Get the FontAwesome icon corresponding to the current collection storage mode */
-	FText GetCollectionStorageModeIconText() const;
+	FLinearColor GetCollectionColor() const;
 
-	/** Get the tooltip corresponding to the current collection storage mode */
-	FText GetCollectionStorageModeToolTipText() const;
+	/** Get the object count text for the current collection item */
+	FText GetCollectionObjectCountText() const;
 
-	/** Get the color to use for the collection item status */
-	FSlateColor GetCollectionStatusColor() const;
+	/** Get the warning message for the current collection item */
+	FText GetCollectionWarningText() const;
 
-	/** Get the tooltip corresponding to the current collection item status */
-	FText GetCollectionStatusToolTipText() const;
+	/** Build the tooltip info for this collection */
+	void BuildToolTipInfo(const FOnBuildAssetTagItemToolTipInfoEntry& InCallback) const;
 
 private:
 	/** A shared pointer to the parent widget. */
@@ -293,10 +296,4 @@ private:
 
 	/** Delegate for when a collection name has been entered for an item to verify the name before commit */
 	FOnVerifyRenameCommit OnVerifyRenameCommit;
-
-	/** Widget to display the name of the collection item and allows for renaming */
-	TSharedPtr< SInlineEditableTextBlock > InlineRenameWidget;
-
-	/** Handle to the registered EnterEditingMode delegate. */
-	FDelegateHandle EnterEditingModeDelegateHandle;
 };
diff --git a/Engine/Source/Editor/ContentBrowser/Public/IContentBrowserSingleton.h b/Engine/Source/Editor/ContentBrowser/Public/IContentBrowserSingleton.h
--- a/Engine/Source/Editor/ContentBrowser/Public/IContentBrowserSingleton.h
+++ b/Engine/Source/Editor/ContentBrowser/Public/IContentBrowserSingleton.h
@@ -86,12 +86,6 @@ struct FContentBrowserConfig
 	/** Whether the sources view should initially be expanded or not */
 	bool bExpandSourcesView;
 
-	/** Whether asset paths are shown in the Content Browser.  Only useful if you only want to show collections */
-	bool bShowAssetPathTree;
-
-	/** Forces collections to be initially visible, regardless of defaults */
-	bool bAlwaysShowCollections;
-
 	/** Collection to view initially */
 	FCollectionNameType SelectedCollectionName;
 
@@ -127,8 +121,6 @@ struct FContentBrowserConfig
 		, bCanShowClasses(true)
 		, bUseSourcesView(true)
 		, bExpandSourcesView(true)
-		, bShowAssetPathTree(true)
-		, bAlwaysShowCollections(false)
 		, SelectedCollectionName( NAME_None, ECollectionShareType::CST_Local )
 		, bUsePathPicker(true)
 		, bCanShowFilters(true)
diff --git a/Engine/Source/Editor/EditorStyle/Private/SlateEditorStyle.cpp b/Engine/Source/Editor/EditorStyle/Private/SlateEditorStyle.cpp
--- a/Engine/Source/Editor/EditorStyle/Private/SlateEditorStyle.cpp
+++ b/Engine/Source/Editor/EditorStyle/Private/SlateEditorStyle.cpp
@@ -6551,6 +6551,64 @@ void FSlateEditorStyle::FStyle::SetupContentBrowserStyle()
 		Set( "ContentBrowser.FilterNameFont", DEFAULT_FONT( "Regular", 10 ) );
 		Set( "ContentBrowser.FilterButtonBorder", new BOX_BRUSH( "Common/RoundedSelection_16x", FMargin(4.0f/16.0f) ) );
 
+		// Sources view
+		Set("ContentBrowser.Sources.Paths", new IMAGE_BRUSH("ContentBrowser/Sources_Paths_16x", Icon16x16));
+		Set("ContentBrowser.Sources.Collections", new IMAGE_BRUSH("ContentBrowser/Sources_Collections_Standard_16x", Icon16x16));
+		Set("ContentBrowser.Sources.Collections.Compact", new IMAGE_BRUSH("ContentBrowser/Sources_Collections_Compact_16x", Icon16x16));
+
+		// Asset tags (common)
+		Set("ContentBrowser.AssetTagBackground", new BOX_BRUSH("Common/RoundedSelection_16x", FMargin(4.0f/16.0f)));
+		Set("ContentBrowser.AssetTagTableRow", FTableRowStyle(NormalTableRowStyle)
+			.SetEvenRowBackgroundBrush(FSlateNoResource())
+			.SetEvenRowBackgroundHoveredBrush(BORDER_BRUSH("Common/RoundedSelection_16x", FMargin(2.0f/16.0f), FLinearColor(1.0f, 1.0f, 1.0f, 0.1f)))
+			.SetOddRowBackgroundBrush(FSlateNoResource())
+			.SetOddRowBackgroundHoveredBrush(BORDER_BRUSH("Common/RoundedSelection_16x", FMargin(2.0f/16.0f), FLinearColor(1.0f, 1.0f, 1.0f, 0.1f)))
+			.SetActiveBrush(BORDER_BRUSH("Common/RoundedSelection_16x", FMargin(2.0f/16.0f), SelectionColor))
+			.SetActiveHoveredBrush(BORDER_BRUSH("Common/RoundedSelection_16x", FMargin(2.0f/16.0f), SelectionColor))
+			.SetInactiveBrush(BORDER_BRUSH("Common/RoundedSelection_16x", FMargin(2.0f/16.0f), SelectionColor_Inactive))
+			.SetInactiveHoveredBrush(BORDER_BRUSH("Common/RoundedSelection_16x", FMargin(2.0f/16.0f), SelectionColor_Inactive))
+			.SetActiveHighlightedBrush(BORDER_BRUSH("Common/RoundedSelection_16x", FMargin(2.0f/16.0f), HighlightColor))
+			.SetInactiveHighlightedBrush(BORDER_BRUSH("Common/RoundedSelection_16x", FMargin(2.0f/16.0f), HighlightColor))
+			.SetTextColor(DefaultForeground)
+			.SetSelectedTextColor(DefaultForeground)
+			);
+
+		// Asset tags (standard)
+		Set("ContentBrowser.AssetTagButton", FCheckBoxStyle()
+			.SetUncheckedImage(IMAGE_BRUSH("ContentBrowser/AssetTagCheckbox_Flat", FVector2D(14.0f, 28.0f)))
+			.SetUncheckedHoveredImage(IMAGE_BRUSH("ContentBrowser/AssetTagCheckbox_Flat", FVector2D(14.0f, 28.0f), FLinearColor(0.5f, 0.5f, 0.5f, 1.0f)))
+			.SetUncheckedPressedImage(IMAGE_BRUSH("ContentBrowser/AssetTagCheckbox_Flat", FVector2D(14.0f, 28.0f), FLinearColor(0.5f, 0.5f, 0.5f, 1.0f)))
+			.SetUndeterminedImage(IMAGE_BRUSH("ContentBrowser/AssetTagCheckbox_Flat", FVector2D(14.0f, 28.0f)))
+			.SetUndeterminedHoveredImage(IMAGE_BRUSH("ContentBrowser/AssetTagCheckbox_Flat", FVector2D(14.0f, 28.0f), FLinearColor(0.5f, 0.5f, 0.5f, 1.0f)))
+			.SetUndeterminedPressedImage(IMAGE_BRUSH("ContentBrowser/AssetTagCheckbox_Flat", FVector2D(14.0f, 28.0f), FLinearColor(0.5f, 0.5f, 0.5f, 1.0f)))
+			.SetCheckedImage(IMAGE_BRUSH("ContentBrowser/AssetTagCheckbox_Flat", FVector2D(14.0f, 28.0f)))
+			.SetCheckedHoveredImage(IMAGE_BRUSH("ContentBrowser/AssetTagCheckbox_Flat", FVector2D(14.0f, 28.0f), FLinearColor(0.5f, 0.5f, 0.5f, 1.0f)))
+			.SetCheckedPressedImage(IMAGE_BRUSH("ContentBrowser/AssetTagCheckbox_Flat", FVector2D(14.0f, 28.0f), FLinearColor(0.5f, 0.5f, 0.5f, 1.0f)))
+			.SetPadding(0.0f)
+			);
+		Set("ContentBrowser.AssetTagNameFont", DEFAULT_FONT("Bold", 11));
+		Set("ContentBrowser.AssetTagNamePadding", FMargin(4.0f));
+		Set("ContentBrowser.AssetTagCountFont", DEFAULT_FONT("Bold", 8));
+		Set("ContentBrowser.AssetTagCountPadding", FMargin(4.0f));
+
+		// Asset tags (compact)
+		Set("ContentBrowser.AssetTagButton.Compact", FCheckBoxStyle()
+			.SetUncheckedImage(IMAGE_BRUSH("ContentBrowser/AssetTagCheckbox_Flat_Compact", FVector2D(10.0f, 20.0f)))
+			.SetUncheckedHoveredImage(IMAGE_BRUSH("ContentBrowser/AssetTagCheckbox_Flat_Compact", FVector2D(10.0f, 20.0f), FLinearColor(0.5f, 0.5f, 0.5f, 1.0f)))
+			.SetUncheckedPressedImage(IMAGE_BRUSH("ContentBrowser/AssetTagCheckbox_Flat_Compact", FVector2D(10.0f, 20.0f), FLinearColor(0.5f, 0.5f, 0.5f, 1.0f)))
+			.SetUndeterminedImage(IMAGE_BRUSH("ContentBrowser/AssetTagCheckbox_Flat_Compact", FVector2D(10.0f, 20.0f)))
+			.SetUndeterminedHoveredImage(IMAGE_BRUSH("ContentBrowser/AssetTagCheckbox_Flat_Compact", FVector2D(10.0f, 20.0f), FLinearColor(0.5f, 0.5f, 0.5f, 1.0f)))
+			.SetUndeterminedPressedImage(IMAGE_BRUSH("ContentBrowser/AssetTagCheckbox_Flat_Compact", FVector2D(10.0f, 20.0f), FLinearColor(0.5f, 0.5f, 0.5f, 1.0f)))
+			.SetCheckedImage(IMAGE_BRUSH("ContentBrowser/AssetTagCheckbox_Flat_Compact", FVector2D(10.0f, 20.0f)))
+			.SetCheckedHoveredImage(IMAGE_BRUSH("ContentBrowser/AssetTagCheckbox_Flat_Compact", FVector2D(10.0f, 20.0f), FLinearColor(0.5f, 0.5f, 0.5f, 1.0f)))
+			.SetCheckedPressedImage(IMAGE_BRUSH("ContentBrowser/AssetTagCheckbox_Flat_Compact", FVector2D(10.0f, 20.0f), FLinearColor(0.5f, 0.5f, 0.5f, 1.0f)))
+			.SetPadding(0.0f)
+			);
+		Set("ContentBrowser.AssetTagNameFont.Compact", DEFAULT_FONT("Regular", 10));
+		Set("ContentBrowser.AssetTagNamePadding.Compact", FMargin(2.0f));
+		Set("ContentBrowser.AssetTagCountFont.Compact", DEFAULT_FONT("Light", 8));
+		Set("ContentBrowser.AssetTagCountPadding.Compact", FMargin(2.0f));
+
 		// Thumbnail editing mode
 		Set( "ContentBrowser.EditModeLabelFont", FTextBlockStyle(NormalText)
 			.SetFont( DEFAULT_FONT( "Regular", 10 ) )
diff --git a/Engine/Source/Editor/EditorWidgets/Private/SAssetSearchBox.cpp b/Engine/Source/Editor/EditorWidgets/Private/SAssetSearchBox.cpp
--- a/Engine/Source/Editor/EditorWidgets/Private/SAssetSearchBox.cpp
+++ b/Engine/Source/Editor/EditorWidgets/Private/SAssetSearchBox.cpp
@@ -9,15 +9,51 @@
 #include "EditorStyleSet.h"
 #include "Widgets/Input/SSearchBox.h"
 
+/** Case sensitive hashing function for TMap */
+template <typename ValueType>
+struct FAssetSearchCategoryKeyMapFuncs : BaseKeyFuncs<ValueType, FText, /*bInAllowDuplicateKeys*/false>
+{
+	static FORCEINLINE const FString& GetSourceString(const FText& InText)
+	{
+		const FString* SourceString = FTextInspector::GetSourceString(InText);
+		check(SourceString);
+		return *SourceString;
+	}
+	static FORCEINLINE const FText& GetSetKey(const TPair<FText, ValueType>& Element)
+	{
+		return Element.Key;
+	}
+	static FORCEINLINE bool Matches(const FText& A, const FText& B)
+	{
+		return GetSourceString(A).Equals(GetSourceString(B), ESearchCase::CaseSensitive);
+	}
+	static FORCEINLINE uint32 GetKeyHash(const FText& Key)
+	{
+		return FLocKey::ProduceHash(GetSourceString(Key));
+	}
+};
+
 void SAssetSearchBox::Construct( const FArguments& InArgs )
 {
 	OnTextChanged = InArgs._OnTextChanged;
 	OnTextCommitted = InArgs._OnTextCommitted;
 	OnKeyDownHandler = InArgs._OnKeyDownHandler;
 	PossibleSuggestions = InArgs._PossibleSuggestions;
+	OnAssetSearchBoxSuggestionFilter = InArgs._OnAssetSearchBoxSuggestionFilter;
+	OnAssetSearchBoxSuggestionChosen = InArgs._OnAssetSearchBoxSuggestionChosen;
 	PreCommittedText = InArgs._InitialText.Get();
 	bMustMatchPossibleSuggestions = InArgs._MustMatchPossibleSuggestions.Get();
 
+	if (!OnAssetSearchBoxSuggestionFilter.IsBound())
+	{
+		OnAssetSearchBoxSuggestionFilter.BindStatic(&SAssetSearchBox::DefaultSuggestionFilterImpl);
+	}
+
+	if (!OnAssetSearchBoxSuggestionChosen.IsBound())
+	{
+		OnAssetSearchBoxSuggestionChosen.BindStatic(&SAssetSearchBox::DefaultSuggestionChosenImpl);
+	}
+
 	ChildSlot
 		[
 			SAssignNew(SuggestionBox, SMenuAnchor)
@@ -39,10 +75,10 @@ void SAssetSearchBox::Construct( const FArguments& InArgs )
 				.Padding( FMargin(2) )
 				[
 					SNew(SBox)
-					.WidthOverride(175)	// to enforce some minimum width, ideally we define the minimum, not a fixed width
-					.HeightOverride(175) // avoids flickering, ideally this would be adaptive to the content without flickering
+					.MinDesiredWidth(175)	// to enforce some minimum width, ideally we define the minimum, not a fixed width
+					.HeightOverride(250)	// avoids flickering, ideally this would be adaptive to the content without flickering
 					[
-						SAssignNew(SuggestionListView, SListView< TSharedPtr<FString> >)
+						SAssignNew(SuggestionListView, SListView< TSharedPtr<FSuggestionListEntry> >)
 						.ListItemsSource(&Suggestions)
 						.SelectionMode( ESelectionMode::Single )							// Ideally the mouse over would not highlight while keyboard controls the UI
 						.OnGenerateRow(this, &SAssetSearchBox::MakeSuggestionListItemWidget)
@@ -89,38 +125,38 @@ FReply SAssetSearchBox::HandleKeyDown( const FGeometry& MyGeometry, const FKeyEv
 	if ( SuggestionBox->IsOpen() && (InKeyEvent.GetKey() == EKeys::Up || InKeyEvent.GetKey() == EKeys::Down) )
 	{
 		const bool bSelectingUp = InKeyEvent.GetKey() == EKeys::Up;
-		TSharedPtr<FString> SelectedSuggestion = GetSelectedSuggestion();
+		TSharedPtr<FSuggestionListEntry> SelectedSuggestion = GetSelectedSuggestion();
 
+		int32 TargetIdx = INDEX_NONE;
 		if ( SelectedSuggestion.IsValid() )
 		{
-			// Find the selection index and select the previous or next one
-			int32 TargetIdx = INDEX_NONE;
-			for ( int32 SuggestionIdx = 0; SuggestionIdx < Suggestions.Num(); ++SuggestionIdx )
+			const int32 SelectionDirection = bSelectingUp ? -1 : 1;
+
+			// Select the next non-header suggestion, based on the direction of travel
+			TargetIdx = Suggestions.IndexOfByKey(SelectedSuggestion);
+			if (Suggestions.IsValidIndex(TargetIdx))
 			{
-				if ( Suggestions[SuggestionIdx] == SelectedSuggestion )
+				do
 				{
-					if ( bSelectingUp )
-					{
-						TargetIdx = SuggestionIdx - 1;
-					}
-					else
-					{
-						TargetIdx = SuggestionIdx + 1;
-					}
-					break;
+					TargetIdx += SelectionDirection;
 				}
+				while (Suggestions.IsValidIndex(TargetIdx) && Suggestions[TargetIdx]->bIsHeader);
 			}
-
-			if ( Suggestions.IsValidIndex(TargetIdx) )
+		}
+		else if ( !bSelectingUp && Suggestions.Num() > 0 )
+		{
+			// Nothing selected and pressed down, select the first non-header suggestion
+			TargetIdx = 0;
+			while (Suggestions.IsValidIndex(TargetIdx) && Suggestions[TargetIdx]->bIsHeader)
 			{
-				SuggestionListView->SetSelection( Suggestions[TargetIdx] );
-				SuggestionListView->RequestScrollIntoView( Suggestions[TargetIdx] );
+				TargetIdx += 1;
 			}
 		}
-		else if ( !bSelectingUp && Suggestions.Num() > 0 )
+
+		if (Suggestions.IsValidIndex(TargetIdx))
 		{
-			// Nothing selected and pressed down, select the first item
-			SuggestionListView->SetSelection( Suggestions[0] );
+			SuggestionListView->SetSelection(Suggestions[TargetIdx]);
+			SuggestionListView->RequestScrollIntoView(Suggestions[TargetIdx]);
 		}
 
 		return FReply::Handled();
@@ -159,14 +195,14 @@ void SAssetSearchBox::HandleTextChanged(const FText& NewText)
 
 void SAssetSearchBox::HandleTextCommitted(const FText& NewText, ETextCommit::Type CommitType)
 {
-	TSharedPtr<FString> SelectedSuggestion = GetSelectedSuggestion();
+	TSharedPtr<FSuggestionListEntry> SelectedSuggestion = GetSelectedSuggestion();
 
 	bool bCommitText = true;
 	FText CommittedText;
-	if ( SelectedSuggestion.IsValid() && CommitType != ETextCommit::OnCleared )
+	if ( SelectedSuggestion.IsValid() && !SelectedSuggestion->bIsHeader && CommitType != ETextCommit::OnCleared )
 	{
 		// Pressed selected a suggestion, set the text
-		CommittedText = FText::FromString( *SelectedSuggestion.Get() );
+		CommittedText = OnAssetSearchBoxSuggestionChosen.Execute(NewText, SelectedSuggestion->Suggestion);
 	}
 	else
 	{
@@ -175,7 +211,7 @@ void SAssetSearchBox::HandleTextCommitted(const FText& NewText, ETextCommit::Typ
 			// Clear text when escape is pressed then commit an empty string
 			CommittedText = FText::GetEmpty();
 		}
-		else if( PossibleSuggestions.Get().Contains(NewText.ToString()) )
+		else if (bMustMatchPossibleSuggestions && PossibleSuggestions.Get().ContainsByPredicate([this, NewTextStr = NewText.ToString()](const FAssetSearchBoxSuggestion& InSuggestion) { return InSuggestion.SuggestionString == NewTextStr; }))
 		{
 			// If the text is a suggestion, set the text.
 			CommittedText = NewText;
@@ -203,69 +239,123 @@ void SAssetSearchBox::HandleTextCommitted(const FText& NewText, ETextCommit::Typ
 	}
 }
 
-void SAssetSearchBox::OnSelectionChanged( TSharedPtr<FString> NewValue, ESelectInfo::Type SelectInfo )
+void SAssetSearchBox::OnSelectionChanged( TSharedPtr<FSuggestionListEntry> NewValue, ESelectInfo::Type SelectInfo )
 {
 	// If the user clicked directly on an item to select it, then accept the choice and close the window
-	if( SelectInfo == ESelectInfo::OnMouseClick )
+	if(SelectInfo == ESelectInfo::OnMouseClick && !NewValue->bIsHeader)
 	{
-		SetText( FText::FromString( *NewValue.Get() ));
+		const FText SearchText = InputText->GetText();
+		const FText NewText = OnAssetSearchBoxSuggestionChosen.Execute(SearchText, NewValue->Suggestion);
+		SetText(NewText);
 		SuggestionBox->SetIsOpen(false, false);
 		FocusEditBox();		
 	}
 }
 
-
-TSharedRef<ITableRow> SAssetSearchBox::MakeSuggestionListItemWidget(TSharedPtr<FString> Text, const TSharedRef<STableViewBase>& OwnerTable)
+TSharedRef<ITableRow> SAssetSearchBox::MakeSuggestionListItemWidget(TSharedPtr<FSuggestionListEntry> Suggestion, const TSharedRef<STableViewBase>& OwnerTable)
 {
-	check(Text.IsValid());
+	check(Suggestion.IsValid());
+	check(Suggestions.Num() > 0);
+
+	const bool bIsFirstItem = Suggestions[0] == Suggestion;
+	const bool bIdentItems = Suggestions[0]->bIsHeader;
+
+	TSharedPtr<SWidget> RowWidget;
+	if (Suggestion->bIsHeader)
+	{
+		TSharedRef<SVerticalBox> HeaderVBox = SNew(SVerticalBox);
+
+		if (!bIsFirstItem)
+		{
+			HeaderVBox->AddSlot()
+			.AutoHeight()
+			.Padding(0.0f, 4.0f, 0.0f, 2.0f) // Add some empty space before the line, and a tiny bit after it
+			[
+				SNew(SBorder)
+				.Padding(FEditorStyle::GetMargin("Menu.Separator.Padding")) // We'll use the border's padding to actually create the horizontal line
+				.BorderImage(FEditorStyle::GetBrush("Menu.Separator"))
+			];
+		}
+
+		HeaderVBox->AddSlot()
+		.AutoHeight()
+		[
+			SNew(STextBlock)
+			.Text(Suggestion->DisplayName)
+			.TextStyle(FEditorStyle::Get(), "Menu.Heading")
+		];
+
+		RowWidget = HeaderVBox;
+	}
+	else
+	{
+		RowWidget =
+			SNew(SBox)
+			.Padding(FEditorStyle::GetMargin(bIdentItems ? "Menu.Block.IndentedPadding" : "Menu.Block.Padding"))
+			[
+				SNew(STextBlock)
+				.Text(Suggestion->DisplayName)
+				.HighlightText(this, &SAssetSearchBox::GetHighlightText)
+			];
+	}
 
 	return
 		SNew(STableRow< TSharedPtr<FString> >, OwnerTable)
 		[
-			SNew(STextBlock)
-			.Text(FText::FromString(*Text.Get()))
-			.HighlightText(this, &SAssetSearchBox::GetHighlightText)
+			RowWidget.ToSharedRef()
 		];
 }
 
 FText SAssetSearchBox::GetHighlightText() const
 {
-	return InputText->GetText();
+	return SuggestionHighlightText;
 }
 
 void SAssetSearchBox::UpdateSuggestionList()
 {
-	const FString TypedText = InputText->GetText().ToString();
+	const FText SearchText = InputText->GetText();
 
-	Suggestions.Empty();
+	Suggestions.Reset();
+	SuggestionHighlightText = FText::GetEmpty();
 
-	if ( TypedText.Len() > 0 )
+	if (!SearchText.IsEmpty())
 	{
-		TArray<FString> AllSuggestions = PossibleSuggestions.Get();
+		typedef TMap<FText, TArray<TSharedPtr<FSuggestionListEntry>>, FDefaultSetAllocator, FAssetSearchCategoryKeyMapFuncs<TArray<TSharedPtr<FSuggestionListEntry>>>> FCategorizedSuggestionsMap;
 
-		for ( auto SuggestionIt = AllSuggestions.CreateConstIterator(); SuggestionIt; ++SuggestionIt )
-		{
-			const FString& Suggestion = *SuggestionIt;
-			if ( Suggestion.Contains(TypedText))
-			{
-				Suggestions.Add( MakeShareable(new FString(Suggestion)) );
-			}
-		}
+		// Get the potential suggestions and run them through the filter
+		TArray<FAssetSearchBoxSuggestion> FilteredSuggestions = PossibleSuggestions.Get();
+		OnAssetSearchBoxSuggestionFilter.Execute(SearchText, FilteredSuggestions, SuggestionHighlightText);
 
-		if ( Suggestions.Num() > 0 )
+		// Split the suggestions list into categories
+		FCategorizedSuggestionsMap CategorizedSuggestions;
+		for (const FAssetSearchBoxSuggestion& Suggestion : FilteredSuggestions)
 		{
-			// At least one suggestion was found, open the menu
-			SuggestionBox->SetIsOpen(true, false);
+			TArray<TSharedPtr<FSuggestionListEntry>>& CategorySuggestions = CategorizedSuggestions.FindOrAdd(Suggestion.CategoryName);
+			CategorySuggestions.Add(MakeShared<FSuggestionListEntry>(FSuggestionListEntry{ Suggestion.SuggestionString, Suggestion.DisplayName, false }));
 		}
-		else
+
+		// Rebuild the flat list in categorized groups
+		// If there is only one category, and that category is empty (undefined), then skip adding the category headers
+		const bool bSkipCategoryHeaders = CategorizedSuggestions.Num() == 1 && CategorizedSuggestions.Contains(FText::GetEmpty());
+		for (const auto& CategorySuggestionsPair : CategorizedSuggestions)
 		{
-			// No suggestions were found, close the menu
-			SuggestionBox->SetIsOpen(false, false);
+			if (!bSkipCategoryHeaders)
+			{
+				const FText CategoryDisplayName = CategorySuggestionsPair.Key.IsEmpty() ? NSLOCTEXT("AssetSearchBox", "UndefinedCategory", "Undefined") : CategorySuggestionsPair.Key;
+				Suggestions.Add(MakeShared<FSuggestionListEntry>(FSuggestionListEntry{ TEXT(""), CategoryDisplayName, true }));
+			}
+			Suggestions.Append(CategorySuggestionsPair.Value);
 		}
 	}
+
+	if (Suggestions.Num() > 0)
+	{
+		// At least one suggestion was found, open the menu
+		SuggestionBox->SetIsOpen(true, false);
+	}
 	else
 	{
-		// No text was typed, close the menu
+		// No suggestions were found, close the menu
 		SuggestionBox->SetIsOpen(false, false);
 	}
 
@@ -279,12 +369,12 @@ void SAssetSearchBox::FocusEditBox()
 	FSlateApplication::Get().SetKeyboardFocus( WidgetToFocusPath, EFocusCause::SetDirectly );
 }
 
-TSharedPtr<FString> SAssetSearchBox::GetSelectedSuggestion()
+TSharedPtr<SAssetSearchBox::FSuggestionListEntry> SAssetSearchBox::GetSelectedSuggestion() const
 {
-	TSharedPtr<FString> SelectedSuggestion;
+	TSharedPtr<FSuggestionListEntry> SelectedSuggestion;
 	if ( SuggestionBox->IsOpen() )
 	{
-		const TArray< TSharedPtr<FString> >& SelectedSuggestionList = SuggestionListView->GetSelectedItems();
+		const TArray< TSharedPtr<FSuggestionListEntry> >& SelectedSuggestionList = SuggestionListView->GetSelectedItems();
 		if ( SelectedSuggestionList.Num() > 0 )
 		{
 			// Selection mode is Single, so there should only be one suggestion at the most
@@ -294,3 +384,20 @@ TSharedPtr<FString> SAssetSearchBox::GetSelectedSuggestion()
 
 	return SelectedSuggestion;
 }
+
+void SAssetSearchBox::DefaultSuggestionFilterImpl(const FText& SearchText, TArray<FAssetSearchBoxSuggestion>& PossibleSuggestions, FText& SuggestionHighlightText)
+{
+	// Default implementation just filters against the current search text
+	PossibleSuggestions.RemoveAll([SearchStr = SearchText.ToString()](const FAssetSearchBoxSuggestion& InSuggestion)
+	{
+		return !InSuggestion.SuggestionString.Contains(SearchStr);
+	});
+
+	SuggestionHighlightText = SearchText;
+}
+
+FText SAssetSearchBox::DefaultSuggestionChosenImpl(const FText& SearchText, const FString& Suggestion)
+{
+	// Default implementation just uses the suggestion as the search text
+	return FText::FromString(Suggestion);
+}
diff --git a/Engine/Source/Editor/EditorWidgets/Public/SAssetSearchBox.h b/Engine/Source/Editor/EditorWidgets/Public/SAssetSearchBox.h
--- a/Engine/Source/Editor/EditorWidgets/Public/SAssetSearchBox.h
+++ b/Engine/Source/Editor/EditorWidgets/Public/SAssetSearchBox.h
@@ -13,6 +13,33 @@
 
 class SMenuAnchor;
 
+/** External suggestion entry provided to SAssetSearchBox */
+struct FAssetSearchBoxSuggestion
+{
+	/** The raw suggestion string that should be used with the search box */
+	FString SuggestionString;
+
+	/** The user-facing display name of this suggestion */
+	FText DisplayName;
+
+	/** The user-facing category name of this suggestion (if any) */
+	FText CategoryName;
+
+	static FAssetSearchBoxSuggestion MakeSimpleSuggestion(FString InSuggestionString)
+	{
+		FAssetSearchBoxSuggestion SimpleSuggestion;
+		SimpleSuggestion.SuggestionString = MoveTemp(InSuggestionString);
+		SimpleSuggestion.DisplayName = FText::FromString(SimpleSuggestion.SuggestionString);
+		return SimpleSuggestion;
+	}
+};
+
+/** A delegate for a callback to filter the given suggestion list, to allow custom filtering behavior */
+DECLARE_DELEGATE_ThreeParams(FOnAssetSearchBoxSuggestionFilter, const FText& /*SearchText*/, TArray<FAssetSearchBoxSuggestion>& /*PossibleSuggestions*/, FText& /*SuggestionHighlightText*/);
+
+/** A delegate for a callback when a suggestion entry is chosen during an asset search, to allow custom compositing behavior of the suggestion into the search text */
+DECLARE_DELEGATE_RetVal_TwoParams(FText, FOnAssetSearchBoxSuggestionChosen, const FText& /*SearchText*/, const FString& /*Suggestion*/);
+
 /**
  * A widget to provide a search box with a filtered dropdown menu.
  */
@@ -26,12 +53,12 @@ public:
 		, _OnTextCommitted()
 		, _InitialText()
 		, _HintText()
-		, _PossibleSuggestions(TArray<FString>())
+		, _PossibleSuggestions(TArray<FAssetSearchBoxSuggestion>())
 		, _DelayChangeNotificationsWhileTyping( false )
 		, _MustMatchPossibleSuggestions( false )
 	{}
 
-	/** Where to place the suggestion list */
+		/** Where to place the suggestion list */
 		SLATE_ARGUMENT( EMenuPlacement, SuggestionListPlacement )
 
 		/** Invoked whenever the text changes */
@@ -47,7 +74,7 @@ public:
 		SLATE_ATTRIBUTE( FText, HintText )
 
 		/** All possible suggestions for the search text */
-		SLATE_ATTRIBUTE( TArray<FString>, PossibleSuggestions )
+		SLATE_ATTRIBUTE( TArray<FAssetSearchBoxSuggestion>, PossibleSuggestions )
 
 		/** Whether the SearchBox should delay notifying listeners of text changed events until the user is done typing */
 		SLATE_ATTRIBUTE( bool, DelayChangeNotificationsWhileTyping )
@@ -55,13 +82,19 @@ public:
 		/** Whether the SearchBox allows entries that don't match the possible suggestions */
 		SLATE_ATTRIBUTE( bool, MustMatchPossibleSuggestions )
 
+		/** Callback to filter the given suggestion list, to allow custom filtering behavior */
+		SLATE_EVENT( FOnAssetSearchBoxSuggestionFilter, OnAssetSearchBoxSuggestionFilter )
+
+		/** Callback when a suggestion entry is chosen during an asset search, to allow custom compositing behavior of the suggestion into the search text */
+		SLATE_EVENT( FOnAssetSearchBoxSuggestionChosen, OnAssetSearchBoxSuggestionChosen )
+
 		/** Callback delegate to have first chance handling of the OnKeyDown event */
 		SLATE_EVENT( FOnKeyDown, OnKeyDownHandler )
 
-		SLATE_END_ARGS()
+	SLATE_END_ARGS()
 
-		/** Constructs this widget with InArgs */
-		void Construct( const FArguments& InArgs );
+	/** Constructs this widget with InArgs */
+	void Construct( const FArguments& InArgs );
 
 	/** Sets the text string currently being edited */
 	void SetText(const TAttribute< FText >& InNewText);
@@ -77,6 +110,13 @@ public:
 	virtual FReply OnFocusReceived( const FGeometry& MyGeometry, const FFocusEvent& InFocusEvent ) override;
 
 private:
+	struct FSuggestionListEntry
+	{
+		FString Suggestion;
+		FText DisplayName;
+		bool bIsHeader = false;
+	};
+
 	/** First chance handler for key down events to the editable text widget */
 	FReply HandleKeyDown(const FGeometry& MyGeometry, const FKeyEvent& InKeyEvent);
 
@@ -87,10 +127,10 @@ private:
 	void HandleTextCommitted(const FText& NewText, ETextCommit::Type CommitType);
 
 	/** Called by SListView when the selection changes in the suggestion list */
-	void OnSelectionChanged( TSharedPtr<FString> NewValue, ESelectInfo::Type SelectInfo );
+	void OnSelectionChanged( TSharedPtr<FSuggestionListEntry> NewValue, ESelectInfo::Type SelectInfo );
 
 	/** Makes the widget for a suggestion message in the list view */
-	TSharedRef<ITableRow> MakeSuggestionListItemWidget(TSharedPtr<FString> Text, const TSharedRef<STableViewBase>& OwnerTable);
+	TSharedRef<ITableRow> MakeSuggestionListItemWidget(TSharedPtr<FSuggestionListEntry> Suggestion, const TSharedRef<STableViewBase>& OwnerTable);
 
 	/** Gets the text to highlight in the suggestion list */
 	FText GetHighlightText() const;
@@ -102,7 +142,13 @@ private:
 	void FocusEditBox();
 
 	/** Returns the currently selected suggestion */
-	TSharedPtr<FString> GetSelectedSuggestion();
+	TSharedPtr<FSuggestionListEntry> GetSelectedSuggestion() const;
+
+	/** Default implementation of OnAssetSearchBoxSuggestionFilter, if no external implementation is provided */
+	static void DefaultSuggestionFilterImpl(const FText& SearchText, TArray<FAssetSearchBoxSuggestion>& PossibleSuggestions, FText& SuggestionHighlightText);
+
+	/** Default implementation of OnAssetSearchBoxSuggestionChosen, if no external implementation is provided */
+	static FText DefaultSuggestionChosenImpl(const FText& SearchText, const FString& Suggestion);
 
 private:
 	/** The editable text field */
@@ -111,14 +157,23 @@ private:
 	/** The the state of the text prior to being committed */
 	FText PreCommittedText;
 
+	/** The highlight text to use for the suggestions list */
+	FText SuggestionHighlightText;
+
 	/** Auto completion elements */
 	TSharedPtr< SMenuAnchor > SuggestionBox;
 
 	/** All suggestions stored in this widget for the list view */
-	TArray< TSharedPtr<FString> > Suggestions;
+	TArray< TSharedPtr<FSuggestionListEntry> > Suggestions;
 
 	/** The list view for showing all suggestions */
-	TSharedPtr< SListView< TSharedPtr<FString> > > SuggestionListView;
+	TSharedPtr< SListView< TSharedPtr<FSuggestionListEntry> > > SuggestionListView;
+
+	/** Delegate to filter the given suggestion list, to allow custom filtering behavior */
+	FOnAssetSearchBoxSuggestionFilter OnAssetSearchBoxSuggestionFilter;
+
+	/** Delegate when a suggestion entry is chosen during an asset search, to allow custom compositing behavior of the suggestion into the search text */
+	FOnAssetSearchBoxSuggestionChosen OnAssetSearchBoxSuggestionChosen;
 
 	/** Delegate for when text is changed in the edit box */
 	FOnTextChanged OnTextChanged;
@@ -130,7 +185,7 @@ private:
 	FOnKeyDown OnKeyDownHandler;
 
 	/** All possible suggestions for the search text */
-	TAttribute< TArray<FString> > PossibleSuggestions;
+	TAttribute< TArray<FAssetSearchBoxSuggestion> > PossibleSuggestions;
 
 	/** Determines whether or not the committed text should match a suggestion */
 	bool bMustMatchPossibleSuggestions;
diff --git a/Engine/Source/Editor/Persona/Private/AssetSearchBoxUtilPersona.cpp b/Engine/Source/Editor/Persona/Private/AssetSearchBoxUtilPersona.cpp
--- a/Engine/Source/Editor/Persona/Private/AssetSearchBoxUtilPersona.cpp
+++ b/Engine/Source/Editor/Persona/Private/AssetSearchBoxUtilPersona.cpp
@@ -4,7 +4,6 @@
 #include "ReferenceSkeleton.h"
 #include "Animation/Skeleton.h"
 #include "Engine/SkeletalMesh.h"
-#include "SAssetSearchBox.h"
 #include "Engine/SkeletalMeshSocket.h"
 
 void SAssetSearchBoxForBones::Construct( const FArguments& InArgs, const class UObject* Outer, TSharedPtr<class IPropertyHandle> BoneNameProperty )
@@ -18,14 +17,14 @@ void SAssetSearchBoxForBones::Construct( const FArguments& InArgs, const class U
 
 	const USkeleton* Skeleton = NULL;
 
-	TArray<FString> PossibleSuggestions;
+	TArray<FAssetSearchBoxSuggestion> PossibleSuggestions;
 	if( const USkeletalMesh* SkeletalMesh = Cast<const USkeletalMesh>(Outer) )
 	{
 		if( InArgs._IncludeSocketsForSuggestions.Get() )
 		{
 			for( auto SocketIt=const_cast<USkeletalMesh*>(SkeletalMesh)->GetMeshOnlySocketList().CreateConstIterator(); SocketIt; ++SocketIt )
 			{
-				PossibleSuggestions.Add((*SocketIt)->SocketName.ToString());
+				PossibleSuggestions.Add(FAssetSearchBoxSuggestion::MakeSimpleSuggestion((*SocketIt)->SocketName.ToString()));
 			}
 		}
 		Skeleton = SkeletalMesh->Skeleton;
@@ -38,7 +37,7 @@ void SAssetSearchBoxForBones::Construct( const FArguments& InArgs, const class U
 	{
 		for( auto SocketIt=Skeleton->Sockets.CreateConstIterator(); SocketIt; ++SocketIt )
 		{
-			PossibleSuggestions.Add((*SocketIt)->SocketName.ToString());
+			PossibleSuggestions.Add(FAssetSearchBoxSuggestion::MakeSimpleSuggestion((*SocketIt)->SocketName.ToString()));
 		}
 	}
 	check(Skeleton);
@@ -46,7 +45,7 @@ void SAssetSearchBoxForBones::Construct( const FArguments& InArgs, const class U
 	const TArray<FMeshBoneInfo>& Bones = Skeleton->GetReferenceSkeleton().GetRefBoneInfo();
 	for( auto BoneIt=Bones.CreateConstIterator(); BoneIt; ++BoneIt )
 	{
-		PossibleSuggestions.Add(BoneIt->Name.ToString());
+		PossibleSuggestions.Add(FAssetSearchBoxSuggestion::MakeSimpleSuggestion(BoneIt->Name.ToString()));
 	}
 
 	// create the asset search box
@@ -108,9 +107,9 @@ void SAssetSearchBoxForCurves::Construct(const FArguments& InArgs, const class U
 		];
 }
 
-TArray<FString> SAssetSearchBoxForCurves::GetCurveSearchSuggestions() const
+TArray<FAssetSearchBoxSuggestion> SAssetSearchBoxForCurves::GetCurveSearchSuggestions() const
 {
-	TArray<FString> PossibleSuggestions;
+	TArray<FAssetSearchBoxSuggestion> PossibleSuggestions;
 	if (USkeleton* Skel = Skeleton.Get())
 	{
 		if (const FSmartNameMapping* Mapping = Skel->GetSmartNameContainer(USkeleton::AnimCurveMappingName))
@@ -119,7 +118,7 @@ TArray<FString> SAssetSearchBoxForCurves::GetCurveSearchSuggestions() const
 			Mapping->FillNameArray(Names);
 			for (const FName& Name : Names)
 			{
-				PossibleSuggestions.Add(Name.ToString());
+				PossibleSuggestions.Add(FAssetSearchBoxSuggestion::MakeSimpleSuggestion(Name.ToString()));
 			}
 		}
 	}
diff --git a/Engine/Source/Editor/Persona/Private/AssetSearchBoxUtilPersona.h b/Engine/Source/Editor/Persona/Private/AssetSearchBoxUtilPersona.h
--- a/Engine/Source/Editor/Persona/Private/AssetSearchBoxUtilPersona.h
+++ b/Engine/Source/Editor/Persona/Private/AssetSearchBoxUtilPersona.h
@@ -7,9 +7,9 @@
 #include "Widgets/SCompoundWidget.h"
 #include "Framework/SlateDelegates.h"
 #include "PropertyHandle.h"
+#include "SAssetSearchBox.h"
 
 class USkeleton;
-class SAssetSearchBox;
 
 class SAssetSearchBoxForBones : public SCompoundWidget
 {
@@ -84,5 +84,5 @@ private:
 	TWeakObjectPtr<USkeleton> Skeleton;
 
 	/** Get the search suggestions */
-	TArray<FString> GetCurveSearchSuggestions() const;
+	TArray<FAssetSearchBoxSuggestion> GetCurveSearchSuggestions() const;
 };
diff --git a/Engine/Source/Editor/UnrealEd/Classes/Settings/ContentBrowserSettings.h b/Engine/Source/Editor/UnrealEd/Classes/Settings/ContentBrowserSettings.h
--- a/Engine/Source/Editor/UnrealEd/Classes/Settings/ContentBrowserSettings.h
+++ b/Engine/Source/Editor/UnrealEd/Classes/Settings/ContentBrowserSettings.h
@@ -104,18 +104,6 @@ public:
 		return ( ( bExcludeOverride ? false : OverrideDisplayPluginFolders ) || DisplayPluginFolders );
 	}
 
-	/** Sets whether we are allowed to display collection folders or not */
-	void SetDisplayCollections(bool bInDisplayCollections)
-	{
-		DisplayCollections = bInDisplayCollections;
-	}
-
-	/** Gets whether we are allowed to display the collection folders or not*/
-	bool GetDisplayCollections() const
-	{
-		return DisplayCollections;
-	}
-
 	/** Sets whether we are allowed to display favorite folders or not */
 	void SetDisplayFavorites(bool bInDisplayFavorites)
 	{
@@ -216,9 +204,6 @@ private:
 	/** Temporary override for the DisplayPluginFolders setting */
 	bool OverrideDisplayPluginFolders;
 
-	UPROPERTY(config)
-	bool DisplayCollections;
-
 	UPROPERTY(config)
 	bool DisplayFavorites;
 
diff --git a/Engine/Source/Editor/UnrealEd/Private/Commandlets/AssetRegistryGenerator.cpp b/Engine/Source/Editor/UnrealEd/Private/Commandlets/AssetRegistryGenerator.cpp
--- a/Engine/Source/Editor/UnrealEd/Private/Commandlets/AssetRegistryGenerator.cpp
+++ b/Engine/Source/Editor/UnrealEd/Private/Commandlets/AssetRegistryGenerator.cpp
@@ -706,6 +706,79 @@ void FAssetRegistryGenerator::UpdateKeptPackagesAssetData()
 	}
 }
 
+void FAssetRegistryGenerator::UpdateCollectionAssetData()
+{
+	// Read out the per-platform settings use to build the list of collections to tag
+	bool bTagAllCollections = false;
+	TArray<FString> CollectionsToIncludeOrExclude;
+	{
+		const FString PlatformIniName = TargetPlatform->IniPlatformName();
+
+		FConfigFile PlatformEngineIni;
+		FConfigCacheIni::LoadLocalIniFile(PlatformEngineIni, TEXT("Engine"), true, (!PlatformIniName.IsEmpty() ? *PlatformIniName : ANSI_TO_TCHAR(FPlatformProperties::IniPlatformName())));
+
+		// The list of collections will either be a inclusive or a exclusive depending on the value of bTagAllCollections
+		PlatformEngineIni.GetBool(TEXT("AssetRegistry"), TEXT("bTagAllCollections"), bTagAllCollections);
+		PlatformEngineIni.GetArray(TEXT("AssetRegistry"), bTagAllCollections ? TEXT("CollectionsToExcludeAsTags") : TEXT("CollectionsToIncludeAsTags"), CollectionsToIncludeOrExclude);
+	}
+
+	// Build the list of collections we should tag for each asset
+	TMap<FName, TArray<FName>> AssetPathNamesToCollectionTags;
+	{
+		ICollectionManager& CollectionManager = FCollectionManagerModule::GetModule().Get();
+
+		TArray<FCollectionNameType> CollectionNamesToTag;
+		CollectionManager.GetCollections(CollectionNamesToTag);
+		if (bTagAllCollections)
+		{
+			CollectionNamesToTag.RemoveAll([&CollectionsToIncludeOrExclude](const FCollectionNameType& CollectionNameAndType)
+			{
+				return CollectionsToIncludeOrExclude.Contains(CollectionNameAndType.Name.ToString());
+			});
+		}
+		else
+		{
+			CollectionNamesToTag.RemoveAll([&CollectionsToIncludeOrExclude](const FCollectionNameType& CollectionNameAndType)
+			{
+				return !CollectionsToIncludeOrExclude.Contains(CollectionNameAndType.Name.ToString());
+			});
+		}
+		
+		TArray<FName> TmpAssetPathNames;
+		for (const FCollectionNameType& CollectionNameToTag : CollectionNamesToTag)
+		{
+			const FName CollectionTagName = *FString::Printf(TEXT("%s%s"), FAssetData::CollectionTagPrefix, *CollectionNameToTag.Name.ToString());
+
+			TmpAssetPathNames.Reset();
+			CollectionManager.GetAssetsInCollection(CollectionNameToTag.Name, CollectionNameToTag.Type, TmpAssetPathNames);
+
+			for (const FName AssetPathName : TmpAssetPathNames)
+			{
+				TArray<FName>& CollectionTagsForAsset = AssetPathNamesToCollectionTags.FindOrAdd(AssetPathName);
+				CollectionTagsForAsset.AddUnique(CollectionTagName);
+			}
+		}
+	}
+
+	// Apply the collection tags to the asset registry state
+	for (const auto& AssetPathNameToCollectionTagsPair : AssetPathNamesToCollectionTags)
+	{
+		const FName AssetPathName = AssetPathNameToCollectionTagsPair.Key;
+		const TArray<FName>& CollectionTagsForAsset = AssetPathNameToCollectionTagsPair.Value;
+
+		const FAssetData* AssetData = State.GetAssetByObjectPath(AssetPathName);
+		if (AssetData)
+		{
+			FAssetDataTagMap TagsAndValues = AssetData->TagsAndValues.GetMap();
+			for (const FName CollectionTagName : CollectionTagsForAsset)
+			{
+				TagsAndValues.Add(CollectionTagName, FString()); // TODO: Does this need a value to avoid being trimmed?
+			}
+			State.UpdateAssetData(FAssetData(AssetData->PackageName, AssetData->PackagePath, AssetData->AssetName, AssetData->AssetClass, TagsAndValues, AssetData->ChunkIDs, AssetData->PackageFlags));
+		}
+	}
+}
+
 void FAssetRegistryGenerator::Initialize(const TArray<FName> &InStartupPackages)
 {
 	StartupPackages.Append(InStartupPackages);
@@ -1045,6 +1118,7 @@ bool FAssetRegistryGenerator::SaveAssetRegistry(const FString& SandboxPath, bool
 	AssetRegistry.InitializeTemporaryAssetRegistryState(State, SaveOptions, true);
 	// Then possibly apply AssetData with TagsAndValues from a previous AssetRegistry for packages kept from a previous cook
 	UpdateKeptPackagesAssetData();
+	UpdateCollectionAssetData();
 
 	if (DevelopmentSaveOptions.bSerializeAssetRegistry && bSerializeDevelopmentAssetRegistry)
 	{
diff --git a/Engine/Source/Editor/UnrealEd/Private/DragAndDrop/CollectionDragDropOp.cpp b/Engine/Source/Editor/UnrealEd/Private/DragAndDrop/CollectionDragDropOp.cpp
--- a/Engine/Source/Editor/UnrealEd/Private/DragAndDrop/CollectionDragDropOp.cpp
+++ b/Engine/Source/Editor/UnrealEd/Private/DragAndDrop/CollectionDragDropOp.cpp
@@ -5,8 +5,9 @@
 #include "CollectionManagerModule.h"
 #include "AssetRegistryModule.h"
 #include "IAssetRegistry.h"
+#include "SAssetTagItem.h"
 
-TArray<FAssetData>  FCollectionDragDropOp::GetAssets() const
+TArray<FAssetData> FCollectionDragDropOp::GetAssets() const
 {
 	ICollectionManager& CollectionManager = FModuleManager::LoadModuleChecked<FCollectionManagerModule>("CollectionManager").Get();
 	FAssetRegistryModule& AssetRegistryModule = FModuleManager::LoadModuleChecked<FAssetRegistryModule>(TEXT("AssetRegistry"));
@@ -32,4 +33,29 @@ TArray<FAssetData>  FCollectionDragDropOp::GetAssets() const
 	}
 	
 	return AssetDatas;
-}
\ No newline at end of file
+}
+
+TSharedPtr<SWidget> FCollectionDragDropOp::GetDefaultDecorator() const
+{
+	return SNew(SBorder)
+		.Padding(0)
+		.BorderImage(FEditorStyle::GetBrush("ContentBrowser.AssetDragDropTooltipBackground"))
+		//.BorderImage(FEditorStyle::GetBrush("ToolPanel.GroupBorder"))
+		[
+			SNew(SAssetTagItem)
+			.ViewMode(AssetTagViewMode)
+			.DisplayName(this, &FCollectionDragDropOp::GetDecoratorText)
+		];
+}
+
+FText FCollectionDragDropOp::GetDecoratorText() const
+{
+	if (CurrentHoverText.IsEmpty() && Collections.Num() > 0)
+	{
+		return (Collections.Num() == 1)
+			? FText::FromName(Collections[0].Name)
+			: FText::Format(NSLOCTEXT("ContentBrowser", "CollectionDragDropDescription", "{0} and {1} {1}|plural(one=other,other=others)"), FText::FromName(Collections[0].Name), Collections.Num() - 1);
+	}
+
+	return CurrentHoverText;
+}
diff --git a/Engine/Source/Editor/UnrealEd/Public/Commandlets/AssetRegistryGenerator.h b/Engine/Source/Editor/UnrealEd/Public/Commandlets/AssetRegistryGenerator.h
--- a/Engine/Source/Editor/UnrealEd/Public/Commandlets/AssetRegistryGenerator.h
+++ b/Engine/Source/Editor/UnrealEd/Public/Commandlets/AssetRegistryGenerator.h
@@ -235,6 +235,12 @@ private:
 	 */
 	void UpdateKeptPackagesAssetData();
 
+	/**
+	 * Updates AssetData with TagsAndValues corresponding to any collections 
+	 * flagged for inclusion as asset registry tags.
+	 */
+	void UpdateCollectionAssetData();
+
 	/**
 	 * Adds a package to chunk manifest
 	 * 
diff --git a/Engine/Source/Editor/UnrealEd/Public/DragAndDrop/CollectionDragDropOp.h b/Engine/Source/Editor/UnrealEd/Public/DragAndDrop/CollectionDragDropOp.h
--- a/Engine/Source/Editor/UnrealEd/Public/DragAndDrop/CollectionDragDropOp.h
+++ b/Engine/Source/Editor/UnrealEd/Public/DragAndDrop/CollectionDragDropOp.h
@@ -3,18 +3,13 @@
 #pragma once
 
 #include "CoreMinimal.h"
-#include "Widgets/DeclarativeSyntaxSupport.h"
 #include "Input/DragAndDrop.h"
-#include "CollectionManagerTypes.h"
-#include "Widgets/Layout/SBorder.h"
-#include "Widgets/SBoxPanel.h"
-#include "Widgets/Text/STextBlock.h"
-#include "Widgets/Images/SImage.h"
-#include "EditorStyleSet.h"
 #include "DragAndDrop/DecoratedDragDropOp.h"
+#include "CollectionManagerTypes.h"
 #include "AssetData.h"
+#include "AssetTagItemTypes.h"
 
-class FCollectionDragDropOp : public FDecoratedDragDropOp
+class UNREALED_API FCollectionDragDropOp : public FDecoratedDragDropOp
 {
 public:
 	DRAG_DROP_OPERATOR_TYPE(FCollectionDragDropOp, FDecoratedDragDropOp)
@@ -22,10 +17,11 @@ public:
 	/** Data for the collections this item represents */
 	TArray<FCollectionNameType> Collections;
 
-	static TSharedRef<FCollectionDragDropOp> New(TArray<FCollectionNameType> InCollections)
+	static TSharedRef<FCollectionDragDropOp> New(TArray<FCollectionNameType> InCollections, const EAssetTagItemViewMode InAssetTagViewMode = EAssetTagItemViewMode::Standard)
 	{
 		TSharedRef<FCollectionDragDropOp> Operation = MakeShareable(new FCollectionDragDropOp);
 		
+		Operation->AssetTagViewMode = InAssetTagViewMode;
 		Operation->MouseCursor = EMouseCursor::GrabHandClosed;
 		Operation->Collections = MoveTemp(InCollections);
 		Operation->Construct();
@@ -37,63 +33,10 @@ public:
 	/** @return The assets from this drag operation */
 	TArray<FAssetData> GetAssets() const;
 
-	FText GetDecoratorText() const
-	{
-		if (CurrentHoverText.IsEmpty() && Collections.Num() > 0)
-		{
-			return (Collections.Num() == 1)
-				? FText::FromName(Collections[0].Name)
-				: FText::Format(NSLOCTEXT("ContentBrowser", "CollectionDragDropDescription", "{0} and {1} {1}|plural(one=other,other=others)"), FText::FromName(Collections[0].Name), Collections.Num() - 1);
-		}
-
-		return CurrentHoverText;
-	}
+	virtual TSharedPtr<SWidget> GetDefaultDecorator() const override;
 
-	virtual TSharedPtr<SWidget> GetDefaultDecorator() const override
-	{
-		const FSlateBrush* const CollectionIcon = (Collections.Num() > 0) ? FEditorStyle::GetBrush(ECollectionShareType::GetIconStyleName(Collections[0].Type)) : nullptr;
+private:
+	FText GetDecoratorText() const;
 
-		return 
-			SNew(SBorder)
-			.BorderImage(FEditorStyle::GetBrush("ContentBrowser.AssetDragDropTooltipBackground"))
-			.Content()
-			[
-				SNew(SHorizontalBox)
-
-				// Left slot is collection icon
-				+ SHorizontalBox::Slot()
-				.AutoWidth()
-				.VAlign(VAlign_Center)
-				[
-					SNew(SImage)
-					.Image(CollectionIcon)
-				]
-
-				// Right slot is for tooltip
-				+ SHorizontalBox::Slot()
-				.AutoWidth()
-				.VAlign(VAlign_Center)
-				[
-					SNew(SHorizontalBox)
-
-					+ SHorizontalBox::Slot()
-					.AutoWidth()
-					.Padding(3.0f)
-					.VAlign(VAlign_Center)
-					[
-						SNew(SImage) 
-						.Image(this, &FCollectionDragDropOp::GetIcon)
-					]
-
-					+SHorizontalBox::Slot()
-					.AutoWidth()
-					.Padding(0,0,3,0)
-					.VAlign(VAlign_Center)
-					[
-						SNew(STextBlock) 
-						.Text(this, &FCollectionDragDropOp::GetDecoratorText)
-					]
-				]
-			];
-	}
+	EAssetTagItemViewMode AssetTagViewMode = EAssetTagItemViewMode::Standard;
 };
diff --git a/Engine/Source/Editor/UnrealEd/UnrealEd.Build.cs b/Engine/Source/Editor/UnrealEd/UnrealEd.Build.cs
--- a/Engine/Source/Editor/UnrealEd/UnrealEd.Build.cs
+++ b/Engine/Source/Editor/UnrealEd/UnrealEd.Build.cs
@@ -110,6 +110,7 @@ public class UnrealEd : ModuleRules
 			new string[]
 			{
 				"AssetRegistry",
+				"AssetTagsEditor",
 				"LevelSequence",
 				"AnimGraph",
 				"AppFramework",
@@ -278,6 +279,7 @@ public class UnrealEd : ModuleRules
 		PublicIncludePathModuleNames.AddRange(
 			new string[] {
 				"AssetRegistry",
+				"AssetTagsEditor",
 				"CollectionManager",
 				"BlueprintGraph",
 				"AddContentDialog",
diff --git a/Engine/Source/Editor/VREditor/UI/VREditorUISystem.cpp b/Engine/Source/Editor/VREditor/UI/VREditorUISystem.cpp
--- a/Engine/Source/Editor/VREditor/UI/VREditorUISystem.cpp
+++ b/Engine/Source/Editor/VREditor/UI/VREditorUISystem.cpp
@@ -953,16 +953,11 @@ void UVREditorUISystem::CreateUIs()
 			Config.ThumbnailScale = 0.4f;
 			Config.bCanShowFilters = true;
 			Config.bUsePathPicker = true;
-			Config.bShowAssetPathTree = true;
-			Config.bAlwaysShowCollections = false;
 			Config.bShowBottomToolbar = true;
 			Config.bCanShowLockButton = false;
 
 			if( GIsDemoMode )
 			{
-				Config.bShowAssetPathTree = false;
-				Config.bAlwaysShowCollections = true;
-
 				Config.SelectedCollectionName.Name = TEXT( "Demo" );
 				Config.SelectedCollectionName.Type = ECollectionShareType::CST_Shared;
 			}
diff --git a/Engine/Source/Runtime/AssetRegistry/Public/AssetData.h b/Engine/Source/Runtime/AssetRegistry/Public/AssetData.h
--- a/Engine/Source/Runtime/AssetRegistry/Public/AssetData.h
+++ b/Engine/Source/Runtime/AssetRegistry/Public/AssetData.h
@@ -54,8 +54,12 @@ USTRUCT(BlueprintType)
 struct FAssetData
 {
 	GENERATED_BODY()
+
 public:
+	/** The prefix used for collection entries inside TagsAndValues */
+	static constexpr TCHAR CollectionTagPrefix[] = TEXT("CL_");
 
+public:
 	/** The object path for the asset in the form PackageName.AssetName. Only top level objects in a package can have AssetData */
 	UPROPERTY(BlueprintReadOnly, Category=AssetData, transient)
 	FName ObjectPath;
diff --git a/Engine/Source/Runtime/AssetRegistry/Public/AssetRegistryState.h b/Engine/Source/Runtime/AssetRegistry/Public/AssetRegistryState.h
--- a/Engine/Source/Runtime/AssetRegistry/Public/AssetRegistryState.h
+++ b/Engine/Source/Runtime/AssetRegistry/Public/AssetRegistryState.h
@@ -95,7 +95,7 @@ public:
 	 * Does the given path contain assets?
 	 * @note This function doesn't recurse into sub-paths.
 	 */
-	 bool HasAssets(const FName PackagePath) const;
+	bool HasAssets(const FName PackagePath) const;
 
 	/**
 	 * Gets asset data for all assets that match the filter.
@@ -190,7 +190,7 @@ public:
 	/**
 	 * Gets the asset data for the specified asset tag
 	 *
-	 * @param TagName the tag name to search for 
+	 * @param TagName the tag name to search for
 	 * @return An array of AssetData*, empty if nothing found
 	 */
 	const TArray<const FAssetData*>& GetAssetsByTagName(const FName TagName) const
@@ -211,6 +211,12 @@ public:
 		return reinterpret_cast<const TMap<FName, const FAssetData*>&>(CachedAssetsByObjectPath);
 	}
 
+	/** Returns const version of internal Tag->AssetDatas map for fast iteration */
+	const TMap<FName, const TArray<const FAssetData*>> GetTagToAssetDatasMap() const
+	{
+		return reinterpret_cast<const TMap<FName, const TArray<const FAssetData*>>&>(CachedAssetsByTag);
+	}
+
 	/** Returns const version of internal PackageName->PackageData map for fast iteration */
 	const TMap<FName, const FAssetPackageData*>& GetAssetPackageDataMap() const
 	{
diff --git a/Engine/Source/Runtime/Core/Private/Misc/TextFilterExpressionEvaluator.cpp b/Engine/Source/Runtime/Core/Private/Misc/TextFilterExpressionEvaluator.cpp
--- a/Engine/Source/Runtime/Core/Private/Misc/TextFilterExpressionEvaluator.cpp
+++ b/Engine/Source/Runtime/Core/Private/Misc/TextFilterExpressionEvaluator.cpp
@@ -255,6 +255,7 @@ FTextFilterExpressionEvaluator& FTextFilterExpressionEvaluator::operator=(const
 	FilterType = ETextFilterExpressionType::Empty;
 	FilterText = FText::GetEmpty();
 	FilterErrorText = FText::GetEmpty();
+	ExpressionTokens.Reset();
 	CompiledFilter.Reset();
 
 	// If our source object has a different evaluation mode, we need to re-construct the expression parser
@@ -298,6 +299,7 @@ bool FTextFilterExpressionEvaluator::SetFilterText(const FText& InFilterText)
 
 	FilterType = ETextFilterExpressionType::Invalid;
 	FilterText = InFilterText;
+	ExpressionTokens.Reset();
 	CompiledFilterSingleBasicStringTextToken.Reset();
 
 	if (FilterText.IsEmptyOrWhitespace())
@@ -314,9 +316,7 @@ bool FTextFilterExpressionEvaluator::SetFilterText(const FText& InFilterText)
 		if (LexResult.IsValid())
 		{
 			auto& TmpLexTokens = LexResult.GetValue();
-
-			TArray<FExpressionToken> FinalTokens;
-			FinalTokens.Reserve(TmpLexTokens.Num());
+			ExpressionTokens.Reserve(TmpLexTokens.Num());
 
 			bool bIsComplexExpression = false;
 
@@ -373,11 +373,11 @@ bool FTextFilterExpressionEvaluator::SetFilterText(const FText& InFilterText)
 									FunctionParameter.Accumulate(NextToken->Context);
 								}
 								FStringToken EmptyToken;
-								FinalTokens.Add(FExpressionToken(EmptyToken, FSubExpressionStart()));
-								FinalTokens.Add(FExpressionToken(FunctionName, FTextToken(FunctionName.GetString(), ETextFilterTextComparisonMode::Partial, FTextToken::EInvertResult::No)));
-								FinalTokens.Add(FExpressionToken(EmptyToken, FFunction()));
-								FinalTokens.Add(FExpressionToken(FunctionParameter, FTextToken(FunctionParameter.GetString(), ETextFilterTextComparisonMode::Partial, FTextToken::EInvertResult::No)));
-								FinalTokens.Add(FExpressionToken(EmptyToken, FSubExpressionEnd()));
+								ExpressionTokens.Add(FExpressionToken(EmptyToken, FSubExpressionStart()));
+								ExpressionTokens.Add(FExpressionToken(FunctionName, FTextToken(FunctionName.GetString(), ETextFilterTextComparisonMode::Partial, FTextToken::EInvertResult::No)));
+								ExpressionTokens.Add(FExpressionToken(EmptyToken, FFunction()));
+								ExpressionTokens.Add(FExpressionToken(FunctionParameter, FTextToken(FunctionParameter.GetString(), ETextFilterTextComparisonMode::Partial, FTextToken::EInvertResult::No)));
+								ExpressionTokens.Add(FExpressionToken(EmptyToken, FSubExpressionEnd()));
 								continue;
 							}
 						}
@@ -398,7 +398,7 @@ bool FTextFilterExpressionEvaluator::SetFilterText(const FText& InFilterText)
 						if (bIsCurrentTokenValidToInjectANDOperatorBefore)
 						{
 							// Inject the new token before we move the current token into the final list
-							FinalTokens.Add(FExpressionToken(CurrentToken.Context, FAnd()));
+							ExpressionTokens.Add(FExpressionToken(CurrentToken.Context, FAnd()));
 						}
 					}
 
@@ -406,7 +406,7 @@ bool FTextFilterExpressionEvaluator::SetFilterText(const FText& InFilterText)
 					bWasPreviousTokenValidToInjectANDOperatorAfter = (CurrentToken.Node.Cast<FTextToken>() || CurrentToken.Node.Cast<FSubExpressionEnd>());
 
 					// Move this token into the final list since we should have injected the AND operator (if required) by now
-					FinalTokens.Add(MoveTemp(CurrentToken));
+					ExpressionTokens.Add(MoveTemp(CurrentToken));
 				}
 			}
 			else if (TmpLexTokens.Num() > 0)
@@ -416,10 +416,19 @@ bool FTextFilterExpressionEvaluator::SetFilterText(const FText& InFilterText)
 				{
 					CombinedTokenContext.Accumulate(TmpLexTokens[TokenIndex].Context);
 				}
-				FinalTokens.Add(FExpressionToken(CombinedTokenContext, CreateTextTokenFromUnquotedString(FText::TrimPrecedingAndTrailing(FilterText).ToString())));
+				ExpressionTokens.Add(FExpressionToken(CombinedTokenContext, CreateTextTokenFromUnquotedString(FText::TrimPrecedingAndTrailing(FilterText).ToString())));
 			}
 
-			CompiledFilter = ExpressionParser::Compile(MoveTemp(FinalTokens), Grammar);
+			{
+				TArray<FExpressionToken> LocalExpressionTokens;
+				LocalExpressionTokens.Reserve(ExpressionTokens.Num());
+				for (const FExpressionToken& Token : ExpressionTokens)
+				{
+					LocalExpressionTokens.Add(FExpressionToken(Token.Context, Token.Node.Copy()));
+				}
+
+				CompiledFilter = ExpressionParser::Compile(MoveTemp(LocalExpressionTokens), Grammar);
+			}
 
 			auto& CompiledResult = CompiledFilter.GetValue();
 			if (CompiledResult.IsValid())
@@ -460,6 +469,11 @@ FText FTextFilterExpressionEvaluator::GetFilterErrorText() const
 	return FilterErrorText;
 }
 
+const TArray<FExpressionToken>& FTextFilterExpressionEvaluator::GetFilterExpressionTokens() const
+{
+	return ExpressionTokens;
+}
+
 bool FTextFilterExpressionEvaluator::TestTextFilter(const ITextFilterExpressionContext& InContext) const
 {
 	if (FilterType == ETextFilterExpressionType::Empty)
diff --git a/Engine/Source/Runtime/Core/Public/Misc/TextFilterExpressionEvaluator.h b/Engine/Source/Runtime/Core/Public/Misc/TextFilterExpressionEvaluator.h
--- a/Engine/Source/Runtime/Core/Public/Misc/TextFilterExpressionEvaluator.h
+++ b/Engine/Source/Runtime/Core/Public/Misc/TextFilterExpressionEvaluator.h
@@ -184,11 +184,15 @@ public:
 	/** Get the last error returned from lexing or compiling the current filter text */
 	FText GetFilterErrorText() const;
 
+	/** Get the filter expression tokens created from the current filter text */
+	const TArray<FExpressionToken>& GetFilterExpressionTokens() const;
+
 	/** Test our compiled filter using the given context */
 	bool TestTextFilter(const ITextFilterExpressionContext& InContext) const;
 
 	/** Helper function to add callbacks for function tokens */
 	void AddFunctionTokenCallback(FString InFunctionName, FTokenFunctionHandler InCallback);
+
 protected:
 	/** Sets up grammar used for evaluation */
 	void SetupGrammar();
@@ -211,6 +215,9 @@ protected:
 	/** The last error returned from lexing or compiling the current filter text */
 	FText FilterErrorText;
 
+	/** The filter expression tokens created from the current filter text (lexed from FilterText) */
+	TArray<FExpressionToken> ExpressionTokens;
+
 	/** The compiled filter created from the current filter text (compiled from FilterText) */
 	TOptional<ExpressionParser::CompileResultType> CompiledFilter;
 
