Fixed missing include

#rnx

diff --git a/Engine/Source/Editor/ContentBrowser/Private/SourcesSearch.h b/Engine/Source/Editor/ContentBrowser/Private/SourcesSearch.h
--- a/Engine/Source/Editor/ContentBrowser/Private/SourcesSearch.h
+++ b/Engine/Source/Editor/ContentBrowser/Private/SourcesSearch.h
@@ -3,6 +3,7 @@
 #pragma once
 
 #include "CoreMinimal.h"
+#include "Misc/Attribute.h"
 #include "Types/SlateEnums.h"
 #include "Delegates/DelegateCombinations.h"
 
