Prevent saving asset while renaming

diff --git a/Engine/Source/Editor/ContentBrowser/Private/SContentBrowser.cpp b/Engine/Source/Editor/ContentBrowser/Private/SContentBrowser.cpp
--- a/Engine/Source/Editor/ContentBrowser/Private/SContentBrowser.cpp
+++ b/Engine/Source/Editor/ContentBrowser/Private/SContentBrowser.cpp
@@ -2424,7 +2424,7 @@ bool SContentBrowser::HandleRenameCommandCanExecute() const
 bool SContentBrowser::HandleSaveAssetCommandCanExecute() const
 {
 	const TArray<TSharedPtr<FAssetViewItem>>& SelectedItems = AssetViewPtr->GetSelectedItems();
-	if (SelectedItems.Num() > 0)
+	if (SelectedItems.Num() > 0 && !AssetViewPtr->IsRenamingAsset())
 	{
 		return AssetContextMenu->CanExecuteSaveAsset();
 	}
