Crash fix during garbage collection when UObject references a SAssetViewItem()

diff --git a/Engine/Source/Editor/ContentBrowser/Private/AssetViewWidgets.cpp b/Engine/Source/Editor/ContentBrowser/Private/AssetViewWidgets.cpp
--- a/Engine/Source/Editor/ContentBrowser/Private/AssetViewWidgets.cpp
+++ b/Engine/Source/Editor/ContentBrowser/Private/AssetViewWidgets.cpp
@@ -428,8 +428,12 @@ SAssetViewItem::~SAssetViewItem()
 
 	if (!GExitPurge)
 	{
-		// This hack is here to make abnormal shutdowns less frequent.  Crashes here are the result UI's being shut down as a result of GC at edtior shutdown.  This code attemps to call FindObject which is not allowed during GC.
-		SetForceMipLevelsToBeResident(false);
+		const bool bIsGarbageCollectingOnGameThread = (IsInGameThread() && IsGarbageCollecting());
+		if (bIsGarbageCollectingOnGameThread)
+		{
+			// This hack is here to make abnormal shutdowns less frequent.  Crashes here are the result UI's being shut down as a result of GC at edtior shutdown.  This code attemps to call FindObject which is not allowed during GC.
+			SetForceMipLevelsToBeResident(false);
+		}
 	}
 }
 
