PR #6302: Split asset path validation into two separate checks to provide clearer messaging to user on length rules (Contributed by [REDACTED_PERSON])

diff --git a/Engine/Source/Editor/ContentBrowser/Private/ContentBrowserUtils.cpp b/Engine/Source/Editor/ContentBrowser/Private/ContentBrowserUtils.cpp
--- a/Engine/Source/Editor/ContentBrowser/Private/ContentBrowserUtils.cpp
+++ b/Engine/Source/Editor/ContentBrowser/Private/ContentBrowserUtils.cpp
@@ -1736,12 +1736,22 @@ bool ContentBrowserUtils::IsValidObjectPathForCreate(const FString& ObjectPath,
 	// Make sure we are not creating an path that is too long for the OS
 	const FString RelativePathFilename = FPackageName::LongPackageNameToFilename(PackageName, FPackageName::GetAssetPackageExtension());	// full relative path with name + extension
 	const FString FullPath = FPaths::ConvertRelativePathToFull(RelativePathFilename);	// path to file on disk
-	if ( ObjectPath.Len() > (FPlatformMisc::GetMaxPathLength() - MAX_CLASS_NAME_LENGTH) || FullPath.Len() > CVarMaxFullPathLength->GetValueOnGameThread() )
+	if (ObjectPath.Len() > (FPlatformMisc::GetMaxPathLength() - MAX_CLASS_NAME_LENGTH))
+	{
+		// The full path for the asset is too long
+		OutErrorMessage = FText::Format(LOCTEXT("ObjectPathTooLong",
+			"The object path for the asset is too long, the maximum is '{0}'. \nPlease choose a shorter name for the asset or create it in a shallower folder structure."),
+			FText::AsNumber((FPlatformMisc::GetMaxPathLength() - MAX_CLASS_NAME_LENGTH)));
+		// Return false to indicate that the user should enter a new name
+		return false;
+	}
+
+	if (FullPath.Len() > CVarMaxFullPathLength->GetValueOnGameThread())
 	{
 		// The full path for the asset is too long
 		OutErrorMessage = FText::Format( LOCTEXT("AssetPathTooLong", 
-			"The full path for the asset is too deep, the maximum is '{0}'. \nPlease choose a shorter name for the asset or create it in a shallower folder structure."), 
-			FText::AsNumber(FPlatformMisc::GetMaxPathLength()) );
+			"The absolute file path for the asset is too long, the maximum is '{0}'. \nPlease choose a shorter name for the asset or create it in a shallower folder structure."), 
+			FText::AsNumber(CVarMaxFullPathLength->GetValueOnGameThread()));
 		// Return false to indicate that the user should enter a new name
 		return false;
 	}
