Fix IsValidFolderPathForCreate to use path on disk for maximum length check

diff --git a/Engine/Source/Editor/ContentBrowser/Private/ContentBrowserUtils.cpp b/Engine/Source/Editor/ContentBrowser/Private/ContentBrowserUtils.cpp
--- a/Engine/Source/Editor/ContentBrowser/Private/ContentBrowserUtils.cpp
+++ b/Engine/Source/Editor/ContentBrowser/Private/ContentBrowserUtils.cpp
@@ -1789,13 +1789,20 @@ bool ContentBrowserUtils::IsValidFolderPathForCreate(const FString& InFolderPath
 		return false;
 	}
 
+	FString NewPathOnDisk;
+	if (!FPackageName::TryConvertLongPackageNameToFilename(NewFolderPath, NewPathOnDisk))
+	{
+		OutErrorMessage = FText::Format(LOCTEXT("RenameFolderFailedDiskPath", "Folder path could not be converted to disk path: '{0}'"), FText::FromString(NewFolderPath));
+		return false;
+	}
+
 	// Make sure we are not creating a folder path that is too long
-	if (NewFolderPath.Len() > FPlatformMisc::GetMaxPathLength() - MAX_CLASS_NAME_LENGTH)
+	if (NewPathOnDisk.Len() > FPlatformMisc::GetMaxPathLength() - MAX_CLASS_NAME_LENGTH)
 	{
 		// The full path for the folder is too long
 		OutErrorMessage = FText::Format(LOCTEXT("RenameFolderPathTooLong",
 			"The full path for the folder is too deep, the maximum is '{0}'. Please choose a shorter name for the folder or create it in a shallower folder structure."),
-			FText::AsNumber(FPlatformMisc::GetMaxPathLength()));
+			FText::AsNumber(FPlatformMisc::GetMaxPathLength() - MAX_CLASS_NAME_LENGTH));
 		// Return false to indicate that the user should enter a new name for the folder
 		return false;
 	}
