Refactor folder blacklists

diff --git a/Engine/Source/Developer/AssetTools/Private/AssetTools.cpp b/Engine/Source/Developer/AssetTools/Private/AssetTools.cpp
--- a/Engine/Source/Developer/AssetTools/Private/AssetTools.cpp
+++ b/Engine/Source/Developer/AssetTools/Private/AssetTools.cpp
@@ -131,6 +131,8 @@
 #include "AssetToolsSettings.h"
 #include "AssetVtConversion.h"
 #include "Misc/ConfigCacheIni.h"
+#include "Misc/BlacklistNames.h"
+
 #if WITH_EDITOR
 #include "Subsystems/AssetEditorSubsystem.h"
 #endif
@@ -153,6 +155,7 @@ UAssetToolsImpl::UAssetToolsImpl(const FObjectInitializer& ObjectInitializer)
 	, AssetRenameManager(MakeShareable(new FAssetRenameManager))
 	, AssetFixUpRedirectors(MakeShareable(new FAssetFixUpRedirectors))
 	, NextUserCategoryBit(EAssetTypeCategories::FirstUser)
+	, FolderBlacklist(MakeShared<FBlacklistPaths>())
 {
 	TArray<FString> SupportedTypesArray;
 	GConfig->GetArray(TEXT("AssetTools"), TEXT("SupportedAssetTypes"), SupportedTypesArray, GEditorIni);
@@ -3143,4 +3146,9 @@ TArray<UFactory*> UAssetToolsImpl::GetNewAssetFactories() const
 	return MoveTemp(Factories);
 }
 
+TSharedRef<FBlacklistPaths>& UAssetToolsImpl::GetFolderBlacklist()
+{
+	return FolderBlacklist;
+}
+
 #undef LOCTEXT_NAMESPACE
diff --git a/Engine/Source/Developer/AssetTools/Private/AssetTools.h b/Engine/Source/Developer/AssetTools/Private/AssetTools.h
--- a/Engine/Source/Developer/AssetTools/Private/AssetTools.h
+++ b/Engine/Source/Developer/AssetTools/Private/AssetTools.h
@@ -7,6 +7,7 @@
 #include "IAssetTypeActions.h"
 #include "AssetData.h"
 #include "AssetRenameManager.h"
+#include "Misc/BlacklistNames.h"
 #include "AssetTools.generated.h"
 
 class FAssetFixUpRedirectors;
@@ -116,6 +117,8 @@ public:
 	virtual void ConvertVirtualTextures(const TArray<UTexture2D*>& Textures, bool bConvertBackToNonVirtual, const TArray<UMaterial*>* RelatedMaterials = nullptr) const override;
 	virtual bool IsAssetClassSupported(const UClass* AssetClass) const override;
 	virtual TArray<UFactory*> GetNewAssetFactories() const override;
+	virtual TSharedRef<FBlacklistPaths>& GetFolderBlacklist() override;
+
 public:
 	/** Gets the asset tools singleton as a FAssetTools for asset tools module use */
 	static UAssetToolsImpl& Get();
@@ -182,6 +185,9 @@ private:
 	
 	/** The next user category bit to allocate (set to 0 when there are no more bits left) */
 	uint32 NextUserCategoryBit;
+
+	/** Blacklist of folder paths */
+	TSharedRef<FBlacklistPaths> FolderBlacklist;
 };
 
 PRAGMA_ENABLE_DEPRECATION_WARNINGS
diff --git a/Engine/Source/Developer/AssetTools/Public/IAssetTools.h b/Engine/Source/Developer/AssetTools/Public/IAssetTools.h
--- a/Engine/Source/Developer/AssetTools/Public/IAssetTools.h
+++ b/Engine/Source/Developer/AssetTools/Public/IAssetTools.h
@@ -22,6 +22,7 @@ class IClassTypeActions;
 class UFactory;
 class UAssetImportTask;
 class UAdvancedCopyCustomization;
+class FBlacklistPaths;
 
 USTRUCT(BlueprintType)
 struct FAssetRenameData
@@ -410,6 +411,9 @@ public:
 
 	/** Find all supported asset factories. */
 	virtual TArray<UFactory*> GetNewAssetFactories() const = 0;
+
+	/** Get folder blacklist for content browser and other systems */
+	virtual TSharedRef<FBlacklistPaths>& GetFolderBlacklist() = 0;
 };
 
 UCLASS(transient)
diff --git a/Engine/Source/Editor/ContentBrowser/Private/ContentBrowserSingleton.cpp b/Engine/Source/Editor/ContentBrowser/Private/ContentBrowserSingleton.cpp
--- a/Engine/Source/Editor/ContentBrowser/Private/ContentBrowserSingleton.cpp
+++ b/Engine/Source/Editor/ContentBrowser/Private/ContentBrowserSingleton.cpp
@@ -721,20 +721,6 @@ void FContentBrowserSingleton::ForceShowPluginContent(bool bEnginePlugin)
 	}
 }
 
-bool FContentBrowserSingleton::PathViewPathPassesFilter(const FString& InPath) const
-{
-	static FContentBrowserModule& ContentBrowserModule = FModuleManager::LoadModuleChecked<FContentBrowserModule>("ContentBrowser");
-	for (const auto& It : ContentBrowserModule.GetDirectoryPathFilteredDelegates())
-	{
-		if (It.IsBound() && It.Execute(InPath))
-		{
-			return false;
-		}
-	}
-
-	return true;
-}
-
 void FContentBrowserSingleton::PopulateConfigValues()
 {
 	const FString ContentBrowserSection = TEXT("ContentBrowser");
diff --git a/Engine/Source/Editor/ContentBrowser/Private/ContentBrowserSingleton.h b/Engine/Source/Editor/ContentBrowser/Private/ContentBrowserSingleton.h
--- a/Engine/Source/Editor/ContentBrowser/Private/ContentBrowserSingleton.h
+++ b/Engine/Source/Editor/ContentBrowser/Private/ContentBrowserSingleton.h
@@ -86,9 +86,6 @@ public:
 	/** Gets the settings for the plugin with the specified name */
 	const FContentBrowserPluginSettings& GetPluginSettings(FName PluginName) const;
 
-	/** Returns true if should show directory path */
-	bool PathViewPathPassesFilter(const FString& InPath) const;
-
 	/** Single storage location for content browser favorites */
 	TArray<FString> FavoriteFolderPaths;
 
diff --git a/Engine/Source/Editor/ContentBrowser/Private/SAssetView.cpp b/Engine/Source/Editor/ContentBrowser/Private/SAssetView.cpp
--- a/Engine/Source/Editor/ContentBrowser/Private/SAssetView.cpp
+++ b/Engine/Source/Editor/ContentBrowser/Private/SAssetView.cpp
@@ -48,6 +48,7 @@
 #include "DesktopPlatformModule.h"
 #include "Misc/FileHelper.h"
 #include "Misc/TextFilterUtils.h"
+#include "Misc/BlacklistNames.h"
 #include "AssetRegistryState.h"
 #include "Materials/Material.h"
 
@@ -118,22 +119,14 @@ namespace AssetViewUtils
 	class FInitialAssetFilter
 	{
 	public:
-		FInitialAssetFilter(bool InDisplayL10N, bool InDisplayEngine, bool InDisplayPlugins) :
+		FInitialAssetFilter(bool InDisplayL10N, bool InDisplayEngine, bool InDisplayPlugins, const FBlacklistPaths* InFolderBlacklist) :
 			bDisplayL10N(InDisplayL10N),
 			bDisplayEngine(InDisplayEngine),
-			bDisplayPlugins(InDisplayPlugins)
-		{
-			Init(InDisplayL10N, InDisplayEngine, InDisplayPlugins);
-		}
-
-		void Init(bool InDisplayL10N, bool InDisplayEngine, bool InDisplayPlugins)
+			bDisplayPlugins(InDisplayPlugins),
+			FolderBlacklist(InFolderBlacklist)
 		{
 			ObjectRedirectorClassName = UObjectRedirector::StaticClass()->GetFName();
 
-			bDisplayL10N = InDisplayL10N;
-			bDisplayEngine = InDisplayEngine;
-			bDisplayPlugins = InDisplayPlugins;
-
 			Plugins = IPluginManager::Get().GetEnabledPluginsWithContent();
 			PluginNamesUpperWide.Reset(Plugins.Num());
 			PluginNamesUpperAnsi.Reset(Plugins.Num());
@@ -172,12 +165,25 @@ namespace AssetViewUtils
 			TextFilterUtils::FNameBufferWithNumber ObjectPathBuffer(PackagePath);
 			if (ObjectPathBuffer.IsWide())
 			{
-				return PassesPackagePathFilter(ObjectPathBuffer.GetWideNamePtr());
+				if (!PassesPackagePathFilter(ObjectPathBuffer.GetWideNamePtr()))
+				{
+					return false;
+				}
 			}
 			else
 			{
-				return PassesPackagePathFilter(ObjectPathBuffer.GetAnsiNamePtr());
+				if (!PassesPackagePathFilter(ObjectPathBuffer.GetAnsiNamePtr()))
+				{
+					return false;
+				}
+			}
+			
+			if (FolderBlacklist && !FolderBlacklist->PassesStartsWithFilter(PackagePath))
+			{
+				return false;
 			}
+
+			return true;
 		}
 
 		template <typename CharType>
@@ -242,6 +248,8 @@ namespace AssetViewUtils
 			return true;
 		}
 
+		
+
 	private:
 		FName ObjectRedirectorClassName;
 		bool bDisplayL10N;
@@ -251,6 +259,7 @@ namespace AssetViewUtils
 		TArray<FString> PluginNamesUpperWide;
 		TArray<bool> PluginLoadedFromEngine;
 		TArray<TSharedRef<IPlugin>> Plugins;
+		const FBlacklistPaths* FolderBlacklist;
 
 		FORCEINLINE int32 FindPluginNameUpper(const WIDECHAR* PluginNameUpper, int32 Length) const
 		{
@@ -532,7 +541,8 @@ void SAssetView::Construct( const FArguments& InArgs )
 
 	NumVisibleColumns = 0;
 
-	SupportedFilter = AssetViewUtils::CreateSupportedFilter();
+	FAssetToolsModule& AssetToolsModule = FModuleManager::LoadModuleChecked<FAssetToolsModule>("AssetTools");
+	FolderBlacklist = AssetToolsModule.Get().GetFolderBlacklist();
 
 	FEditorWidgetsModule& EditorWidgetsModule = FModuleManager::LoadModuleChecked<FEditorWidgetsModule>("EditorWidgets");
 	TSharedRef<SWidget> AssetDiscoveryIndicator = EditorWidgetsModule.CreateAssetDiscoveryIndicator(EAssetDiscoveryIndicatorScaleMode::Scale_Vertical);
@@ -1984,11 +1994,12 @@ void SAssetView::RefreshSourceItems()
 
 	TArray<FAssetData>& Items = OnShouldFilterAsset.IsBound() ? QueriedAssetItems : AssetItems;
 
-	const bool bShowAll = SourcesData.IsEmpty() && SupportedFilter.IsEmpty() && BackendFilter.IsEmpty();
+	const FBlacklistPaths* FolderBlacklistToUse = (FolderBlacklist.IsValid() && FolderBlacklist->HasFiltering()) ? FolderBlacklist.Get() : nullptr;
+	const bool bShowAll = SourcesData.IsEmpty() && BackendFilter.IsEmpty() && FolderBlacklistToUse;
 
 	bool bShowClasses = false;
 	TArray<FName> ClassPathsToShow;
-	AssetViewUtils::FInitialAssetFilter InitialAssetFilter(IsShowingLocalizedContent(), IsShowingEngineContent(), IsShowingPluginContent());
+	AssetViewUtils::FInitialAssetFilter InitialAssetFilter(IsShowingLocalizedContent(), IsShowingEngineContent(), IsShowingPluginContent(), FolderBlacklistToUse);
 
 	if ( bShowAll )
 	{
@@ -2347,6 +2358,11 @@ void SAssetView::RefreshFolders()
 					continue;
 				}
 
+				if (FolderBlacklist.IsValid() && !FolderBlacklist->PassesStartsWithFilter(SubPath))
+				{
+					continue;
+				}
+
 				if(!Folders.Contains(SubPath))
 				{
 					FoldersToAdd.Add(SubPath);
@@ -2743,7 +2759,8 @@ void SAssetView::OnAssetRegistryPathAdded(const FString& Path)
 		const bool bDisplayL10N = IsShowingLocalizedContent();
 		if ((bDisplayEmpty || EmptyFolderVisibilityManager->ShouldShowPath(Path)) &&  
 			(bDisplayDev || !ContentBrowserUtils::IsDevelopersFolder(Path)) && 
-			(bDisplayL10N || !ContentBrowserUtils::IsLocalizationFolder(Path))
+			(bDisplayL10N || !ContentBrowserUtils::IsLocalizationFolder(Path)) &&
+			(!FolderBlacklist.IsValid() || FolderBlacklist->PassesStartsWithFilter(Path))
 			)
 		{
 			for (const FName& SourcePathName : SourcesData.PackagePaths)
@@ -2973,6 +2990,32 @@ bool SAssetView::PassesCurrentFrontendFilter(const FAssetData& Item) const
 	return true;
 }
 
+void SAssetView::RunAssetsThroughBlacklists(TArray<FAssetData>& InOutAssetDataList) const
+{
+	if (InOutAssetDataList.Num() == 0)
+	{
+		return;
+	}
+
+	if (FolderBlacklist.IsValid() && FolderBlacklist->HasFiltering())
+	{
+		if (FolderBlacklist->IsBlacklistAll())
+		{
+			InOutAssetDataList.Reset();
+			return;
+		}
+
+		InOutAssetDataList.RemoveAll([this](const FAssetData& AssetData)
+		{
+			if (FolderBlacklist.IsValid() && !FolderBlacklist->PassesStartsWithFilter(AssetData.PackagePath))
+			{
+				return true;
+			}
+			return false;
+		});
+	}
+}
+
 void SAssetView::RunAssetsThroughBackendFilter(TArray<FAssetData>& InOutAssetDataList) const
 {
 	const bool bRecurse = ShouldFilterRecursively();
@@ -2992,6 +3035,8 @@ void SAssetView::RunAssetsThroughBackendFilter(TArray<FAssetData>& InOutAssetDat
 		FAssetRegistryModule& AssetRegistryModule = FModuleManager::LoadModuleChecked<FAssetRegistryModule>(TEXT("AssetRegistry"));
 		AssetRegistryModule.Get().RunAssetsThroughFilter(InOutAssetDataList, Filter);
 
+		RunAssetsThroughBlacklists(InOutAssetDataList);
+
 		if (SourcesData.HasCollections() && !bIsDynamicCollection)
 		{
 			// Include objects from child collections if we're recursing
diff --git a/Engine/Source/Editor/ContentBrowser/Private/SAssetView.h b/Engine/Source/Editor/ContentBrowser/Private/SAssetView.h
--- a/Engine/Source/Editor/ContentBrowser/Private/SAssetView.h
+++ b/Engine/Source/Editor/ContentBrowser/Private/SAssetView.h
@@ -24,6 +24,7 @@
 #include "HistoryManager.h"
 
 class FMenuBuilder;
+class FBlacklistPaths;
 class FWeakWidgetPath;
 class FWidgetPath;
 class SAssetColumnView;
@@ -442,6 +443,9 @@ private:
 	/** Removes asset data items from the given array that don't pass all applied backend (asset registry) filters */
 	void RunAssetsThroughBackendFilter(TArray<FAssetData>& InOutAssetDataList) const;
 
+	/** Removes asset data items from the given array that don't pass all blacklist filters */
+	void RunAssetsThroughBlacklists(TArray<FAssetData>& InOutAssetDataList) const;
+
 	/** Returns true if the current filters deem that the asset view should be filtered recursively (overriding folder view) */
 	bool ShouldFilterRecursively() const;
 
@@ -843,6 +847,7 @@ private:
 	FSourcesData SourcesData;
 	FARFilter BackendFilter;
 	FARFilter SupportedFilter;
+	TSharedPtr<FBlacklistPaths> FolderBlacklist;
 	TSharedPtr<FAssetFilterCollectionType> FrontendFilters;
 
 	/** If true, the source items will be refreshed next frame. Very slow. */
diff --git a/Engine/Source/Editor/ContentBrowser/Private/SPathView.cpp b/Engine/Source/Editor/ContentBrowser/Private/SPathView.cpp
--- a/Engine/Source/Editor/ContentBrowser/Private/SPathView.cpp
+++ b/Engine/Source/Editor/ContentBrowser/Private/SPathView.cpp
@@ -25,6 +25,7 @@
 #include "NativeClassHierarchy.h"
 #include "EmptyFolderVisibilityManager.h"
 #include "ContentBrowserModule.h"
+#include "Misc/BlacklistNames.h"
 
 #include "Application/SlateApplicationBase.h"
 
@@ -75,7 +76,8 @@ void SPathView::Construct( const FArguments& InArgs )
 		RegisterActiveTimer( 0.f, FWidgetActiveTimerDelegate::CreateSP( this, &SPathView::SetFocusPostConstruct ) );
 	}
 
-	ContentBrowserSingleton = &FContentBrowserSingleton::Get();
+	FAssetToolsModule& AssetToolsModule = FModuleManager::LoadModuleChecked<FAssetToolsModule>("AssetTools");
+	FolderBlacklist = AssetToolsModule.Get().GetFolderBlacklist();
 
 	// Listen for when view settings are changed
 	UContentBrowserSettings::OnSettingChanged().AddSP(this, &SPathView::HandleSettingChanged);
@@ -338,7 +340,7 @@ TSharedPtr<FTreeItem> SPathView::AddPath(const FString& Path, bool bUserNamed)
 		return TSharedPtr<FTreeItem>();
 	}
 
-	if (!ContentBrowserSingleton->PathViewPathPassesFilter(Path))
+	if (FolderBlacklist.IsValid() && !FolderBlacklist->PassesStartsWithFilter(Path))
 	{
 		return TSharedPtr<FTreeItem>();
 	}
@@ -1863,7 +1865,7 @@ TSharedPtr<FTreeItem> SFavoritePathView::AddPath(const FString& Path, bool bUser
 		return TSharedPtr<FTreeItem>();
 	}
 
-	if (!ContentBrowserSingleton->PathViewPathPassesFilter(Path))
+	if (FolderBlacklist.IsValid() && !FolderBlacklist->PassesStartsWithFilter(Path))
 	{
 		return TSharedPtr<FTreeItem>();
 	}
diff --git a/Engine/Source/Editor/ContentBrowser/Private/SPathView.h b/Engine/Source/Editor/ContentBrowser/Private/SPathView.h
--- a/Engine/Source/Editor/ContentBrowser/Private/SPathView.h
+++ b/Engine/Source/Editor/ContentBrowser/Private/SPathView.h
@@ -20,6 +20,7 @@ class FSourcesSearch;
 struct FHistoryData;
 struct FTreeItem;
 class FContentBrowserSingleton;
+class FBlacklistPaths;
 
 typedef TTextFilter< const FString& > FolderTextFilter;
 
@@ -327,8 +328,8 @@ protected:
 
 	TSharedPtr<SWidget> PathViewWidget;
 
-	/** Cached pointer to content browser singleton */
-	FContentBrowserSingleton* ContentBrowserSingleton;
+	/** Blacklist filter to hide folders */
+	TSharedPtr<FBlacklistPaths> FolderBlacklist;
 
 private:
 
diff --git a/Engine/Source/Editor/ContentBrowser/Public/ContentBrowserModule.h b/Engine/Source/Editor/ContentBrowser/Public/ContentBrowserModule.h
--- a/Engine/Source/Editor/ContentBrowser/Public/ContentBrowserModule.h
+++ b/Engine/Source/Editor/ContentBrowser/Public/ContentBrowserModule.h
@@ -53,8 +53,6 @@ public:
 	DECLARE_MULTICAST_DELEGATE_OneParam( FOnSourcesViewChanged, bool /*bExpanded*/ );
 	/** */
 	DECLARE_MULTICAST_DELEGATE_OneParam( FOnAssetPathChanged, const FString& /*NewPath*/ );
-	/** */
-	DECLARE_DELEGATE_RetVal_OneParam(bool, FIsDirectoryPathFiltered, const FString&);
 
 	/**
 	 * Called right after the plugin DLL has been loaded and the plugin object has been created
@@ -99,9 +97,6 @@ public:
 	/** Delegates to be called to extend the drag-and-drop support of the asset view */
 	virtual TArray<FAssetViewDragAndDropExtender>& GetAssetViewDragAndDropExtenders() { return AssetViewDragAndDropExtenders; }
 
-	/** Delegates to call to hide folders displayed in content browser */
-	virtual TArray<FIsDirectoryPathFiltered>& GetDirectoryPathFilteredDelegates() { return DirectoryPathFilteredDelegates; }
-
 	/** Delegate accessors */
 	FOnFilterChanged& GetOnFilterChanged() { return OnFilterChanged; } 
 	FOnSearchBoxChanged& GetOnSearchBoxChanged() { return OnSearchBoxChanged; } 
@@ -133,9 +128,6 @@ private:
 	TArray<FContentBrowserMenuExtender> AssetViewViewMenuExtenders;
 	TArray<FContentBrowserCommandExtender> ContentBrowserCommandExtenders;
 
-	/** All delegates that hide folders displayed in content browser */
-	TArray<FIsDirectoryPathFiltered> DirectoryPathFilteredDelegates;
-
 	/** All delegates generating extra state indicators */
 	TArray<FAssetViewExtraStateGenerator> AssetViewExtraStateGenerators;
 
diff --git a/Engine/Source/Runtime/Slate/Private/Framework/Docking/TabManager.cpp b/Engine/Source/Runtime/Slate/Private/Framework/Docking/TabManager.cpp
--- a/Engine/Source/Runtime/Slate/Private/Framework/Docking/TabManager.cpp
+++ b/Engine/Source/Runtime/Slate/Private/Framework/Docking/TabManager.cpp
@@ -18,6 +18,7 @@
 #include "Framework/Docking/SDockingTabStack.h"
 #include "Framework/Docking/SDockingTabWell.h"
 #include "Framework/Docking/LayoutExtender.h"
+#include "Misc/BlacklistNames.h"
 #include "HAL/PlatformApplicationMisc.h"
 #if PLATFORM_MAC
 #include "../MultiBox/Mac/MacMenu.h"
@@ -1021,7 +1022,7 @@ void FTabManager::PopulateTabSpawnerMenu( FMenuBuilder& PopulateMe, TSharedRef<F
 		const TSharedRef<FTabSpawnerEntry>& SpawnerEntry = SpawnerIterator.Value();
 		if ( SpawnerEntry->bAutoGenerateMenuEntry )
 		{
-			if (SpawnerEntry->TabType == NAME_None || TabBlacklist.PassesFilter(SpawnerEntry->TabType))
+			if (SpawnerEntry->TabType == NAME_None || TabBlacklist->PassesFilter(SpawnerEntry->TabType))
 			{
 				AllSpawners->AddUnique(SpawnerEntry);
 			}
@@ -1034,7 +1035,7 @@ void FTabManager::PopulateTabSpawnerMenu( FMenuBuilder& PopulateMe, TSharedRef<F
 		const TSharedRef<FTabSpawnerEntry>& SpawnerEntry = SpawnerIterator.Value();
 		if ( SpawnerEntry->bAutoGenerateMenuEntry )
 		{
-			if (SpawnerEntry->TabType == NAME_None || TabBlacklist.PassesFilter(SpawnerEntry->TabType))
+			if (SpawnerEntry->TabType == NAME_None || TabBlacklist->PassesFilter(SpawnerEntry->TabType))
 			{
 				AllSpawners->AddUnique(SpawnerEntry);
 			}
@@ -1335,6 +1336,7 @@ FTabManager::FTabManager( const TSharedPtr<SDockTab>& InOwnerTab, const TSharedR
 , LastDocumentUID( 0 )
 , bIsSavingVisualState( false )
 , bCanDoDragOperation( true )
+, TabBlacklist( MakeShareable(new FBlacklistNames()) )
 {
 	LocalWorkspaceMenuRoot = FWorkspaceItem::NewGroup(LOCTEXT("LocalWorkspaceRoot", "Local Workspace Root"));
 }
@@ -1624,7 +1626,7 @@ bool FTabManager::HasTabSpawner(FName TabId) const
 	return Spawner != nullptr;
 }
 
-FBlacklistNames& FTabManager::GetTabBlacklist()
+TSharedRef<FBlacklistNames>& FTabManager::GetTabBlacklist()
 {
 	return TabBlacklist;
 }
diff --git a/Engine/Source/Runtime/Slate/Public/Framework/Docking/TabManager.h b/Engine/Source/Runtime/Slate/Public/Framework/Docking/TabManager.h
--- a/Engine/Source/Runtime/Slate/Public/Framework/Docking/TabManager.h
+++ b/Engine/Source/Runtime/Slate/Public/Framework/Docking/TabManager.h
@@ -5,13 +5,13 @@
 #include "CoreMinimal.h"
 #include "SlateFwd.h"
 #include "Misc/Attribute.h"
-#include "Misc/BlacklistNames.h"
 #include "Textures/SlateIcon.h"
 #include "Widgets/SWindow.h"
 #include "Framework/Docking/WorkspaceItem.h"
 
 class FJsonObject;
 class FMenuBuilder;
+class FBlacklistNames;
 class FMultiBox;
 class FProxyTabmanager;
 class SDockingArea;
@@ -803,7 +803,7 @@ class SLATE_API FTabManager : public TSharedFromThis<FTabManager>
 		TSharedPtr<SDockTab> GetOwnerTab() { return OwnerTabPtr.Pin(); }
 
 		/** Returns filter for additional control over available tabs */
-		FBlacklistNames& GetTabBlacklist();
+		TSharedRef<FBlacklistNames>& GetTabBlacklist();
 
 	protected:
 		void InvokeTabForMenu( FName TabId );
@@ -972,7 +972,7 @@ class SLATE_API FTabManager : public TSharedFromThis<FTabManager>
 		bool bCanDoDragOperation;
 
 		/** Allow systems to dynamically hide tabs */
-		FBlacklistNames TabBlacklist;
+		TSharedRef<FBlacklistNames> TabBlacklist;
 };
 
 
