Fix folder blacklist flipped logic
#rnx

diff --git a/Engine/Source/Editor/ContentBrowser/Private/SAssetView.cpp b/Engine/Source/Editor/ContentBrowser/Private/SAssetView.cpp
--- a/Engine/Source/Editor/ContentBrowser/Private/SAssetView.cpp
+++ b/Engine/Source/Editor/ContentBrowser/Private/SAssetView.cpp
@@ -1995,7 +1995,7 @@ void SAssetView::RefreshSourceItems()
 	TArray<FAssetData>& Items = OnShouldFilterAsset.IsBound() ? QueriedAssetItems : AssetItems;
 
 	const FBlacklistPaths* FolderBlacklistToUse = (FolderBlacklist.IsValid() && FolderBlacklist->HasFiltering()) ? FolderBlacklist.Get() : nullptr;
-	const bool bShowAll = SourcesData.IsEmpty() && BackendFilter.IsEmpty() && FolderBlacklistToUse;
+	const bool bShowAll = SourcesData.IsEmpty() && BackendFilter.IsEmpty() && !FolderBlacklistToUse;
 
 	bool bShowClasses = false;
 	TArray<FName> ClassPathsToShow;
