Content browser class blacklist refactor to include subclasses

diff --git a/Engine/Source/Developer/AssetTools/Private/AssetTools.cpp b/Engine/Source/Developer/AssetTools/Private/AssetTools.cpp
--- a/Engine/Source/Developer/AssetTools/Private/AssetTools.cpp
+++ b/Engine/Source/Developer/AssetTools/Private/AssetTools.cpp
@@ -155,6 +155,7 @@ UAssetToolsImpl::UAssetToolsImpl(const FObjectInitializer& ObjectInitializer)
 	, AssetRenameManager(MakeShareable(new FAssetRenameManager))
 	, AssetFixUpRedirectors(MakeShareable(new FAssetFixUpRedirectors))
 	, NextUserCategoryBit(EAssetTypeCategories::FirstUser)
+	, AssetClassBlacklist(MakeShared<FBlacklistNames>())
 	, FolderBlacklist(MakeShared<FBlacklistPaths>())
 {
 	TArray<FString> SupportedTypesArray;
@@ -162,9 +163,11 @@ UAssetToolsImpl::UAssetToolsImpl(const FObjectInitializer& ObjectInitializer)
 
 	for (const FString& Type : SupportedTypesArray)
 	{
-		SupportedAssetTypes.Add(*Type);
+		AssetClassBlacklist->AddWhitelistItem("AssetToolsConfigFile", *Type);
 	}
 
+	AssetClassBlacklist->OnFilterChanged().AddUObject(this, &UAssetToolsImpl::AssetClassBlacklistChanged);
+
 	// Register the built-in advanced categories
 	AllocatedCategoryBits.Add(TEXT("_BuiltIn_0"), FAdvancedAssetCategory(EAssetTypeCategories::Animation, LOCTEXT("AnimationAssetCategory", "Animation")));
 	AllocatedCategoryBits.Add(TEXT("_BuiltIn_1"), FAdvancedAssetCategory(EAssetTypeCategories::Blueprint, LOCTEXT("BlueprintAssetCategory", "Blueprints")));
@@ -264,22 +267,8 @@ UAssetToolsImpl::UAssetToolsImpl(const FObjectInitializer& ObjectInitializer)
 
 void UAssetToolsImpl::RegisterAssetTypeActions(const TSharedRef<IAssetTypeActions>& NewActions)
 {
-	if (SupportedAssetTypes.Num() > 0)
-	{
-		const UClass* SupportedClass = NewActions->GetSupportedClass();
-		if (SupportedClass != nullptr)
-		{
-			const FName ClassName = SupportedClass->GetFName();
-			if (!SupportedAssetTypes.Contains(ClassName))
-			{
-				NewActions->SetSupported(false);
-			}
-		}
-		else
-		{
-			NewActions->SetSupported(false);
-		}
-	}
+	const UClass* SupportedClass = NewActions->GetSupportedClass();
+	NewActions->SetSupported(SupportedClass && AssetClassBlacklist->PassesFilter(SupportedClass->GetFName()));
 
 	AssetTypeActionsList.Add(NewActions);
 }
@@ -3146,6 +3135,20 @@ TArray<UFactory*> UAssetToolsImpl::GetNewAssetFactories() const
 	return MoveTemp(Factories);
 }
 
+TSharedRef<FBlacklistNames>& UAssetToolsImpl::GetAssetClassBlacklist()
+{
+	return AssetClassBlacklist;
+}
+
+void UAssetToolsImpl::AssetClassBlacklistChanged()
+{
+	for (TSharedRef<IAssetTypeActions>& ActionsIt : AssetTypeActionsList)
+	{
+		const UClass* SupportedClass = ActionsIt->GetSupportedClass();
+		ActionsIt->SetSupported(SupportedClass && AssetClassBlacklist->PassesFilter(SupportedClass->GetFName()));
+	}
+}
+
 TSharedRef<FBlacklistPaths>& UAssetToolsImpl::GetFolderBlacklist()
 {
 	return FolderBlacklist;
diff --git a/Engine/Source/Developer/AssetTools/Private/AssetTools.h b/Engine/Source/Developer/AssetTools/Private/AssetTools.h
--- a/Engine/Source/Developer/AssetTools/Private/AssetTools.h
+++ b/Engine/Source/Developer/AssetTools/Private/AssetTools.h
@@ -117,6 +117,7 @@ public:
 	virtual void ConvertVirtualTextures(const TArray<UTexture2D*>& Textures, bool bConvertBackToNonVirtual, const TArray<UMaterial*>* RelatedMaterials = nullptr) const override;
 	virtual bool IsAssetClassSupported(const UClass* AssetClass) const override;
 	virtual TArray<UFactory*> GetNewAssetFactories() const override;
+	virtual TSharedRef<FBlacklistNames>& GetAssetClassBlacklist() override;
 	virtual TSharedRef<FBlacklistPaths>& GetFolderBlacklist() override;
 
 public:
@@ -171,6 +172,9 @@ private:
 
 	UObject* PerformDuplicateAsset(const FString& AssetName, const FString& PackagePath, UObject* OriginalObject, bool bWithDialog);
 
+	/** Internal method that performs actions when asset class blacklist filter changes */
+	void AssetClassBlacklistChanged();
+
 private:
 	/** The list of all registered AssetTypeActions */
 	TArray<TSharedRef<IAssetTypeActions>> AssetTypeActionsList;
@@ -181,11 +185,12 @@ private:
 	/** The categories that have been allocated already */
 	TMap<FName, FAdvancedAssetCategory> AllocatedCategoryBits;
 	
-	TSet<FName> SupportedAssetTypes;
-	
 	/** The next user category bit to allocate (set to 0 when there are no more bits left) */
 	uint32 NextUserCategoryBit;
 
+	/** Blacklist of assets by class name */
+	TSharedRef<FBlacklistNames> AssetClassBlacklist;
+
 	/** Blacklist of folder paths */
 	TSharedRef<FBlacklistPaths> FolderBlacklist;
 };
diff --git a/Engine/Source/Developer/AssetTools/Public/IAssetTools.h b/Engine/Source/Developer/AssetTools/Public/IAssetTools.h
--- a/Engine/Source/Developer/AssetTools/Public/IAssetTools.h
+++ b/Engine/Source/Developer/AssetTools/Public/IAssetTools.h
@@ -22,6 +22,7 @@ class IClassTypeActions;
 class UFactory;
 class UAssetImportTask;
 class UAdvancedCopyCustomization;
+class FBlacklistNames;
 class FBlacklistPaths;
 
 USTRUCT(BlueprintType)
@@ -412,6 +413,9 @@ public:
 	/** Find all supported asset factories. */
 	virtual TArray<UFactory*> GetNewAssetFactories() const = 0;
 
+	/** Get asset class blacklist for content browser and other systems */
+	virtual TSharedRef<FBlacklistNames>& GetAssetClassBlacklist() = 0;
+
 	/** Get folder blacklist for content browser and other systems */
 	virtual TSharedRef<FBlacklistPaths>& GetFolderBlacklist() = 0;
 };
diff --git a/Engine/Source/Editor/ContentBrowser/Private/SAssetView.cpp b/Engine/Source/Editor/ContentBrowser/Private/SAssetView.cpp
--- a/Engine/Source/Editor/ContentBrowser/Private/SAssetView.cpp
+++ b/Engine/Source/Editor/ContentBrowser/Private/SAssetView.cpp
@@ -116,14 +116,46 @@ namespace AssetViewUtils
 		return false;
 	}
 
+	/** Expands blacklist of class names to include subclasses */
+	void ExpandClassesBlacklist(const FBlacklistNames& InBlacklist, FBlacklistNames& ExpandedBlacklist)
+	{
+		FAssetRegistryModule& AssetRegistryModule = FModuleManager::LoadModuleChecked<FAssetRegistryModule>("AssetRegistry");
+		if (InBlacklist.GetBlacklist().Num() > 0)
+		{
+			TArray<FName> BlacklistKeys;
+			InBlacklist.GetBlacklist().GetKeys(BlacklistKeys);
+			const TSet<FName> EmptySet;
+			TSet<FName> BlacklistSubClassNames;
+			AssetRegistryModule.Get().GetDerivedClassNames(BlacklistKeys, EmptySet, BlacklistSubClassNames);
+			for (const FName ClassName : BlacklistSubClassNames)
+			{
+				ExpandedBlacklist.AddBlacklistItem(NAME_None, ClassName);
+			}
+		}
+
+		if (InBlacklist.GetWhitelist().Num() > 0)
+		{
+			TArray<FName> WhitelistKeys;
+			InBlacklist.GetWhitelist().GetKeys(WhitelistKeys);
+			const TSet<FName> EmptySet;
+			TSet<FName> WhitelistSubClassNames;
+			AssetRegistryModule.Get().GetDerivedClassNames(WhitelistKeys, EmptySet, WhitelistSubClassNames);
+			for (const FName ClassName : WhitelistSubClassNames)
+			{
+				ExpandedBlacklist.AddWhitelistItem(NAME_None, ClassName);
+			}
+		}
+	}
+
 	class FInitialAssetFilter
 	{
 	public:
-		FInitialAssetFilter(bool InDisplayL10N, bool InDisplayEngine, bool InDisplayPlugins, const FBlacklistPaths* InFolderBlacklist) :
+		FInitialAssetFilter(bool InDisplayL10N, bool InDisplayEngine, bool InDisplayPlugins, const FBlacklistPaths* InFolderBlacklist, const FBlacklistNames* InClassBlacklist) :
 			bDisplayL10N(InDisplayL10N),
 			bDisplayEngine(InDisplayEngine),
 			bDisplayPlugins(InDisplayPlugins),
-			FolderBlacklist(InFolderBlacklist)
+			FolderBlacklist(InFolderBlacklist),
+			AssetClassBlacklist(nullptr)
 		{
 			ObjectRedirectorClassName = UObjectRedirector::StaticClass()->GetFName();
 
@@ -137,6 +169,12 @@ namespace AssetViewUtils
 				TextFilterUtils::TryConvertWideToAnsi(PluginNameUpperWide, PluginNamesUpperAnsi.AddDefaulted_GetRef());
 				PluginLoadedFromEngine.Add(PluginRef->GetLoadedFrom() == EPluginLoadedFrom::Engine);
 			}
+
+			if (InClassBlacklist && InClassBlacklist->HasFiltering())
+			{
+				ExpandClassesBlacklist(*InClassBlacklist, ExpandedClassesBlacklist);
+				AssetClassBlacklist = &ExpandedClassesBlacklist;
+			}
 		}
 
 		FORCEINLINE bool PassesFilter(const FAssetData& Item) const
@@ -146,6 +184,11 @@ namespace AssetViewUtils
 				return false;
 			}
 
+			if (AssetClassBlacklist && !AssetClassBlacklist->PassesFilter(Item.AssetClass))
+			{
+				return false;
+			}
+
 			return PassesPackagePathFilter(Item.PackagePath);
 		}
 
@@ -260,6 +303,8 @@ namespace AssetViewUtils
 		TArray<bool> PluginLoadedFromEngine;
 		TArray<TSharedRef<IPlugin>> Plugins;
 		const FBlacklistPaths* FolderBlacklist;
+		const FBlacklistNames* AssetClassBlacklist;
+		FBlacklistNames ExpandedClassesBlacklist;
 
 		FORCEINLINE int32 FindPluginNameUpper(const WIDECHAR* PluginNameUpper, int32 Length) const
 		{
@@ -292,44 +337,6 @@ namespace AssetViewUtils
 			return INDEX_NONE;
 		}
 	};
-
-	static FARFilter CreateSupportedFilter()
-	{
-		FAssetToolsModule& AssetToolsModule = FModuleManager::LoadModuleChecked<FAssetToolsModule>("AssetTools");
-		TArray<TWeakPtr<IAssetTypeActions>> AssetTypeActionsList;
-		AssetToolsModule.Get().GetAssetTypeActionsList(AssetTypeActionsList);
-
-		bool bAnyUnsupported = false;
-		FARFilter SupportedFilter;
-
-		for (const TWeakPtr<IAssetTypeActions>& WeakTypeActions : AssetTypeActionsList)
-		{
-			if (const TSharedPtr<IAssetTypeActions> TypeActions = WeakTypeActions.Pin())
-			{
-				if (!TypeActions->IsSupported())
-				{
-					bAnyUnsupported = true;
-				}
-				else
-				{
-					SupportedFilter.bRecursiveClasses = true;
-					const UClass* SupportedClass = TypeActions->GetSupportedClass();
-					if (SupportedClass != nullptr)
-					{
-						const FName ClassName = SupportedClass->GetFName();
-						SupportedFilter.ClassNames.AddUnique(ClassName);
-					}
-				}
-			}
-		}
-
-		if (!bAnyUnsupported)
-		{
-			return FARFilter();
-		}
-
-		return SupportedFilter;
-	}
 } // namespace AssetViewUtils
 
 SAssetView::~SAssetView()
@@ -542,6 +549,7 @@ void SAssetView::Construct( const FArguments& InArgs )
 	NumVisibleColumns = 0;
 
 	FAssetToolsModule& AssetToolsModule = FModuleManager::LoadModuleChecked<FAssetToolsModule>("AssetTools");
+	AssetClassBlacklist = AssetToolsModule.Get().GetAssetClassBlacklist();
 	FolderBlacklist = AssetToolsModule.Get().GetFolderBlacklist();
 
 	FEditorWidgetsModule& EditorWidgetsModule = FModuleManager::LoadModuleChecked<FEditorWidgetsModule>("EditorWidgets");
@@ -774,20 +782,8 @@ void SAssetView::SetBackendFilter(const FARFilter& InBackendFilter)
 
 void SAssetView::AppendBackendFilter(FARFilter& FilterToAppendTo) const
 {
-	FilterToAppendTo.Append(SupportedFilter);
 	FilterToAppendTo.Append(BackendFilter);
 
-	if (SupportedFilter.bRecursiveClasses && BackendFilter.bRecursiveClasses)
-	{
-		FilterToAppendTo.ClassNames.RemoveAll(
-			[this](const FName& Class)
-			{
-				return !SupportedFilter.ClassNames.Contains(Class) ||
-					!BackendFilter.ClassNames.Contains(Class);
-			});
-
-		FilterToAppendTo.bRecursiveClasses = FilterToAppendTo.ClassNames.Num() > 0;
-	}
 }
 
 void SAssetView::OnCreateNewFolder(const FString& FolderName, const FString& FolderPath)
@@ -1995,11 +1991,19 @@ void SAssetView::RefreshSourceItems()
 	TArray<FAssetData>& Items = OnShouldFilterAsset.IsBound() ? QueriedAssetItems : AssetItems;
 
 	const FBlacklistPaths* FolderBlacklistToUse = (FolderBlacklist.IsValid() && FolderBlacklist->HasFiltering()) ? FolderBlacklist.Get() : nullptr;
-	const bool bShowAll = SourcesData.IsEmpty() && BackendFilter.IsEmpty() && !FolderBlacklistToUse;
+	const FBlacklistNames* AssetClassBlacklistToUse = (AssetClassBlacklist.IsValid() && AssetClassBlacklist->HasFiltering()) ? AssetClassBlacklist.Get() : nullptr;
+	const bool bShowAll = SourcesData.IsEmpty() && BackendFilter.IsEmpty() && !FolderBlacklistToUse && !AssetClassBlacklistToUse;
+
+	bool bOriginalRegistryTemporaryCachingMode = AssetRegistryModule.Get().GetTemporaryCachingMode();
+	if (AssetClassBlacklistToUse && !bOriginalRegistryTemporaryCachingMode)
+	{
+		// Optimization to reduce number of times class tree is generated
+		AssetRegistryModule.Get().SetTemporaryCachingMode(true);
+	}
 
 	bool bShowClasses = false;
 	TArray<FName> ClassPathsToShow;
-	AssetViewUtils::FInitialAssetFilter InitialAssetFilter(IsShowingLocalizedContent(), IsShowingEngineContent(), IsShowingPluginContent(), FolderBlacklistToUse);
+	AssetViewUtils::FInitialAssetFilter InitialAssetFilter(IsShowingLocalizedContent(), IsShowingEngineContent(), IsShowingPluginContent(), FolderBlacklistToUse, AssetClassBlacklistToUse);
 
 	if ( bShowAll )
 	{
@@ -2152,6 +2156,8 @@ void SAssetView::RefreshSourceItems()
 			Items.Add(FAssetData(CurrentClass));
 		}
 	}
+
+	AssetRegistryModule.Get().SetTemporaryCachingMode(bOriginalRegistryTemporaryCachingMode);
 }
 
 bool SAssetView::IsFilteringRecursively() const
@@ -2997,17 +3003,29 @@ void SAssetView::RunAssetsThroughBlacklists(TArray<FAssetData>& InOutAssetDataLi
 		return;
 	}
 
-	if (FolderBlacklist.IsValid() && FolderBlacklist->HasFiltering())
+	const bool bClassFiltering = AssetClassBlacklist.IsValid() && AssetClassBlacklist->HasFiltering();
+	if ((FolderBlacklist.IsValid() && FolderBlacklist->HasFiltering()) || bClassFiltering)
 	{
-		if (FolderBlacklist->IsBlacklistAll())
+		if ((FolderBlacklist.IsValid() && FolderBlacklist->IsBlacklistAll()) || (bClassFiltering && AssetClassBlacklist->IsBlacklistAll()))
 		{
 			InOutAssetDataList.Reset();
 			return;
 		}
 
-		InOutAssetDataList.RemoveAll([this](const FAssetData& AssetData)
+		// Expand filter to include subclasses so an exact FName comparison can occur
+		FBlacklistNames ExpandedClassesBlacklist;
+		if (bClassFiltering)
 		{
-			if (FolderBlacklist.IsValid() && !FolderBlacklist->PassesStartsWithFilter(AssetData.PackagePath))
+			AssetViewUtils::ExpandClassesBlacklist(*AssetClassBlacklist.Get(), ExpandedClassesBlacklist);
+		}
+
+		InOutAssetDataList.RemoveAll([this, &ExpandedClassesBlacklist](const FAssetData& AssetData)
+		{
+			if (!ExpandedClassesBlacklist.PassesFilter(AssetData.AssetClass))
+			{
+				return true;
+			}
+			else if (FolderBlacklist.IsValid() && !FolderBlacklist->PassesStartsWithFilter(AssetData.PackagePath))
 			{
 				return true;
 			}
diff --git a/Engine/Source/Editor/ContentBrowser/Private/SAssetView.h b/Engine/Source/Editor/ContentBrowser/Private/SAssetView.h
--- a/Engine/Source/Editor/ContentBrowser/Private/SAssetView.h
+++ b/Engine/Source/Editor/ContentBrowser/Private/SAssetView.h
@@ -22,6 +22,7 @@
 #include "Editor/ContentBrowser/Private/AssetViewSortManager.h"
 #include "AssetViewTypes.h"
 #include "HistoryManager.h"
+#include "Misc/BlacklistNames.h"
 
 class FMenuBuilder;
 class FBlacklistPaths;
@@ -846,7 +847,7 @@ private:
 	/** The current base source filter for the view */
 	FSourcesData SourcesData;
 	FARFilter BackendFilter;
-	FARFilter SupportedFilter;
+	TSharedPtr<FBlacklistNames> AssetClassBlacklist;
 	TSharedPtr<FBlacklistPaths> FolderBlacklist;
 	TSharedPtr<FAssetFilterCollectionType> FrontendFilters;
 
diff --git a/Engine/Source/Runtime/AssetRegistry/Private/AssetRegistry.cpp b/Engine/Source/Runtime/AssetRegistry/Private/AssetRegistry.cpp
--- a/Engine/Source/Runtime/AssetRegistry/Private/AssetRegistry.cpp
+++ b/Engine/Source/Runtime/AssetRegistry/Private/AssetRegistry.cpp
@@ -2774,6 +2774,11 @@ void UAssetRegistryImpl::SetTemporaryCachingMode(bool bEnable)
 	}
 }
 
+bool UAssetRegistryImpl::GetTemporaryCachingMode() const
+{
+	return bTempCachingEnabled;
+}
+
 void UAssetRegistryImpl::ClearTemporaryCaches() const
 {
 	if (!bTempCachingEnabled)
diff --git a/Engine/Source/Runtime/AssetRegistry/Private/AssetRegistry.h b/Engine/Source/Runtime/AssetRegistry/Private/AssetRegistry.h
--- a/Engine/Source/Runtime/AssetRegistry/Private/AssetRegistry.h
+++ b/Engine/Source/Runtime/AssetRegistry/Private/AssetRegistry.h
@@ -57,6 +57,7 @@ public:
 	virtual void UseFilterToExcludeAssets(TArray<FAssetData>& AssetDataList, const FARFilter& Filter) const override;
 	virtual void ExpandRecursiveFilter(const FARFilter& InFilter, FARFilter& ExpandedFilter) const override;
 	virtual void SetTemporaryCachingMode(bool bEnable) override;
+	virtual bool GetTemporaryCachingMode() const override;
 	virtual EAssetAvailability::Type GetAssetAvailability(const FAssetData& AssetData) const override;	
 	virtual float GetAssetAvailabilityProgress(const FAssetData& AssetData, EAssetAvailabilityProgressReportingType::Type ReportType) const override;
 	virtual bool GetAssetAvailabilityProgressTypeSupported(EAssetAvailabilityProgressReportingType::Type ReportType) const override;
diff --git a/Engine/Source/Runtime/AssetRegistry/Public/IAssetRegistry.h b/Engine/Source/Runtime/AssetRegistry/Public/IAssetRegistry.h
--- a/Engine/Source/Runtime/AssetRegistry/Public/IAssetRegistry.h
+++ b/Engine/Source/Runtime/AssetRegistry/Public/IAssetRegistry.h
@@ -299,6 +299,9 @@ public:
 	/** Enables or disable temporary search caching, when this is enabled scanning/searching is faster because we assume no objects are loaded between scans. Disabling frees any caches created */
 	virtual void SetTemporaryCachingMode(bool bEnable) = 0;
 
+	/** Returns true if temporary caching mode enabled */
+	virtual bool GetTemporaryCachingMode() const = 0;
+
 	/**
 	 * Gets the current availability of an asset, primarily for streaming install purposes.
 	 *
