#UE4 Nonunity fix

diff --git a/Engine/Source/Editor/ClassViewer/Private/ClassViewerFilter.cpp b/Engine/Source/Editor/ClassViewer/Private/ClassViewerFilter.cpp
--- a/Engine/Source/Editor/ClassViewer/Private/ClassViewerFilter.cpp
+++ b/Engine/Source/Editor/ClassViewer/Private/ClassViewerFilter.cpp
@@ -9,8 +9,7 @@
 #include "Misc/PackageName.h"
 #include "Misc/Paths.h"
 #include "Misc/TextFilterExpressionEvaluator.h"
-#include "Editor/UnrealEdEngine.h"
-#include "UnrealEdGlobals.h"
+#include "Editor.h"
 #include "AssetRegistryModule.h"
 #include "PropertyHandle.h"
 
diff --git a/Engine/Source/Editor/ContentBrowser/Private/SAssetPicker.cpp b/Engine/Source/Editor/ContentBrowser/Private/SAssetPicker.cpp
--- a/Engine/Source/Editor/ContentBrowser/Private/SAssetPicker.cpp
+++ b/Engine/Source/Editor/ContentBrowser/Private/SAssetPicker.cpp
@@ -20,8 +20,7 @@
 #include "SAssetView.h"
 #include "SContentBrowser.h"
 #include "Framework/Commands/GenericCommands.h"
-#include "UnrealEdGlobals.h"
-#include "Editor/UnrealEdEngine.h"
+#include "Editor.h"
 #include "PropertyHandle.h"
 
 #define LOCTEXT_NAMESPACE "ContentBrowser"
