ContentBrowserSingleton now ensures that the editor style module is loaded before it access public static functions from within that module.

diff --git a/Engine/Source/Editor/ContentBrowser/Private/ContentBrowserSingleton.cpp b/Engine/Source/Editor/ContentBrowser/Private/ContentBrowserSingleton.cpp
--- a/Engine/Source/Editor/ContentBrowser/Private/ContentBrowserSingleton.cpp
+++ b/Engine/Source/Editor/ContentBrowser/Private/ContentBrowserSingleton.cpp
@@ -38,6 +38,9 @@ FContentBrowserSingleton::FContentBrowserSingleton()
 	, CollectionAssetRegistryBridge(MakeShared<FCollectionAssetRegistryBridge>())
 	, SettingsStringID(0)
 {
+	// We're going to call a static function in the editor style module, so we need to make sure the module has actually been loaded
+	FModuleManager::Get().LoadModuleChecked("EditorStyle");
+
 	// Register the tab spawners for all content browsers
 	const FSlateIcon ContentBrowserIcon(FEditorStyle::GetStyleSetName(), "ContentBrowser.TabIcon");
 	const IWorkspaceMenuStructure& MenuStructure = WorkspaceMenu::GetMenuStructure();
