Content Browser: Folder breadcrumb menus now use blacklist

diff --git a/Engine/Source/Editor/ContentBrowser/Private/SContentBrowser.cpp b/Engine/Source/Editor/ContentBrowser/Private/SContentBrowser.cpp
--- a/Engine/Source/Editor/ContentBrowser/Private/SContentBrowser.cpp
+++ b/Engine/Source/Editor/ContentBrowser/Private/SContentBrowser.cpp
@@ -1966,6 +1966,16 @@ TSharedRef<SWidget> SContentBrowser::OnGetCrumbDelimiterContent(const FString& C
 			IAssetRegistry& AssetRegistry = AssetRegistryModule.Get();
 
 			AssetRegistry.GetSubPaths( CrumbData, SubPaths, bRecurse );
+
+			FAssetToolsModule& AssetToolsModule = FModuleManager::LoadModuleChecked<FAssetToolsModule>("AssetTools");
+			TSharedRef<FBlacklistPaths> FolderBlacklist = AssetToolsModule.Get().GetFolderBlacklist();
+			if (FolderBlacklist->HasFiltering())
+			{
+				SubPaths.RemoveAll([FolderBlacklist](const FString& SubPath)
+				{
+					return !FolderBlacklist->PassesStartsWithFilter(SubPath);
+				});
+			}
 		}
 
 		if( SubPaths.Num() > 0 )
