Content browser: Frontend filter list now applies class blacklist to list of filters

diff --git a/Engine/Source/Editor/ContentBrowser/Private/SFilterList.cpp b/Engine/Source/Editor/ContentBrowser/Private/SFilterList.cpp
--- a/Engine/Source/Editor/ContentBrowser/Private/SFilterList.cpp
+++ b/Engine/Source/Editor/ContentBrowser/Private/SFilterList.cpp
@@ -24,6 +24,7 @@
 #include "AssetToolsModule.h"
 #include "FrontendFilters.h"
 #include "ContentBrowserFrontEndFilterExtension.h"
+#include "Misc/BlacklistNames.h"
 
 #define LOCTEXT_NAMESPACE "ContentBrowser"
 
@@ -1207,6 +1208,8 @@ TSharedRef<SWidget> SFilterList::MakeAddFilterMenu(EAssetTypeCategories::Type Me
 	};
 	AssetTypeActionsList.Sort( FCompareIAssetTypeActions() );
 
+	TSharedRef<FBlacklistNames> AssetClassBlacklist = AssetToolsModule.Get().GetAssetClassBlacklist();
+
 	// For every asset type, move it into all the categories it should appear in
 	for (int32 ClassIdx = 0; ClassIdx < AssetTypeActionsList.Num(); ++ClassIdx)
 	{
@@ -1216,7 +1219,8 @@ TSharedRef<SWidget> SFilterList::MakeAddFilterMenu(EAssetTypeCategories::Type Me
 			TSharedPtr<IAssetTypeActions> TypeActions = WeakTypeActions.Pin();
 			if ( ensure(TypeActions.IsValid()) && TypeActions->CanFilter() )
 			{
-				if(!IsFilteredByPicker(InitialClassFilters, TypeActions->GetSupportedClass()))
+				UClass* SupportedClass = TypeActions->GetSupportedClass();
+				if ((!SupportedClass || AssetClassBlacklist->PassesFilter(SupportedClass->GetFName())) && !IsFilteredByPicker(InitialClassFilters, SupportedClass))
 				{
 					for ( auto MenuIt = CategoryToMenuMap.CreateIterator(); MenuIt; ++MenuIt )
 					{
