Content Browser: Hide 'Basic' filters that do not pass an active class blacklist

diff --git a/Engine/Source/Editor/ContentBrowser/Private/SFilterList.cpp b/Engine/Source/Editor/ContentBrowser/Private/SFilterList.cpp
--- a/Engine/Source/Editor/ContentBrowser/Private/SFilterList.cpp
+++ b/Engine/Source/Editor/ContentBrowser/Private/SFilterList.cpp
@@ -1445,6 +1445,7 @@ void SFilterList::GetTypeActionsForCategory(EAssetTypeCategories::Type Category,
 	FAssetToolsModule& AssetToolsModule = FModuleManager::LoadModuleChecked<FAssetToolsModule>(TEXT("AssetTools"));
 	TArray<TWeakPtr<IAssetTypeActions>> AssetTypeActionsList;
 	AssetToolsModule.Get().GetAssetTypeActionsList(AssetTypeActionsList);
+	TSharedRef<FBlacklistNames> AssetClassBlacklist = AssetToolsModule.Get().GetAssetClassBlacklist();
 
 	// Find all asset type actions that match the category
 	for (int32 ClassIdx = 0; ClassIdx < AssetTypeActionsList.Num(); ++ClassIdx)
@@ -1454,7 +1455,10 @@ void SFilterList::GetTypeActionsForCategory(EAssetTypeCategories::Type Category,
 
 		if (ensure(AssetTypeActions.IsValid()) && AssetTypeActions->CanFilter() && AssetTypeActions->GetCategories() & Category)
 		{
-			TypeActions.Add(WeakTypeActions);
+			if (AssetTypeActions->GetSupportedClass() == nullptr || AssetClassBlacklist->PassesFilter(AssetTypeActions->GetSupportedClass()->GetFName()))
+			{
+				TypeActions.Add(WeakTypeActions);
+			}
 		}
 	}
 }
