Content Browser: Asset View Button context menu

diff --git a/Engine/Source/Editor/ContentBrowser/Private/SAssetView.cpp b/Engine/Source/Editor/ContentBrowser/Private/SAssetView.cpp
--- a/Engine/Source/Editor/ContentBrowser/Private/SAssetView.cpp
+++ b/Engine/Source/Editor/ContentBrowser/Private/SAssetView.cpp
@@ -51,6 +51,8 @@
 #include "Misc/BlacklistNames.h"
 #include "AssetRegistryState.h"
 #include "Materials/Material.h"
+#include "ContentBrowserMenuContexts.h"
+#include "ToolMenus.h"
 
 #define LOCTEXT_NAMESPACE "ContentBrowser"
 #define MAX_THUMBNAIL_SIZE 4096
@@ -3122,6 +3124,8 @@ FSlateColor SAssetView::GetViewButtonForegroundColor() const
 
 TSharedRef<SWidget> SAssetView::GetViewButtonContent()
 {
+	SAssetView::RegisterGetViewButtonMenu();
+
 	// Get all menu extenders for this context menu from the content browser module
 	FContentBrowserModule& ContentBrowserModule = FModuleManager::GetModuleChecked<FContentBrowserModule>( TEXT("ContentBrowser") );
 	TArray<FContentBrowserMenuExtender> MenuExtenderDelegates = ContentBrowserModule.GetAllAssetViewViewMenuExtenders();
@@ -3136,11 +3140,37 @@ TSharedRef<SWidget> SAssetView::GetViewButtonContent()
 	}
 	TSharedPtr<FExtender> MenuExtender = FExtender::Combine(Extenders);
 
-	FMenuBuilder MenuBuilder(/*bInShouldCloseWindowAfterMenuSelection=*/true, NULL, MenuExtender, /*bCloseSelfOnly=*/ true);
+	UContentBrowserAssetViewContextMenuContext* Context = NewObject<UContentBrowserAssetViewContextMenuContext>();
+	Context->AssetView = SharedThis(this);
+	FToolMenuContext MenuContext(nullptr, MenuExtender, Context);
+	return UToolMenus::Get()->GenerateWidget("ContentBrowser.AssetViewOptions", MenuContext);
+}
+
+void SAssetView::RegisterGetViewButtonMenu()
+{
+	if (!UToolMenus::Get()->IsMenuRegistered("ContentBrowser.AssetViewOptions"))
+	{
+		UToolMenu* Menu = UToolMenus::Get()->RegisterMenu("ContentBrowser.AssetViewOptions");
+		Menu->bCloseSelfOnly = true;
+		Menu->AddDynamicSection("DynamicContent", FNewToolMenuDelegate::CreateLambda([](UToolMenu* InMenu)
+		{
+			if (UContentBrowserAssetViewContextMenuContext* Context = InMenu->FindContext<UContentBrowserAssetViewContextMenuContext>())
+			{
+				if (Context->AssetView.IsValid())
+				{
+					Context->AssetView.Pin()->PopulateViewButtonMenu(InMenu);
+				}
+			}
+		}));
+	}
+}
 
-	MenuBuilder.BeginSection("AssetViewType", LOCTEXT("ViewTypeHeading", "View Type"));
+void SAssetView::PopulateViewButtonMenu(UToolMenu* Menu)
+{
 	{
-		MenuBuilder.AddMenuEntry(
+		FToolMenuSection& Section = Menu->AddSection("AssetViewType", LOCTEXT("ViewTypeHeading", "View Type"));
+		Section.AddMenuEntry(
+			"TileView",
 			LOCTEXT("TileViewOption", "Tiles"),
 			LOCTEXT("TileViewOptionToolTip", "View assets as tiles in a grid."),
 			FSlateIcon(),
@@ -3149,11 +3179,11 @@ TSharedRef<SWidget> SAssetView::GetViewButtonContent()
 				FCanExecuteAction(),
 				FIsActionChecked::CreateSP( this, &SAssetView::IsCurrentViewType, EAssetViewType::Tile )
 				),
-			NAME_None,
 			EUserInterfaceActionType::RadioButton
 			);
 
-		MenuBuilder.AddMenuEntry(
+		Section.AddMenuEntry(
+			"ListView",
 			LOCTEXT("ListViewOption", "List"),
 			LOCTEXT("ListViewOptionToolTip", "View assets in a list with thumbnails."),
 			FSlateIcon(),
@@ -3162,11 +3192,11 @@ TSharedRef<SWidget> SAssetView::GetViewButtonContent()
 				FCanExecuteAction(),
 				FIsActionChecked::CreateSP( this, &SAssetView::IsCurrentViewType, EAssetViewType::List )
 				),
-			NAME_None,
 			EUserInterfaceActionType::RadioButton
 			);
 
-		MenuBuilder.AddMenuEntry(
+		Section.AddMenuEntry(
+			"ColumnView",
 			LOCTEXT("ColumnViewOption", "Columns"),
 			LOCTEXT("ColumnViewOptionToolTip", "View assets in a list with columns of details."),
 			FSlateIcon(),
@@ -3175,17 +3205,17 @@ TSharedRef<SWidget> SAssetView::GetViewButtonContent()
 				FCanExecuteAction(),
 				FIsActionChecked::CreateSP( this, &SAssetView::IsCurrentViewType, EAssetViewType::Column )
 				),
-			NAME_None,
 			EUserInterfaceActionType::RadioButton
 			);
 	}
-	MenuBuilder.EndSection();
 
-	MenuBuilder.BeginSection("View", LOCTEXT("ViewHeading", "View"));
 	{
-		auto CreateShowFoldersSubMenu = [this](FMenuBuilder& SubMenuBuilder)
+		FToolMenuSection& Section = Menu->AddSection("View", LOCTEXT("ViewHeading", "View"));
+		auto CreateShowFoldersSubMenu = [this](UToolMenu* SubMenu)
 		{
-			SubMenuBuilder.AddMenuEntry(
+			FToolMenuSection& ShowEmptyFoldersSection = SubMenu->AddSection("ShowEmptyFolders");
+			ShowEmptyFoldersSection.AddMenuEntry(
+				"ShowEmptyFolders",
 				LOCTEXT("ShowEmptyFoldersOption", "Show Empty Folders"),
 				LOCTEXT("ShowEmptyFoldersOptionToolTip", "Show empty folders in the view as well as assets?"),
 				FSlateIcon(),
@@ -3194,25 +3224,25 @@ TSharedRef<SWidget> SAssetView::GetViewButtonContent()
 					FCanExecuteAction::CreateSP( this, &SAssetView::IsToggleShowEmptyFoldersAllowed ),
 					FIsActionChecked::CreateSP( this, &SAssetView::IsShowingEmptyFolders )
 				),
-				NAME_None,
 				EUserInterfaceActionType::ToggleButton
 			);
 		};
 
-		MenuBuilder.AddSubMenu(
+		Section.AddEntry(FToolMenuEntry::InitSubMenu(
+			"ShowFolders",
 			LOCTEXT("ShowFoldersOption", "Show Folders"),
 			LOCTEXT("ShowFoldersOptionToolTip", "Show folders in the view as well as assets?"),
-			FNewMenuDelegate::CreateLambda(CreateShowFoldersSubMenu),
+			FNewToolMenuDelegate::CreateLambda(CreateShowFoldersSubMenu),
 			FUIAction(
 				FExecuteAction::CreateSP( this, &SAssetView::ToggleShowFolders ),
 				FCanExecuteAction::CreateSP( this, &SAssetView::IsToggleShowFoldersAllowed ),
 				FIsActionChecked::CreateSP( this, &SAssetView::IsShowingFolders )
 			),
-			NAME_None,
 			EUserInterfaceActionType::ToggleButton
-		);
+		));
 
-		MenuBuilder.AddMenuEntry(
+		Section.AddMenuEntry(
+			"ShowFavorite",
 			LOCTEXT("ShowFavoriteOptions", "Show Favorites"),
 			LOCTEXT("ShowFavoriteOptionToolTip", "Show the favorite folders in the view?"),
 			FSlateIcon(),
@@ -3221,11 +3251,11 @@ TSharedRef<SWidget> SAssetView::GetViewButtonContent()
 				FCanExecuteAction::CreateSP(this, &SAssetView::IsToggleShowFavoritesAllowed),
 				FIsActionChecked::CreateSP(this, &SAssetView::IsShowingFavorites)
 			),
-			NAME_None,
 			EUserInterfaceActionType::ToggleButton
 		);
 
-		MenuBuilder.AddMenuEntry(
+		Section.AddMenuEntry(
+			"DockCollections",
 			LOCTEXT("DockCollectionsOptions", "Dock Collections"),
 			LOCTEXT("DockCollectionsOptionToolTip", "Dock the collections view under the path view?"),
 			FSlateIcon(),
@@ -3234,11 +3264,11 @@ TSharedRef<SWidget> SAssetView::GetViewButtonContent()
 				FCanExecuteAction::CreateSP(this, &SAssetView::IsToggleDockCollectionsAllowed),
 				FIsActionChecked::CreateSP(this, &SAssetView::HasDockedCollections)
 			),
-			NAME_None,
 			EUserInterfaceActionType::ToggleButton
 		);
 
-		MenuBuilder.AddMenuEntry(
+		Section.AddMenuEntry(
+			"FilterRecursively",
 			LOCTEXT("FilterRecursivelyOption", "Filter Recursively"),
 			LOCTEXT("FilterRecursivelyOptionToolTip", "Should filters apply recursively in the view?"),
 			FSlateIcon(),
@@ -3247,15 +3277,14 @@ TSharedRef<SWidget> SAssetView::GetViewButtonContent()
 				FCanExecuteAction::CreateSP(this, &SAssetView::IsToggleFilteringRecursivelyAllowed),
 				FIsActionChecked::CreateSP(this, &SAssetView::IsFilteringRecursively)
 			),
-			NAME_None,
 			EUserInterfaceActionType::ToggleButton
 		);
 	}
-	MenuBuilder.EndSection();
 
-	MenuBuilder.BeginSection("Content", LOCTEXT("ContentHeading", "Content"));
 	{
-		MenuBuilder.AddMenuEntry(
+		FToolMenuSection& Section = Menu->AddSection("Content", LOCTEXT("ContentHeading", "Content"));
+		Section.AddMenuEntry(
+			"ShowCppClasses",
 			LOCTEXT("ShowCppClassesOption", "Show C++ Classes"),
 			LOCTEXT("ShowCppClassesOptionToolTip", "Show C++ classes in the view?"),
 			FSlateIcon(),
@@ -3264,11 +3293,11 @@ TSharedRef<SWidget> SAssetView::GetViewButtonContent()
 				FCanExecuteAction::CreateSP( this, &SAssetView::IsToggleShowCppContentAllowed ),
 				FIsActionChecked::CreateSP( this, &SAssetView::IsShowingCppContent )
 			),
-			NAME_None,
 			EUserInterfaceActionType::ToggleButton
 		);
 
-		MenuBuilder.AddMenuEntry(
+		Section.AddMenuEntry(
+			"ShowDevelopersContent",
 			LOCTEXT("ShowDevelopersContentOption", "Show Developers Content"),
 			LOCTEXT("ShowDevelopersContentOptionToolTip", "Show developers content in the view?"),
 			FSlateIcon(),
@@ -3277,11 +3306,11 @@ TSharedRef<SWidget> SAssetView::GetViewButtonContent()
 				FCanExecuteAction::CreateSP( this, &SAssetView::IsToggleShowDevelopersContentAllowed ),
 				FIsActionChecked::CreateSP( this, &SAssetView::IsShowingDevelopersContent )
 				),
-			NAME_None,
 			EUserInterfaceActionType::ToggleButton
 		);
 
-		MenuBuilder.AddMenuEntry(
+		Section.AddMenuEntry(
+			"ShowEngineFolder",
 			LOCTEXT("ShowEngineFolderOption", "Show Engine Content"),
 			LOCTEXT("ShowEngineFolderOptionToolTip", "Show engine content in the view?"),
 			FSlateIcon(),
@@ -3290,11 +3319,11 @@ TSharedRef<SWidget> SAssetView::GetViewButtonContent()
 				FCanExecuteAction(),
 				FIsActionChecked::CreateSP( this, &SAssetView::IsShowingEngineContent )
 			),
-			NAME_None,
 			EUserInterfaceActionType::ToggleButton
 		);
 
-		MenuBuilder.AddMenuEntry(
+		Section.AddMenuEntry(
+			"ShowPluginFolder",
 			LOCTEXT("ShowPluginFolderOption", "Show Plugin Content"),
 			LOCTEXT("ShowPluginFolderOptionToolTip", "Show plugin content in the view?"),
 			FSlateIcon(),
@@ -3303,11 +3332,11 @@ TSharedRef<SWidget> SAssetView::GetViewButtonContent()
 				FCanExecuteAction(),
 				FIsActionChecked::CreateSP( this, &SAssetView::IsShowingPluginContent )
 			),
-			NAME_None,
 			EUserInterfaceActionType::ToggleButton
 		);
 
-		MenuBuilder.AddMenuEntry(
+		Section.AddMenuEntry(
+			"ShowLocalizedContent",
 			LOCTEXT("ShowLocalizedContentOption", "Show Localized Content"),
 			LOCTEXT("ShowLocalizedContentOptionToolTip", "Show localized content in the view?"),
 			FSlateIcon(),
@@ -3316,15 +3345,14 @@ TSharedRef<SWidget> SAssetView::GetViewButtonContent()
 				FCanExecuteAction::CreateSP(this, &SAssetView::IsToggleShowLocalizedContentAllowed),
 				FIsActionChecked::CreateSP(this, &SAssetView::IsShowingLocalizedContent)
 				),
-			NAME_None,
 			EUserInterfaceActionType::ToggleButton
 			);
 	}
-	MenuBuilder.EndSection();
 
-	MenuBuilder.BeginSection("Search", LOCTEXT("SearchHeading", "Search"));
 	{
-		MenuBuilder.AddMenuEntry(
+		FToolMenuSection& Section = Menu->AddSection("Search", LOCTEXT("SearchHeading", "Search"));
+		Section.AddMenuEntry(
+			"IncludeClassName",
 			LOCTEXT("IncludeClassNameOption", "Search Asset Class Names"),
 			LOCTEXT("IncludeClassesNameOptionTooltip", "Include asset type names in search criteria?  (e.g. Blueprint, Texture, Sound)"),
 			FSlateIcon(),
@@ -3333,11 +3361,11 @@ TSharedRef<SWidget> SAssetView::GetViewButtonContent()
 				FCanExecuteAction::CreateSP(this, &SAssetView::IsToggleIncludeClassNamesAllowed),
 				FIsActionChecked::CreateSP(this, &SAssetView::IsIncludingClassNames)
 			),
-			NAME_None,
 			EUserInterfaceActionType::ToggleButton
 		);
 
-		MenuBuilder.AddMenuEntry(
+		Section.AddMenuEntry(
+			"IncludeAssetPath",
 			LOCTEXT("IncludeAssetPathOption", "Search Asset Path"),
 			LOCTEXT("IncludeAssetPathOptionTooltip", "Include entire asset path in search criteria?"),
 			FSlateIcon(),
@@ -3346,11 +3374,11 @@ TSharedRef<SWidget> SAssetView::GetViewButtonContent()
 				FCanExecuteAction::CreateSP(this, &SAssetView::IsToggleIncludeAssetPathsAllowed),
 				FIsActionChecked::CreateSP(this, &SAssetView::IsIncludingAssetPaths)
 			),
-			NAME_None,
 			EUserInterfaceActionType::ToggleButton
 		);
 
-		MenuBuilder.AddMenuEntry(
+		Section.AddMenuEntry(
+			"IncludeCollectionName",
 			LOCTEXT("IncludeCollectionNameOption", "Search Collection Names"),
 			LOCTEXT("IncludeCollectionNameOptionTooltip", "Include Collection names in search criteria?"),
 			FSlateIcon(),
@@ -3359,15 +3387,14 @@ TSharedRef<SWidget> SAssetView::GetViewButtonContent()
 				FCanExecuteAction::CreateSP(this, &SAssetView::IsToggleIncludeCollectionNamesAllowed),
 				FIsActionChecked::CreateSP(this, &SAssetView::IsIncludingCollectionNames)
 			),
-			NAME_None,
 			EUserInterfaceActionType::ToggleButton
 		);
 	}
-	MenuBuilder.EndSection();
 
-	MenuBuilder.BeginSection("AssetThumbnails", LOCTEXT("ThumbnailsHeading", "Thumbnails"));
 	{
-		MenuBuilder.AddWidget(
+		FToolMenuSection& Section = Menu->AddSection("AssetThumbnails", LOCTEXT("ThumbnailsHeading", "Thumbnails"));
+		Section.AddEntry(FToolMenuEntry::InitWidget(
+			"ThumbnailScale",
 			SNew(SSlider)
 				.ToolTipText( LOCTEXT("ThumbnailScaleToolTip", "Adjust the size of thumbnails.") )
 				.Value( this, &SAssetView::GetThumbnailScale )
@@ -3375,9 +3402,10 @@ TSharedRef<SWidget> SAssetView::GetViewButtonContent()
 				.Locked( this, &SAssetView::IsThumbnailScalingLocked ),
 			LOCTEXT("ThumbnailScaleLabel", "Scale"),
 			/*bNoIndent=*/true
-			);
+			));
 
-		MenuBuilder.AddMenuEntry(
+		Section.AddMenuEntry(
+			"ThumbnailEditMode",
 			LOCTEXT("ThumbnailEditModeOption", "Thumbnail Edit Mode"),
 			LOCTEXT("ThumbnailEditModeOptionToolTip", "Toggle thumbnail editing mode. When in this mode you can rotate the camera on 3D thumbnails by dragging them."),
 			FSlateIcon(),
@@ -3386,11 +3414,11 @@ TSharedRef<SWidget> SAssetView::GetViewButtonContent()
 				FCanExecuteAction::CreateSP( this, &SAssetView::IsThumbnailEditModeAllowed ),
 				FIsActionChecked::CreateSP( this, &SAssetView::IsThumbnailEditMode )
 				),
-			NAME_None,
 			EUserInterfaceActionType::ToggleButton
 			);
 
-		MenuBuilder.AddMenuEntry(
+		Section.AddMenuEntry(
+			"RealTimeThumbnails",
 			LOCTEXT("RealTimeThumbnailsOption", "Real-Time Thumbnails"),
 			LOCTEXT("RealTimeThumbnailsOptionToolTip", "Renders the assets thumbnails in real-time"),
 			FSlateIcon(),
@@ -3399,17 +3427,16 @@ TSharedRef<SWidget> SAssetView::GetViewButtonContent()
 				FCanExecuteAction::CreateSP( this, &SAssetView::CanShowRealTimeThumbnails ),
 				FIsActionChecked::CreateSP( this, &SAssetView::IsShowingRealTimeThumbnails )
 			),
-			NAME_None,
 			EUserInterfaceActionType::ToggleButton
 			);
 	}
-	MenuBuilder.EndSection();
 
 	if (GetColumnViewVisibility() == EVisibility::Visible)
 	{
-		MenuBuilder.BeginSection("AssetColumns", LOCTEXT("ToggleColumnsHeading", "Columns"));
 		{
-			MenuBuilder.AddSubMenu(
+			FToolMenuSection& Section = Menu->AddSection("AssetColumns", LOCTEXT("ToggleColumnsHeading", "Columns"));
+			Section.AddSubMenu(
+				"ToggleColumns",
 				LOCTEXT("ToggleColumnsMenu", "Toggle columns"),
 				LOCTEXT("ToggleColumnsMenuTooltip", "Show or hide specific columns."),
 				FNewMenuDelegate::CreateSP(this, &SAssetView::FillToggleColumnsMenu),
@@ -3418,28 +3445,25 @@ TSharedRef<SWidget> SAssetView::GetViewButtonContent()
 				false
 				);
 
-			MenuBuilder.AddMenuEntry(
+			Section.AddMenuEntry(
+				"ResetColumns",
 				LOCTEXT("ResetColumns", "Reset Columns"),
 				LOCTEXT("ResetColumnsToolTip", "Reset all columns to be visible again."),
 				FSlateIcon(),
 				FUIAction(FExecuteAction::CreateSP(this, &SAssetView::ResetColumns)),
-				NAME_None,
 				EUserInterfaceActionType::Button
 				);
 
-			MenuBuilder.AddMenuEntry(
+			Section.AddMenuEntry(
+				"ExportColumns",
 				LOCTEXT("ExportColumns", "Export to CSV"),
 				LOCTEXT("ExportColumnsToolTip", "Export column data to CSV."),
 				FSlateIcon(),
 				FUIAction(FExecuteAction::CreateSP(this, &SAssetView::ExportColumns)),
-				NAME_None,
 				EUserInterfaceActionType::Button
 			);
 		}
-		MenuBuilder.EndSection();
 	}
-
-	return MenuBuilder.MakeWidget();
 }
 
 void SAssetView::ToggleShowFolders()
diff --git a/Engine/Source/Editor/ContentBrowser/Private/SAssetView.h b/Engine/Source/Editor/ContentBrowser/Private/SAssetView.h
--- a/Engine/Source/Editor/ContentBrowser/Private/SAssetView.h
+++ b/Engine/Source/Editor/ContentBrowser/Private/SAssetView.h
@@ -462,6 +462,12 @@ private:
 	/** Handler for when the view combo button is clicked */
 	TSharedRef<SWidget> GetViewButtonContent();
 
+	/** Register menu for when the view combo button is clicked */
+	static void RegisterGetViewButtonMenu();
+
+	/** Fill in menu content for when the view combo button is clicked */
+	void PopulateViewButtonMenu(class UToolMenu* Menu);
+
 	/** Toggle whether folders should be shown or not */
 	void ToggleShowFolders();
 
diff --git a/Engine/Source/Editor/ContentBrowser/Public/ContentBrowserMenuContexts.h b/Engine/Source/Editor/ContentBrowser/Public/ContentBrowserMenuContexts.h
--- a/Engine/Source/Editor/ContentBrowser/Public/ContentBrowserMenuContexts.h
+++ b/Engine/Source/Editor/ContentBrowser/Public/ContentBrowserMenuContexts.h
@@ -10,6 +10,7 @@
 
 class FAssetContextMenu;
 class IAssetTypeActions;
+class SAssetView;
 
 UCLASS()
 class CONTENTBROWSER_API UContentBrowserAssetContextMenuContext : public UObject
@@ -41,3 +42,13 @@ public:
 	}
 };
 
+UCLASS()
+class CONTENTBROWSER_API UContentBrowserAssetViewContextMenuContext : public UObject
+{
+	GENERATED_BODY()
+
+public:
+
+	TWeakPtr<SAssetView> AssetView;
+};
+
