Added missing Content Browser style to paths view splitter

#rnx

diff --git a/Engine/Source/Editor/ContentBrowser/Private/SContentBrowser.cpp b/Engine/Source/Editor/ContentBrowser/Private/SContentBrowser.cpp
--- a/Engine/Source/Editor/ContentBrowser/Private/SContentBrowser.cpp
+++ b/Engine/Source/Editor/ContentBrowser/Private/SContentBrowser.cpp
@@ -543,6 +543,8 @@ void SContentBrowser::Construct( const FArguments& InArgs, const FName& InInstan
 						+SWidgetSwitcher::Slot()
 						[
 							SAssignNew(PathFavoriteSplitterPtr, SSplitter)
+							.Style(FEditorStyle::Get(), "ContentBrowser.Splitter")
+							.PhysicalSplitterHandleSize(2.0f)
 							.Orientation(EOrientation::Orient_Vertical)
 							.MinimumSlotHeight(70.0f)
 							.Visibility( this, &SContentBrowser::GetSourcesViewVisibility )
