Content Browser: Context menu when clicking on folders converted to tool menu
#rnx

diff --git a/Engine/Source/Editor/ContentBrowser/Private/NewAssetOrClassContextMenu.cpp b/Engine/Source/Editor/ContentBrowser/Private/NewAssetOrClassContextMenu.cpp
--- a/Engine/Source/Editor/ContentBrowser/Private/NewAssetOrClassContextMenu.cpp
+++ b/Engine/Source/Editor/ContentBrowser/Private/NewAssetOrClassContextMenu.cpp
@@ -24,6 +24,7 @@
 #include "IDocumentation.h"
 #include "ClassIconFinder.h"
 #include "AssetToolsModule.h"
+#include "ToolMenus.h"
 
 #define LOCTEXT_NAMESPACE "ContentBrowser"
 
@@ -180,7 +181,7 @@ public:
 };
 
 void FNewAssetOrClassContextMenu::MakeContextMenu(
-	FMenuBuilder& MenuBuilder, 
+	UToolMenu* Menu, 
 	const TArray<FName>& InSelectedPaths,
 	const FOnNewAssetRequested& InOnNewAssetRequested, 
 	const FOnNewClassRequested& InOnNewClassRequested, 
@@ -196,11 +197,11 @@ void FNewAssetOrClassContextMenu::MakeContextMenu(
 		SelectedStringPaths.Add(Path.ToString());
 	}
 
-	MakeContextMenu(MenuBuilder, SelectedStringPaths, InOnNewAssetRequested, InOnNewClassRequested, InOnNewFolderRequested, InOnImportAssetRequested, InOnGetContentRequested);
+	MakeContextMenu(Menu, SelectedStringPaths, InOnNewAssetRequested, InOnNewClassRequested, InOnNewFolderRequested, InOnImportAssetRequested, InOnGetContentRequested);
 }
 
 void FNewAssetOrClassContextMenu::MakeContextMenu(
-	FMenuBuilder& MenuBuilder, 
+	UToolMenu* Menu, 
 	const TArray<FString>& InSelectedPaths,
 	const FOnNewAssetRequested& InOnNewAssetRequested, 
 	const FOnNewClassRequested& InOnNewClassRequested, 
@@ -242,23 +243,23 @@ void FNewAssetOrClassContextMenu::MakeContextMenu(
 	// Get Content
 	if ( InOnGetContentRequested.IsBound() )
 	{
-		MenuBuilder.BeginSection( "ContentBrowserGetContent", LOCTEXT( "GetContentMenuHeading", "Content" ) );
 		{
-			MenuBuilder.AddMenuEntry(
+			FToolMenuSection& Section = Menu->AddSection( "ContentBrowserGetContent", LOCTEXT( "GetContentMenuHeading", "Content" ) );
+			Section.AddMenuEntry(
+				"GetContent",
 				LOCTEXT( "GetContentText", "Add Feature or Content Pack..." ),
 				LOCTEXT( "GetContentTooltip", "Add features and content packs to the project." ),
 				FSlateIcon( FEditorStyle::GetStyleSetName(), "ContentBrowser.AddContent" ),
 				FUIAction( FExecuteAction::CreateStatic( &FNewAssetOrClassContextMenu::ExecuteGetContent, InOnGetContentRequested ) )
 				);
 		}
-		MenuBuilder.EndSection();
 	}
 
 	// New Folder
 	if(InOnNewFolderRequested.IsBound() && GetDefault<UContentBrowserSettings>()->DisplayFolders)
 	{
-		MenuBuilder.BeginSection("ContentBrowserNewFolder", LOCTEXT("FolderMenuHeading", "Folder") );
 		{
+			FToolMenuSection& Section = Menu->AddSection("ContentBrowserNewFolder", LOCTEXT("FolderMenuHeading", "Folder") );
 			FText NewFolderToolTip;
 			if(bHasSinglePathSelected)
 			{
@@ -276,7 +277,8 @@ void FNewAssetOrClassContextMenu::MakeContextMenu(
 				NewFolderToolTip = LOCTEXT("NewFolderTooltip_InvalidNumberOfPaths", "Can only create folders when there is a single path selected.");
 			}
 
-			MenuBuilder.AddMenuEntry(
+			Section.AddMenuEntry(
+				"NewFolder",
 				LOCTEXT("NewFolderLabel", "New Folder"),
 				NewFolderToolTip,
 				FSlateIcon(FEditorStyle::GetStyleSetName(), "ContentBrowser.NewFolderIcon"),
@@ -286,7 +288,6 @@ void FNewAssetOrClassContextMenu::MakeContextMenu(
 					)
 				);
 		}
-		MenuBuilder.EndSection(); //ContentBrowserNewFolder
 	}
 
 	// Add Class
@@ -311,9 +312,10 @@ void FNewAssetOrClassContextMenu::MakeContextMenu(
 			NewClassToolTip = LOCTEXT("NewClassTooltip_InvalidNumberOfPaths", "Can only create classes when there is a single path selected.");
 		}
 
-		MenuBuilder.BeginSection("ContentBrowserNewClass", LOCTEXT("ClassMenuHeading", "C++ Class") );
 		{
-			MenuBuilder.AddMenuEntry(
+			FToolMenuSection& Section = Menu->AddSection("ContentBrowserNewClass", LOCTEXT("ClassMenuHeading", "C++ Class"));
+			Section.AddMenuEntry(
+				"NewClass",
 				LOCTEXT("NewClassLabel", "New C++ Class..."),
 				NewClassToolTip,
 				FSlateIcon(FEditorStyle::GetStyleSetName(), "MainFrame.AddCodeToProject"),
@@ -323,15 +325,15 @@ void FNewAssetOrClassContextMenu::MakeContextMenu(
 					)
 				);
 		}
-		MenuBuilder.EndSection(); //ContentBrowserNewClass
 	}
 
 	// Import
 	if (InOnImportAssetRequested.IsBound() && !FirstSelectedPath.IsEmpty())
 	{
-		MenuBuilder.BeginSection( "ContentBrowserImportAsset", LOCTEXT( "ImportAssetMenuHeading", "Import Asset" ) );
 		{
-			MenuBuilder.AddMenuEntry(
+			FToolMenuSection& Section = Menu->AddSection("ContentBrowserImportAsset", LOCTEXT( "ImportAssetMenuHeading", "Import Asset" ));
+			Section.AddMenuEntry(
+				"ImportAsset",
 				FText::Format( LOCTEXT( "ImportAsset", "Import to {0}..." ), FText::FromString( FirstSelectedPath ) ),
 				LOCTEXT( "ImportAssetTooltip_NewAssetOrClass", "Imports an asset from file to this folder." ),
 				FSlateIcon( FEditorStyle::GetStyleSetName(), "ContentBrowser.ImportIcon" ),
@@ -341,28 +343,28 @@ void FNewAssetOrClassContextMenu::MakeContextMenu(
 					)
 				);
 		}
-		MenuBuilder.EndSection();
 	}
 
 	
 	if (InOnNewAssetRequested.IsBound())
 	{
 		// Add Basic Asset
-		MenuBuilder.BeginSection("ContentBrowserNewBasicAsset", LOCTEXT("CreateBasicAssetsMenuHeading", "Create Basic Asset") );
 		{
+			FToolMenuSection& Section = Menu->AddSection("ContentBrowserNewBasicAsset", LOCTEXT("CreateBasicAssetsMenuHeading", "Create Basic Asset"));
 			CreateNewAssetMenuCategory(
-				MenuBuilder, 
+				Menu,
+				"ContentBrowserNewBasicAsset",
 				EAssetTypeCategories::Basic, 
 				FirstSelectedPath, 
 				InOnNewAssetRequested, 
 				CanExecuteAssetActionsDelegate
 				);
 		}
-		MenuBuilder.EndSection(); //ContentBrowserNewBasicAsset
 
 		// Add Advanced Asset
-		MenuBuilder.BeginSection("ContentBrowserNewAdvancedAsset", LOCTEXT("CreateAdvancedAssetsMenuHeading", "Create Advanced Asset"));
 		{
+			FToolMenuSection& Section = Menu->AddSection("ContentBrowserNewAdvancedAsset", LOCTEXT("CreateAdvancedAssetsMenuHeading", "Create Advanced Asset"));
+
 			FAssetToolsModule& AssetToolsModule = FAssetToolsModule::GetModule();
 
 			TArray<FAdvancedAssetCategory> AdvancedAssetCategories;
@@ -376,11 +378,13 @@ void FNewAssetOrClassContextMenu::MakeContextMenu(
 				TArray<FFactoryItem> Factories = FindFactoriesInCategory(AdvancedAssetCategory.CategoryType);
 				if (Factories.Num() > 0)
 				{
-					MenuBuilder.AddSubMenu(
+					Section.AddSubMenu(
+						NAME_None,
 						AdvancedAssetCategory.CategoryName,
 						FText::GetEmpty(),
-						FNewMenuDelegate::CreateStatic(
+						FNewToolMenuDelegate::CreateStatic(
 							&FNewAssetOrClassContextMenu::CreateNewAssetMenuCategory, 
+							FName("Section"),
 							AdvancedAssetCategory.CategoryType, 
 							FirstSelectedPath, 
 							InOnNewAssetRequested, 
@@ -390,17 +394,15 @@ void FNewAssetOrClassContextMenu::MakeContextMenu(
 							FExecuteAction(),
 							CanExecuteAssetActionsDelegate
 						),
-						NAME_None,
 						EUserInterfaceActionType::Button
 						);
 				}
 			}
 		}
-		MenuBuilder.EndSection(); //ContentBrowserNewAdvancedAsset
 	}
 }
 
-void FNewAssetOrClassContextMenu::CreateNewAssetMenuCategory(FMenuBuilder& MenuBuilder, EAssetTypeCategories::Type AssetTypeCategory, FString InPath, FOnNewAssetRequested InOnNewAssetRequested, FCanExecuteAction InCanExecuteAction)
+void FNewAssetOrClassContextMenu::CreateNewAssetMenuCategory(UToolMenu* Menu, FName SectionName, EAssetTypeCategories::Type AssetTypeCategory, FString InPath, FOnNewAssetRequested InOnNewAssetRequested, FCanExecuteAction InCanExecuteAction)
 {
 	// Find UFactory classes that can create new objects in this category.
 	TArray<FFactoryItem> FactoriesInThisCategory = FindFactoriesInCategory(AssetTypeCategory);
@@ -433,22 +435,23 @@ void FNewAssetOrClassContextMenu::CreateNewAssetMenuCategory(FMenuBuilder& MenuB
 		SubMenu->Factories.Add(Item);
 	}
 	ParentMenuData->SortSubMenus();
-	CreateNewAssetMenus(MenuBuilder, ParentMenuData, InPath, InOnNewAssetRequested, InCanExecuteAction);
+	CreateNewAssetMenus(Menu, SectionName, ParentMenuData, InPath, InOnNewAssetRequested, InCanExecuteAction);
 }
 
-void FNewAssetOrClassContextMenu::CreateNewAssetMenus(FMenuBuilder& MenuBuilder, TSharedPtr<FCategorySubMenuItem> SubMenuData, FString InPath, FOnNewAssetRequested InOnNewAssetRequested, FCanExecuteAction InCanExecuteAction)
+void FNewAssetOrClassContextMenu::CreateNewAssetMenus(UToolMenu* Menu, FName SectionName, TSharedPtr<FCategorySubMenuItem> SubMenuData, FString InPath, FOnNewAssetRequested InOnNewAssetRequested, FCanExecuteAction InCanExecuteAction)
 {
+	FToolMenuSection& Section = Menu->FindOrAddSection(SectionName);
 	for (const FFactoryItem& FactoryItem : SubMenuData->Factories)
 	{
 		TWeakObjectPtr<UClass> WeakFactoryClass = FactoryItem.Factory->GetClass();
 
-		MenuBuilder.AddMenuEntry(
+		Section.AddEntry(FToolMenuEntry::InitMenuEntry(
+			NAME_None, 
 			FUIAction(
 				FExecuteAction::CreateStatic(&FNewAssetOrClassContextMenu::ExecuteNewAsset, InPath, WeakFactoryClass, InOnNewAssetRequested),
 				InCanExecuteAction
 			),
-			SNew(SFactoryMenuEntry, FactoryItem.Factory)
-		);
+			SNew(SFactoryMenuEntry, FactoryItem.Factory)));
 	}
 
 	if (SubMenuData->Children.Num() == 0)
@@ -456,7 +459,7 @@ void FNewAssetOrClassContextMenu::CreateNewAssetMenus(FMenuBuilder& MenuBuilder,
 		return;
 	}
 
-	MenuBuilder.AddMenuSeparator();
+	Section.AddMenuSeparator(NAME_None);
 
 	TArray<TSharedPtr<FCategorySubMenuItem>> SortedMenus;
 	SubMenuData->Children.GenerateValueArray(SortedMenus);
@@ -469,11 +472,13 @@ void FNewAssetOrClassContextMenu::CreateNewAssetMenus(FMenuBuilder& MenuBuilder,
 	{
 		check(ChildMenuData.IsValid());
 
-		MenuBuilder.AddSubMenu(
+		Section.AddSubMenu(
+			NAME_None,
 			ChildMenuData->Name,
 			FText::GetEmpty(),
-			FNewMenuDelegate::CreateStatic(
+			FNewToolMenuDelegate::CreateStatic(
 				&FNewAssetOrClassContextMenu::CreateNewAssetMenus,
+				FName("Section"),
 				ChildMenuData,
 				InPath,
 				InOnNewAssetRequested,
@@ -483,7 +488,6 @@ void FNewAssetOrClassContextMenu::CreateNewAssetMenus(FMenuBuilder& MenuBuilder,
 				FExecuteAction(),
 				InCanExecuteAction
 			),
-			NAME_None,
 			EUserInterfaceActionType::Button
 		);
 	}
diff --git a/Engine/Source/Editor/ContentBrowser/Private/NewAssetOrClassContextMenu.h b/Engine/Source/Editor/ContentBrowser/Private/NewAssetOrClassContextMenu.h
--- a/Engine/Source/Editor/ContentBrowser/Private/NewAssetOrClassContextMenu.h
+++ b/Engine/Source/Editor/ContentBrowser/Private/NewAssetOrClassContextMenu.h
@@ -10,6 +10,7 @@
 
 class FMenuBuilder;
 class UFactory;
+class UToolMenu;
 
 struct FFactoryItem
 {
@@ -39,7 +40,7 @@ public:
 
 	/** Makes the context menu widget */
 	static void MakeContextMenu(
-		FMenuBuilder& MenuBuilder, 
+		UToolMenu* Menu, 
 		const TArray<FName>& InSelectedPaths, 
 		const FOnNewAssetRequested& InOnNewAssetRequested, 
 		const FOnNewClassRequested& InOnNewClassRequested, 
@@ -50,7 +51,7 @@ public:
 
 	/** Makes the context menu widget */
 	static void MakeContextMenu(
-		FMenuBuilder& MenuBuilder, 
+		UToolMenu* Menu, 
 		const TArray<FString>& InSelectedPaths, 
 		const FOnNewAssetRequested& InOnNewAssetRequested, 
 		const FOnNewClassRequested& InOnNewClassRequested, 
@@ -61,9 +62,9 @@ public:
 
 private:
 	/** Handle creating a new asset from an asset category */
-	static void CreateNewAssetMenuCategory(FMenuBuilder& MenuBuilder, EAssetTypeCategories::Type AssetTypeCategory, FString InPath, FOnNewAssetRequested InOnNewAssetRequested, FCanExecuteAction InCanExecuteAction);
+	static void CreateNewAssetMenuCategory(UToolMenu* Menu, FName SectionName, EAssetTypeCategories::Type AssetTypeCategory, FString InPath, FOnNewAssetRequested InOnNewAssetRequested, FCanExecuteAction InCanExecuteAction);
 
-	static void CreateNewAssetMenus(FMenuBuilder& MenuBuilder, TSharedPtr<FCategorySubMenuItem> SubMenuData, FString InPath, FOnNewAssetRequested InOnNewAssetRequested, FCanExecuteAction InCanExecuteAction);
+	static void CreateNewAssetMenus(UToolMenu* Menu, FName SectionName, TSharedPtr<FCategorySubMenuItem> SubMenuData, FString InPath, FOnNewAssetRequested InOnNewAssetRequested, FCanExecuteAction InCanExecuteAction);
 
 	/** Handle when the "Import" button is clicked */
 	static void ExecuteImportAsset(FOnImportAssetRequested InOnImportAssetRequested, FString InPath);
diff --git a/Engine/Source/Editor/ContentBrowser/Private/PathContextMenu.cpp b/Engine/Source/Editor/ContentBrowser/Private/PathContextMenu.cpp
--- a/Engine/Source/Editor/ContentBrowser/Private/PathContextMenu.cpp
+++ b/Engine/Source/Editor/ContentBrowser/Private/PathContextMenu.cpp
@@ -38,7 +38,7 @@
 #include "Framework/Notifications/NotificationManager.h"
 #include "Widgets/Notifications/SNotificationList.h"
 #include "ContentBrowserCommands.h"
-
+#include "ToolMenus.h"
 
 #define LOCTEXT_NAMESPACE "ContentBrowser"
 
@@ -111,14 +111,15 @@ TSharedRef<FExtender> FPathContextMenu::MakePathViewContextMenuExtender(const TA
 		}
 	}
 	TSharedPtr<FExtender> MenuExtender = FExtender::Combine(Extenders);
-
-	MenuExtender->AddMenuExtension("FolderContext", EExtensionHook::After, TSharedPtr<FUICommandList>(), FMenuExtensionDelegate::CreateSP(this, &FPathContextMenu::MakePathViewContextMenu));
-
 	return MenuExtender.ToSharedRef();
 }
 
-void FPathContextMenu::MakePathViewContextMenu(FMenuBuilder& MenuBuilder)
+void FPathContextMenu::MakePathViewContextMenu(UToolMenu* Menu)
 {
+	// Cache any vars that are used in determining if you can execute any actions.
+	// Useful for actions whose "CanExecute" will not change or is expensive to calculate.
+	CacheCanExecuteVars();
+
 	int32 NumAssetPaths, NumClassPaths;
 	ContentBrowserUtils::CountPathTypes(SelectedPaths, NumAssetPaths, NumClassPaths);
 
@@ -129,8 +130,8 @@ void FPathContextMenu::MakePathViewContextMenu(FMenuBuilder& MenuBuilder)
 		const bool bHasClassPaths = NumClassPaths > 0;
 
 		// Common operations section //
-		MenuBuilder.BeginSection("PathViewFolderOptions", LOCTEXT("PathViewOptionsMenuHeading", "Folder Options") );
 		{
+			FToolMenuSection& Section = Menu->AddSection("PathViewFolderOptions", LOCTEXT("PathViewOptionsMenuHeading", "Folder Options") );
 			if(bHasAssetPaths)
 			{
 				FText NewAssetToolTip;
@@ -151,15 +152,15 @@ void FPathContextMenu::MakePathViewContextMenu(FMenuBuilder& MenuBuilder)
 				}
 
 				// New Asset (submenu)
-				MenuBuilder.AddSubMenu(
+				Section.AddSubMenu(
+					"NewAsset",
 					LOCTEXT( "NewAssetLabel", "New Asset" ),
 					NewAssetToolTip,
-					FNewMenuDelegate::CreateRaw( this, &FPathContextMenu::MakeNewAssetSubMenu ),
+					FNewToolMenuDelegate::CreateRaw( this, &FPathContextMenu::MakeNewAssetSubMenu ),
 					FUIAction(
 						FExecuteAction(),
 						FCanExecuteAction::CreateRaw( this, &FPathContextMenu::CanCreateAsset )
 						),
-					NAME_None,
 					EUserInterfaceActionType::Button,
 					false,
 					FSlateIcon()
@@ -186,7 +187,8 @@ void FPathContextMenu::MakePathViewContextMenu(FMenuBuilder& MenuBuilder)
 				}
 
 				// New Class
-				MenuBuilder.AddMenuEntry(
+				Section.AddMenuEntry(
+					"NewClass",
 					LOCTEXT("NewClassLabel", "New C++ Class..."),
 					NewClassToolTip,
 					FSlateIcon(FEditorStyle::GetStyleSetName(), "MainFrame.AddCodeToProject"),
@@ -198,14 +200,15 @@ void FPathContextMenu::MakePathViewContextMenu(FMenuBuilder& MenuBuilder)
 			}
 
 			// Explore
-			MenuBuilder.AddMenuEntry(
+			Section.AddMenuEntry(
+				"ExploreTooltip",
 				ContentBrowserUtils::GetExploreFolderText(),
 				LOCTEXT("ExploreTooltip", "Finds this folder on disk."),
 				FSlateIcon(),
 				FUIAction( FExecuteAction::CreateSP( this, &FPathContextMenu::ExecuteExplore ) )
 				);
 
-			MenuBuilder.AddMenuEntry(FGenericCommands::Get().Rename, NAME_None,
+			Section.AddMenuEntry(FGenericCommands::Get().Rename,
 				LOCTEXT("RenameFolder", "Rename"),
 				LOCTEXT("RenameFolderTooltip", "Rename the selected folder.")
 				);
@@ -214,10 +217,11 @@ void FPathContextMenu::MakePathViewContextMenu(FMenuBuilder& MenuBuilder)
 			if ( ContentBrowserUtils::HasCustomColors() )
 			{
 				// Set Color (submenu)
-				MenuBuilder.AddSubMenu(
+				Section.AddSubMenu(
+					"SetColor",
 					LOCTEXT("SetColor", "Set Color"),
 					LOCTEXT("SetColorTooltip", "Sets the color this folder should appear as."),
-					FNewMenuDelegate::CreateRaw( this, &FPathContextMenu::MakeSetColorSubMenu ),
+					FNewToolMenuDelegate::CreateRaw( this, &FPathContextMenu::MakeSetColorSubMenu ),
 					false,
 					FSlateIcon()
 					);
@@ -225,7 +229,8 @@ void FPathContextMenu::MakePathViewContextMenu(FMenuBuilder& MenuBuilder)
 			else
 			{
 				// Set Color
-				MenuBuilder.AddMenuEntry(
+				Section.AddMenuEntry(
+					"SetColor",
 					LOCTEXT("SetColor", "Set Color"),
 					LOCTEXT("SetColorTooltip", "Sets the color this folder should appear as."),
 					FSlateIcon(),
@@ -237,7 +242,8 @@ void FPathContextMenu::MakePathViewContextMenu(FMenuBuilder& MenuBuilder)
 			if (ContentBrowserUtils::IsFavoriteFolder(SelectedPaths[0]))
 			{
 				// Remove from favorites
-				MenuBuilder.AddMenuEntry(
+				Section.AddMenuEntry(
+					"RemoveFromFavorites",
 					LOCTEXT("RemoveFromFavorites", "Remove From Favorites"),
 					LOCTEXT("RemoveFromFavoritesTooltip", "Removes this folder from the favorites section."),
 					FSlateIcon(FEditorStyle::GetStyleSetName(), "PropertyWindow.Favorites_Disabled"),
@@ -247,7 +253,8 @@ void FPathContextMenu::MakePathViewContextMenu(FMenuBuilder& MenuBuilder)
 			else
 			{
 				// Add to favorites
-				MenuBuilder.AddMenuEntry(
+				Section.AddMenuEntry(
+					"AddToFavorites",
 					LOCTEXT("AddToFavorites", "Add To Favorites"),
 					LOCTEXT("AddToFavoritesTooltip", "Adds this folder to the favorites section for easy access."),
 					FSlateIcon(FEditorStyle::GetStyleSetName(), "PropertyWindow.Favorites_Enabled"),
@@ -255,30 +262,31 @@ void FPathContextMenu::MakePathViewContextMenu(FMenuBuilder& MenuBuilder)
 				);
 			}
 		}
-		MenuBuilder.EndSection();
 
 		if(bHasAssetPaths)
 		{
 			// Bulk operations section //
-			MenuBuilder.BeginSection("PathContextBulkOperations", LOCTEXT("AssetTreeBulkMenuHeading", "Bulk Operations") );
 			{
+				FToolMenuSection& Section = Menu->AddSection("PathContextBulkOperations", LOCTEXT("AssetTreeBulkMenuHeading", "Bulk Operations") );
+
 				// Save
-				MenuBuilder.AddMenuEntry(FContentBrowserCommands::Get().SaveAllCurrentFolder, NAME_None,
+				Section.AddMenuEntry(FContentBrowserCommands::Get().SaveAllCurrentFolder,
 					LOCTEXT("SaveFolder", "Save All"),
 					LOCTEXT("SaveFolderTooltip", "Saves all modified assets in this folder.")
 					);
 
 				// Resave
-				MenuBuilder.AddMenuEntry(FContentBrowserCommands::Get().ResaveAllCurrentFolder);
+				Section.AddMenuEntry(FContentBrowserCommands::Get().ResaveAllCurrentFolder);
 
 				// Delete
-				MenuBuilder.AddMenuEntry(FGenericCommands::Get().Delete, NAME_None,
+				Section.AddMenuEntry(FGenericCommands::Get().Delete,
 					LOCTEXT("DeleteFolder", "Delete"),
 					LOCTEXT("DeleteFolderTooltip", "Removes this folder and all assets it contains.")
 					);
 
 				// Fix Up Redirectors in Folder
-				MenuBuilder.AddMenuEntry(
+				Section.AddMenuEntry(
+					"FixUpRedirectorsInFolder",
 					LOCTEXT("FixUpRedirectorsInFolder", "Fix Up Redirectors in Folder"),
 					LOCTEXT("FixUpRedirectorsInFolderTooltip", "Finds referencers to all redirectors in the selected folders and resaves them if possible, then deletes any redirectors that had all their referencers fixed."),
 					FSlateIcon(),
@@ -288,7 +296,8 @@ void FPathContextMenu::MakePathViewContextMenu(FMenuBuilder& MenuBuilder)
 				if ( NumAssetPaths == 1 && NumClassPaths == 0 )
 				{
 					// Migrate Folder
-					MenuBuilder.AddMenuEntry(
+					Section.AddMenuEntry(
+						"MigrateFolder",
 						LOCTEXT("MigrateFolder", "Migrate..."),
 						LOCTEXT("MigrateFolderTooltip", "Copies assets found in this folder and their dependencies to another game content folder."),
 						FSlateIcon(),
@@ -296,72 +305,76 @@ void FPathContextMenu::MakePathViewContextMenu(FMenuBuilder& MenuBuilder)
 						);
 				}
 			}
-			MenuBuilder.EndSection();
 
 			// Source control section //
-			MenuBuilder.BeginSection("PathContextSourceControl", LOCTEXT("AssetTreeSCCMenuHeading", "Source Control") );
-
-			ISourceControlProvider& SourceControlProvider = ISourceControlModule::Get().GetProvider();
-			if ( SourceControlProvider.IsEnabled() )
 			{
-				// Check out
-				MenuBuilder.AddMenuEntry(
-					LOCTEXT("FolderSCCCheckOut", "Check Out"),
-					LOCTEXT("FolderSCCCheckOutTooltip", "Checks out all assets from source control which are in this folder."),
-					FSlateIcon(),
-					FUIAction(
-						FExecuteAction::CreateSP( this, &FPathContextMenu::ExecuteSCCCheckOut ),
-						FCanExecuteAction::CreateSP( this, &FPathContextMenu::CanExecuteSCCCheckOut )
+				FToolMenuSection& Section = Menu->AddSection("PathContextSourceControl", LOCTEXT("AssetTreeSCCMenuHeading", "Source Control"));
+
+				ISourceControlProvider& SourceControlProvider = ISourceControlModule::Get().GetProvider();
+				if (SourceControlProvider.IsEnabled())
+				{
+					// Check out
+					Section.AddMenuEntry(
+						"FolderSCCCheckOut",
+						LOCTEXT("FolderSCCCheckOut", "Check Out"),
+						LOCTEXT("FolderSCCCheckOutTooltip", "Checks out all assets from source control which are in this folder."),
+						FSlateIcon(),
+						FUIAction(
+							FExecuteAction::CreateSP(this, &FPathContextMenu::ExecuteSCCCheckOut),
+							FCanExecuteAction::CreateSP(this, &FPathContextMenu::CanExecuteSCCCheckOut)
 						)
 					);
 
-				// Open for Add
-				MenuBuilder.AddMenuEntry(
-					LOCTEXT("FolderSCCOpenForAdd", "Mark For Add"),
-					LOCTEXT("FolderSCCOpenForAddTooltip", "Adds all assets to source control that are in this folder and not already added."),
-					FSlateIcon(),
-					FUIAction(
-						FExecuteAction::CreateSP( this, &FPathContextMenu::ExecuteSCCOpenForAdd ),
-						FCanExecuteAction::CreateSP( this, &FPathContextMenu::CanExecuteSCCOpenForAdd )
+					// Open for Add
+					Section.AddMenuEntry(
+						"FolderSCCOpenForAdd",
+						LOCTEXT("FolderSCCOpenForAdd", "Mark For Add"),
+						LOCTEXT("FolderSCCOpenForAddTooltip", "Adds all assets to source control that are in this folder and not already added."),
+						FSlateIcon(),
+						FUIAction(
+							FExecuteAction::CreateSP(this, &FPathContextMenu::ExecuteSCCOpenForAdd),
+							FCanExecuteAction::CreateSP(this, &FPathContextMenu::CanExecuteSCCOpenForAdd)
 						)
 					);
 
-				// Check in
-				MenuBuilder.AddMenuEntry(
-					LOCTEXT("FolderSCCCheckIn", "Check In"),
-					LOCTEXT("FolderSCCCheckInTooltip", "Checks in all assets to source control which are in this folder."),
-					FSlateIcon(),
-					FUIAction(
-						FExecuteAction::CreateSP( this, &FPathContextMenu::ExecuteSCCCheckIn ),
-						FCanExecuteAction::CreateSP( this, &FPathContextMenu::CanExecuteSCCCheckIn )
+					// Check in
+					Section.AddMenuEntry(
+						"FolderSCCCheckIn",
+						LOCTEXT("FolderSCCCheckIn", "Check In"),
+						LOCTEXT("FolderSCCCheckInTooltip", "Checks in all assets to source control which are in this folder."),
+						FSlateIcon(),
+						FUIAction(
+							FExecuteAction::CreateSP(this, &FPathContextMenu::ExecuteSCCCheckIn),
+							FCanExecuteAction::CreateSP(this, &FPathContextMenu::CanExecuteSCCCheckIn)
 						)
 					);
 
-				// Sync
-				MenuBuilder.AddMenuEntry(
-					LOCTEXT("FolderSCCSync", "Sync"),
-					LOCTEXT("FolderSCCSyncTooltip", "Syncs all the assets in this folder to the latest version."),
-					FSlateIcon(),
-					FUIAction(
-						FExecuteAction::CreateSP( this, &FPathContextMenu::ExecuteSCCSync ),
-						FCanExecuteAction::CreateSP( this, &FPathContextMenu::CanExecuteSCCSync )
+					// Sync
+					Section.AddMenuEntry(
+						"FolderSCCSync",
+						LOCTEXT("FolderSCCSync", "Sync"),
+						LOCTEXT("FolderSCCSyncTooltip", "Syncs all the assets in this folder to the latest version."),
+						FSlateIcon(),
+						FUIAction(
+							FExecuteAction::CreateSP(this, &FPathContextMenu::ExecuteSCCSync),
+							FCanExecuteAction::CreateSP(this, &FPathContextMenu::CanExecuteSCCSync)
 						)
 					);
-			}
-			else
-			{
-				MenuBuilder.AddMenuEntry(
-					LOCTEXT("FolderSCCConnect", "Connect To Source Control"),
-					LOCTEXT("FolderSCCConnectTooltip", "Connect to source control to allow source control operations to be performed on content and levels."),
-					FSlateIcon(),
-					FUIAction(
-						FExecuteAction::CreateSP( this, &FPathContextMenu::ExecuteSCCConnect ),
-						FCanExecuteAction::CreateSP( this, &FPathContextMenu::CanExecuteSCCConnect )
+				}
+				else
+				{
+					Section.AddMenuEntry(
+						"FolderSCCConnect",
+						LOCTEXT("FolderSCCConnect", "Connect To Source Control"),
+						LOCTEXT("FolderSCCConnectTooltip", "Connect to source control to allow source control operations to be performed on content and levels."),
+						FSlateIcon(),
+						FUIAction(
+							FExecuteAction::CreateSP(this, &FPathContextMenu::ExecuteSCCConnect),
+							FCanExecuteAction::CreateSP(this, &FPathContextMenu::CanExecuteSCCConnect)
 						)
-					);	
+					);
+				}
 			}
-
-			MenuBuilder.EndSection();
 		}
 	}
 }
@@ -372,12 +385,12 @@ bool FPathContextMenu::CanCreateAsset() const
 	return SelectedPaths.Num() == 1 && !ContentBrowserUtils::IsClassPath(SelectedPaths[0]);
 }
 
-void FPathContextMenu::MakeNewAssetSubMenu(FMenuBuilder& MenuBuilder)
+void FPathContextMenu::MakeNewAssetSubMenu(UToolMenu* Menu)
 {
 	if ( SelectedPaths.Num() )
 	{
 		FNewAssetOrClassContextMenu::MakeContextMenu(
-			MenuBuilder, 
+			Menu,
 			SelectedPaths, 
 			OnNewAssetRequested, 
 			FNewAssetOrClassContextMenu::FOnNewClassRequested(), 
@@ -399,57 +412,63 @@ bool FPathContextMenu::CanCreateClass() const
 	return SelectedPaths.Num() == 1 && ContentBrowserUtils::IsValidPathToCreateNewClass(SelectedPaths[0]);
 }
 
-void FPathContextMenu::MakeSetColorSubMenu(FMenuBuilder& MenuBuilder)
+void FPathContextMenu::MakeSetColorSubMenu(UToolMenu* Menu)
 {
-	// New Color
-	MenuBuilder.AddMenuEntry(
-		LOCTEXT("NewColor", "New Color"),
-		LOCTEXT("NewColorTooltip", "Changes the color this folder should appear as."),
-		FSlateIcon(),
-		FUIAction( FExecuteAction::CreateSP( this, &FPathContextMenu::ExecutePickColor ) )
-		);
-
-	// Clear Color (only required if any of the selection has one)
-	if ( SelectedHasCustomColors() )
 	{
-		MenuBuilder.AddMenuEntry(
-			LOCTEXT("ClearColor", "Clear Color"),
-			LOCTEXT("ClearColorTooltip", "Resets the color this folder appears as."),
+		FToolMenuSection& Section = Menu->AddSection("Section");
+
+		// New Color
+		Section.AddMenuEntry(
+			"NewColor",
+			LOCTEXT("NewColor", "New Color"),
+			LOCTEXT("NewColorTooltip", "Changes the color this folder should appear as."),
 			FSlateIcon(),
-			FUIAction( FExecuteAction::CreateSP( this, &FPathContextMenu::ExecuteResetColor ) )
+			FUIAction(FExecuteAction::CreateSP(this, &FPathContextMenu::ExecutePickColor))
+		);
+
+		// Clear Color (only required if any of the selection has one)
+		if (SelectedHasCustomColors())
+		{
+			Section.AddMenuEntry(
+				"ClearColor",
+				LOCTEXT("ClearColor", "Clear Color"),
+				LOCTEXT("ClearColorTooltip", "Resets the color this folder appears as."),
+				FSlateIcon(),
+				FUIAction(FExecuteAction::CreateSP(this, &FPathContextMenu::ExecuteResetColor))
 			);
+		}
 	}
 
 	// Add all the custom colors the user has chosen so far
 	TArray< FLinearColor > CustomColors;
 	if ( ContentBrowserUtils::HasCustomColors( &CustomColors ) )
 	{	
-		MenuBuilder.BeginSection("PathContextCustomColors", LOCTEXT("CustomColorsExistingColors", "Existing Colors") );
 		{
+			FToolMenuSection& Section = Menu->AddSection("PathContextCustomColors", LOCTEXT("CustomColorsExistingColors", "Existing Colors") );
 			for ( int32 ColorIndex = 0; ColorIndex < CustomColors.Num(); ColorIndex++ )
 			{
 				const FLinearColor& Color = CustomColors[ ColorIndex ];
-				MenuBuilder.AddWidget(
-						SNew(SHorizontalBox)
-						+SHorizontalBox::Slot()
-						.AutoWidth()
-						.Padding(2, 0, 0, 0)
+				Section.AddEntry(FToolMenuEntry::InitWidget(
+					NAME_None,
+					SNew(SHorizontalBox)
+					+SHorizontalBox::Slot()
+					.AutoWidth()
+					.Padding(2, 0, 0, 0)
+					[
+						SNew(SButton)
+						.ButtonStyle( FEditorStyle::Get(), "Menu.Button" )
+						.OnClicked( this, &FPathContextMenu::OnColorClicked, Color )
 						[
-							SNew(SButton)
-							.ButtonStyle( FEditorStyle::Get(), "Menu.Button" )
-							.OnClicked( this, &FPathContextMenu::OnColorClicked, Color )
-							[
-								SNew(SColorBlock)
-								.Color( Color )
-								.Size( FVector2D(77,16) )
-							]
-						],
+							SNew(SColorBlock)
+							.Color( Color )
+							.Size( FVector2D(77,16) )
+						]
+					],
 					FText::GetEmpty(),
 					/*bNoIndent=*/true
-				);
+				));
 			}
 		}
-		MenuBuilder.EndSection();
 	}
 }
 
diff --git a/Engine/Source/Editor/ContentBrowser/Private/PathContextMenu.h b/Engine/Source/Editor/ContentBrowser/Private/PathContextMenu.h
--- a/Engine/Source/Editor/ContentBrowser/Private/PathContextMenu.h
+++ b/Engine/Source/Editor/ContentBrowser/Private/PathContextMenu.h
@@ -10,6 +10,7 @@ class FExtender;
 class FMenuBuilder;
 class SWidget;
 class SWindow;
+class UToolMenu;
 
 class FPathContextMenu : public TSharedFromThis<FPathContextMenu>
 {
@@ -48,13 +49,13 @@ public:
 	TSharedRef<FExtender> MakePathViewContextMenuExtender(const TArray<FString>& InSelectedPaths);
 
 	/** Makes the asset tree context menu widget */
-	void MakePathViewContextMenu(FMenuBuilder& MenuBuilder);
+	void MakePathViewContextMenu(UToolMenu* Menu);
 
 	/** Handler to check to see if creating a new asset is allowed */
 	bool CanCreateAsset() const;
 
 	/** Makes the new asset submenu */
-	void MakeNewAssetSubMenu(FMenuBuilder& MenuBuilder);
+	void MakeNewAssetSubMenu(UToolMenu* Menu);
 
 	/** Handler for when "New Class" is selected */
 	void ExecuteCreateClass();
@@ -63,7 +64,7 @@ public:
 	bool CanCreateClass() const;
 
 	/** Makes the set color submenu */
-	void MakeSetColorSubMenu(FMenuBuilder& MenuBuilder);
+	void MakeSetColorSubMenu(UToolMenu* Menu);
 
 	/** Handler for when "Migrate Folder" is selected */
 	void ExecuteMigrateFolder();
diff --git a/Engine/Source/Editor/ContentBrowser/Private/SContentBrowser.cpp b/Engine/Source/Editor/ContentBrowser/Private/SContentBrowser.cpp
--- a/Engine/Source/Editor/ContentBrowser/Private/SContentBrowser.cpp
+++ b/Engine/Source/Editor/ContentBrowser/Private/SContentBrowser.cpp
@@ -61,7 +61,8 @@
 #include "GameProjectGenerationModule.h"
 #include "Toolkits/GlobalEditorCommonCommands.h"
 #include "Subsystems/AssetEditorSubsystem.h"
-
+#include "ContentBrowserMenuContexts.h"
+#include "ToolMenus.h"
 
 #include "Brushes/SlateColorBrush.h"
 
@@ -2072,6 +2073,19 @@ FString SContentBrowser::GetCurrentPath() const
 
 TSharedRef<SWidget> SContentBrowser::MakeAddNewContextMenu(bool bShowGetContent, bool bShowImport)
 {
+	if (!UToolMenus::Get()->IsMenuRegistered("ContentBrowser.AddNewContextMenu"))
+	{
+		UToolMenu* Menu = UToolMenus::Get()->RegisterMenu("ContentBrowser.AddNewContextMenu");
+		Menu->AddDynamicSection("DynamicSection", FNewToolMenuDelegate::CreateLambda([](UToolMenu* InMenu)
+		{
+			UContentBrowserAddNewContextMenuContext* Context = InMenu->FindContext<UContentBrowserAddNewContextMenuContext>();
+			if (Context && Context->ContentBrowser.IsValid())
+			{
+				Context->ContentBrowser.Pin()->PopulateAddNewContextMenu(InMenu, Context->bShowGetContent, Context->bShowImport, Context->NumAssetPaths);
+			}
+		}));
+	}
+
 	const FSourcesData& SourcesData = AssetViewPtr->GetSourcesData();
 
 	int32 NumAssetPaths, NumClassPaths;
@@ -2098,7 +2112,33 @@ TSharedRef<SWidget> SContentBrowser::MakeAddNewContextMenu(bool bShowGetContent,
 	}
 	TSharedPtr<FExtender> MenuExtender = FExtender::Combine(Extenders);
 
-	FMenuBuilder MenuBuilder(/*bInShouldCloseWindowAfterMenuSelection=*/true, NULL, MenuExtender);
+	UContentBrowserAddNewContextMenuContext* ContextObject = NewObject<UContentBrowserAddNewContextMenuContext>();
+	ContextObject->ContentBrowser = SharedThis(this);
+	ContextObject->NumAssetPaths = NumAssetPaths;
+	ContextObject->bShowGetContent = bShowGetContent;
+	ContextObject->bShowImport = bShowImport;
+	FToolMenuContext ToolMenuContext(nullptr, MenuExtender, ContextObject);
+
+	FDisplayMetrics DisplayMetrics;
+	FSlateApplication::Get().GetCachedDisplayMetrics( DisplayMetrics );
+
+	const FVector2D DisplaySize(
+		DisplayMetrics.PrimaryDisplayWorkAreaRect.Right - DisplayMetrics.PrimaryDisplayWorkAreaRect.Left,
+		DisplayMetrics.PrimaryDisplayWorkAreaRect.Bottom - DisplayMetrics.PrimaryDisplayWorkAreaRect.Top );
+
+	return 
+		SNew(SVerticalBox)
+
+		+SVerticalBox::Slot()
+		.MaxHeight(DisplaySize.Y * 0.9)
+		[
+			UToolMenus::Get()->GenerateWidget("ContentBrowser.AddNewContextMenu", ToolMenuContext)
+		];
+}
+
+void SContentBrowser::PopulateAddNewContextMenu(class UToolMenu* Menu, bool bShowGetContent, bool bShowImport, const int32 NumAssetPaths)
+{
+	const FSourcesData& SourcesData = AssetViewPtr->GetSourcesData();
 
 	// Only add "New Folder" item if we do not have a collection selected
 	FNewAssetOrClassContextMenu::FOnNewFolderRequested OnNewFolderRequested;
@@ -2131,7 +2171,7 @@ TSharedRef<SWidget> SContentBrowser::MakeAddNewContextMenu(bool bShowGetContent,
 	FNewAssetOrClassContextMenu::FOnNewClassRequested OnNewClassRequested = FNewAssetOrClassContextMenu::FOnNewClassRequested::CreateSP(this, &SContentBrowser::NewClassRequested);
 
 	FNewAssetOrClassContextMenu::MakeContextMenu(
-		MenuBuilder, 
+		Menu,
 		SourcesData.PackagePaths, 
 		OnNewAssetRequested,
 		OnNewClassRequested,
@@ -2139,22 +2179,6 @@ TSharedRef<SWidget> SContentBrowser::MakeAddNewContextMenu(bool bShowGetContent,
 		OnImportAssetRequested,
 		OnGetContentRequested
 		);
-
-	FDisplayMetrics DisplayMetrics;
-	FSlateApplication::Get().GetCachedDisplayMetrics( DisplayMetrics );
-
-	const FVector2D DisplaySize(
-		DisplayMetrics.PrimaryDisplayWorkAreaRect.Right - DisplayMetrics.PrimaryDisplayWorkAreaRect.Left,
-		DisplayMetrics.PrimaryDisplayWorkAreaRect.Bottom - DisplayMetrics.PrimaryDisplayWorkAreaRect.Top );
-
-	return 
-		SNew(SVerticalBox)
-
-		+SVerticalBox::Slot()
-		.MaxHeight(DisplaySize.Y * 0.9)
-		[
-			MenuBuilder.MakeWidget()
-		];
 }
 
 bool SContentBrowser::IsAddNewEnabled() const
@@ -3177,9 +3201,35 @@ TSharedPtr<SWidget> SContentBrowser::GetFolderContextMenu(const TArray<FString>&
 		Extender = InMenuExtender.Execute(SelectedPaths);
 	}
 
-	const bool bInShouldCloseWindowAfterSelection = true;
-	FMenuBuilder MenuBuilder(bInShouldCloseWindowAfterSelection, Commands, Extender, true);
-	
+	if (!UToolMenus::Get()->IsMenuRegistered("ContentBrowser.FolderContextMenu"))
+	{
+		UToolMenu* Menu = UToolMenus::Get()->RegisterMenu("ContentBrowser.FolderContextMenu");
+		Menu->bCloseSelfOnly = true;
+		Menu->AddDynamicSection("Section", FNewToolMenuDelegate::CreateLambda([](UToolMenu* InMenu)
+		{
+			UContentBrowserFolderContext* Context = InMenu->FindContext<UContentBrowserFolderContext>();
+			if (Context && Context->ContentBrowser.IsValid())
+			{
+				Context->ContentBrowser.Pin()->PopulateFolderContextMenu(InMenu);
+			}
+		}));
+	}
+
+	UContentBrowserFolderContext* Context = NewObject<UContentBrowserFolderContext>();
+	Context->ContentBrowser = SharedThis(this);
+	Context->OnCreateNewFolder = InOnCreateNewFolder;
+	FToolMenuContext MenuContext(Commands, Extender, Context);
+
+	return UToolMenus::Get()->GenerateWidget("ContentBrowser.FolderContextMenu", MenuContext);
+}
+
+void SContentBrowser::PopulateFolderContextMenu(UToolMenu* Menu)
+{
+	UContentBrowserFolderContext* Context = NewObject<UContentBrowserFolderContext>();
+	check(Context);
+
+	const TArray<FString>& SelectedPaths = PathContextMenu->GetSelectedPaths();
+
 	// We can only create folders when we have a single path selected
 	const bool bCanCreateNewFolder = SelectedPaths.Num() == 1 && ContentBrowserUtils::IsValidPathToCreateNewFolder(SelectedPaths[0]);
 
@@ -3200,26 +3250,31 @@ TSharedPtr<SWidget> SContentBrowser::GetFolderContextMenu(const TArray<FString>&
 		NewFolderToolTip = LOCTEXT("NewFolderTooltip_InvalidNumberOfPaths", "Can only create folders when there is a single path selected.");
 	}
 
-	// New Folder
-	MenuBuilder.AddMenuEntry(
-		LOCTEXT("NewFolder", "New Folder"),
-		NewFolderToolTip,
-		FSlateIcon(FEditorStyle::GetStyleSetName(), "ContentBrowser.NewFolderIcon"),
-		FUIAction(
-			FExecuteAction::CreateSP( this, &SContentBrowser::CreateNewFolder, SelectedPaths.Num() > 0 ? SelectedPaths[0] : FString(), InOnCreateNewFolder ),
-			FCanExecuteAction::CreateLambda( [bCanCreateNewFolder] { return bCanCreateNewFolder; } )
+	{
+		FToolMenuSection& Section = Menu->AddSection("Section");
+
+		// New Folder
+		Section.AddMenuEntry(
+			"NewFolder",
+			LOCTEXT("NewFolder", "New Folder"),
+			NewFolderToolTip,
+			FSlateIcon(FEditorStyle::GetStyleSetName(), "ContentBrowser.NewFolderIcon"),
+			FUIAction(
+				FExecuteAction::CreateSP(this, &SContentBrowser::CreateNewFolder, SelectedPaths.Num() > 0 ? SelectedPaths[0] : FString(), Context->OnCreateNewFolder),
+				FCanExecuteAction::CreateLambda([bCanCreateNewFolder] { return bCanCreateNewFolder; })
 			)
 		);
 
-	MenuBuilder.AddMenuEntry(
-		LOCTEXT("ShowInNewContentBrowser", "Show in New Content Browser"),
-		LOCTEXT("ShowInNewContentBrowserTooltip", "Opens a new Content Browser at this folder location (at least 1 Content Browser window needs to be locked)"),
-		FSlateIcon(),
-		FUIAction(FExecuteAction::CreateSP(this, &SContentBrowser::OpenNewContentBrowser)),
-		"FolderContext"
-	);
+		Section.AddMenuEntry(
+			"FolderContext",
+			LOCTEXT("ShowInNewContentBrowser", "Show in New Content Browser"),
+			LOCTEXT("ShowInNewContentBrowserTooltip", "Opens a new Content Browser at this folder location (at least 1 Content Browser window needs to be locked)"),
+			FSlateIcon(),
+			FUIAction(FExecuteAction::CreateSP(this, &SContentBrowser::OpenNewContentBrowser))
+		);
+	}
 
-	return MenuBuilder.MakeWidget();
+	PathContextMenu->MakePathViewContextMenu(Menu);
 }
 
 void SContentBrowser::CreateNewFolder(FString FolderPath, FOnCreateNewFolder InOnCreateNewFolder)
diff --git a/Engine/Source/Editor/ContentBrowser/Private/SContentBrowser.h b/Engine/Source/Editor/ContentBrowser/Private/SContentBrowser.h
--- a/Engine/Source/Editor/ContentBrowser/Private/SContentBrowser.h
+++ b/Engine/Source/Editor/ContentBrowser/Private/SContentBrowser.h
@@ -241,6 +241,9 @@ private:
 	/** Handle creating a context menu for the "Add New" button */
 	TSharedRef<SWidget> MakeAddNewContextMenu(bool bShowGetContent, bool bShowImport);
 
+	/** Handle populating a context menu for the "Add New" button */
+	void PopulateAddNewContextMenu(class UToolMenu* Menu, bool bShowGetContent, bool bShowImport, const int32 NumAssetPaths);
+
 	/** Called to work out whether the import button should be enabled */
 	bool IsAddNewEnabled() const;
 
@@ -433,6 +436,9 @@ private:
 	/** Delegate called when generating the context menu for a folder */
 	TSharedPtr<SWidget> GetFolderContextMenu(const TArray<FString>& SelectedPaths, FContentBrowserMenuExtender_SelectedPaths InMenuExtender, FOnCreateNewFolder OnCreateNewFolder, bool bPathView);
 
+	/** Populate the context menu for a folder */
+	void PopulateFolderContextMenu(UToolMenu* Menu);
+
 	/** Delegate called to get the current selection state */
 	void GetSelectionState(TArray<FAssetData>& SelectedAssets, TArray<FString>& SelectedPaths);
 
diff --git a/Engine/Source/Editor/ContentBrowser/Public/ContentBrowserMenuContexts.h b/Engine/Source/Editor/ContentBrowser/Public/ContentBrowserMenuContexts.h
--- a/Engine/Source/Editor/ContentBrowser/Public/ContentBrowserMenuContexts.h
+++ b/Engine/Source/Editor/ContentBrowser/Public/ContentBrowserMenuContexts.h
@@ -5,12 +5,14 @@
 #include "CoreMinimal.h"
 #include "UObject/ObjectMacros.h"
 #include "UObject/Object.h"
+#include "ContentBrowserDelegates.h"
 
 #include "ContentBrowserMenuContexts.generated.h"
 
 class FAssetContextMenu;
 class IAssetTypeActions;
 class SAssetView;
+class SContentBrowser;
 
 UCLASS()
 class CONTENTBROWSER_API UContentBrowserAssetContextMenuContext : public UObject
@@ -52,3 +54,34 @@ public:
 	TWeakPtr<SAssetView> AssetView;
 };
 
+UCLASS()
+class CONTENTBROWSER_API UContentBrowserFolderContext : public UObject
+{
+	GENERATED_BODY()
+
+public:
+
+	TWeakPtr<SContentBrowser> ContentBrowser;
+	FOnCreateNewFolder OnCreateNewFolder;
+};
+
+UCLASS()
+class CONTENTBROWSER_API UContentBrowserAddNewContextMenuContext : public UObject
+{
+	GENERATED_BODY()
+
+public:
+
+	UContentBrowserAddNewContextMenuContext() :
+		NumAssetPaths(0),
+		bShowGetContent(false),
+		bShowImport(false)
+	{
+	}
+
+	TWeakPtr<SContentBrowser> ContentBrowser;
+	int32 NumAssetPaths;
+	bool bShowGetContent;
+	bool bShowImport;
+};
+
