Fix create new folder not working in content browser's left folder tree view
#rnx

diff --git a/Engine/Source/Editor/ContentBrowser/Private/SContentBrowser.cpp b/Engine/Source/Editor/ContentBrowser/Private/SContentBrowser.cpp
--- a/Engine/Source/Editor/ContentBrowser/Private/SContentBrowser.cpp
+++ b/Engine/Source/Editor/ContentBrowser/Private/SContentBrowser.cpp
@@ -3225,7 +3225,7 @@ TSharedPtr<SWidget> SContentBrowser::GetFolderContextMenu(const TArray<FString>&
 
 void SContentBrowser::PopulateFolderContextMenu(UToolMenu* Menu)
 {
-	UContentBrowserFolderContext* Context = NewObject<UContentBrowserFolderContext>();
+	UContentBrowserFolderContext* Context = Menu->FindContext<UContentBrowserFolderContext>();
 	check(Context);
 
 	const TArray<FString>& SelectedPaths = PathContextMenu->GetSelectedPaths();
