Allow plugins to control which folders are writable

diff --git a/Engine/Source/Developer/AssetTools/Private/AssetTools.cpp b/Engine/Source/Developer/AssetTools/Private/AssetTools.cpp
--- a/Engine/Source/Developer/AssetTools/Private/AssetTools.cpp
+++ b/Engine/Source/Developer/AssetTools/Private/AssetTools.cpp
@@ -157,6 +157,7 @@ UAssetToolsImpl::UAssetToolsImpl(const FObjectInitializer& ObjectInitializer)
 	, NextUserCategoryBit(EAssetTypeCategories::FirstUser)
 	, AssetClassBlacklist(MakeShared<FBlacklistNames>())
 	, FolderBlacklist(MakeShared<FBlacklistPaths>())
+	, WritableFolderBlacklist(MakeShared<FBlacklistPaths>())
 {
 	TArray<FString> SupportedTypesArray;
 	GConfig->GetArray(TEXT("AssetTools"), TEXT("SupportedAssetTypes"), SupportedTypesArray, GEditorIni);
@@ -1020,6 +1021,12 @@ TArray<UObject*> UAssetToolsImpl::ImportAssets(const FString& DestinationPath)
 
 TArray<UObject*> UAssetToolsImpl::ImportAssetsWithDialog(const FString& DestinationPath)
 {
+	if (!GetWritableFolderBlacklist()->PassesStartsWithFilter(DestinationPath))
+	{
+		NotifyBlockedByWritableFolderFilter();
+		return TArray<UObject*>();
+	}
+
 	TArray<UObject*> ReturnObjects;
 	FString FileTypes, AllExtensions;
 	TArray<UFactory*> Factories;
@@ -3159,4 +3166,30 @@ TSharedRef<FBlacklistPaths>& UAssetToolsImpl::GetFolderBlacklist()
 	return FolderBlacklist;
 }
 
+TSharedRef<FBlacklistPaths>& UAssetToolsImpl::GetWritableFolderBlacklist()
+{
+	return WritableFolderBlacklist;
+}
+
+bool UAssetToolsImpl::AllPassWritableFolderFilter(const TArray<FString>& InPaths) const
+{
+	if (WritableFolderBlacklist->HasFiltering())
+	{
+		for (const FString& Path : InPaths)
+		{
+			if (!WritableFolderBlacklist->PassesStartsWithFilter(Path))
+			{
+				return false;
+			}
+		}
+	}
+
+	return true;
+}
+
+void UAssetToolsImpl::NotifyBlockedByWritableFolderFilter() const
+{
+	FSlateNotificationManager::Get().AddNotification(FNotificationInfo(LOCTEXT("NotifyBlockedByWritableFolderFilter", "Folder is locked")));
+}
+
 #undef LOCTEXT_NAMESPACE
diff --git a/Engine/Source/Developer/AssetTools/Private/AssetTools.h b/Engine/Source/Developer/AssetTools/Private/AssetTools.h
--- a/Engine/Source/Developer/AssetTools/Private/AssetTools.h
+++ b/Engine/Source/Developer/AssetTools/Private/AssetTools.h
@@ -120,6 +120,9 @@ public:
 	virtual TArray<UFactory*> GetNewAssetFactories() const override;
 	virtual TSharedRef<FBlacklistNames>& GetAssetClassBlacklist() override;
 	virtual TSharedRef<FBlacklistPaths>& GetFolderBlacklist() override;
+	virtual TSharedRef<FBlacklistPaths>& GetWritableFolderBlacklist() override;
+	virtual bool AllPassWritableFolderFilter(const TArray<FString>& InPaths) const override;
+	virtual void NotifyBlockedByWritableFolderFilter() const;
 
 public:
 	/** Gets the asset tools singleton as a FAssetTools for asset tools module use */
@@ -194,6 +197,9 @@ private:
 
 	/** Blacklist of folder paths */
 	TSharedRef<FBlacklistPaths> FolderBlacklist;
+
+	/** Blacklist of folder paths to write to */
+	TSharedRef<FBlacklistPaths> WritableFolderBlacklist;
 };
 
 PRAGMA_ENABLE_DEPRECATION_WARNINGS
diff --git a/Engine/Source/Developer/AssetTools/Public/IAssetTools.h b/Engine/Source/Developer/AssetTools/Public/IAssetTools.h
--- a/Engine/Source/Developer/AssetTools/Public/IAssetTools.h
+++ b/Engine/Source/Developer/AssetTools/Public/IAssetTools.h
@@ -425,6 +425,15 @@ public:
 
 	/** Get folder blacklist for content browser and other systems */
 	virtual TSharedRef<FBlacklistPaths>& GetFolderBlacklist() = 0;
+
+	/** Get writable folder blacklist for content browser and other systems */
+	virtual TSharedRef<FBlacklistPaths>& GetWritableFolderBlacklist() = 0;
+
+	/** Returns true if all in list pass writable folder filter */
+	virtual bool AllPassWritableFolderFilter(const TArray<FString>& InPaths) const = 0;
+
+	/** Show notification that writable folder filter blocked an action */
+	virtual void NotifyBlockedByWritableFolderFilter() const = 0;
 };
 
 UCLASS(transient)
diff --git a/Engine/Source/Editor/ContentBrowser/Private/AssetContextMenu.cpp b/Engine/Source/Editor/ContentBrowser/Private/AssetContextMenu.cpp
--- a/Engine/Source/Editor/ContentBrowser/Private/AssetContextMenu.cpp
+++ b/Engine/Source/Editor/ContentBrowser/Private/AssetContextMenu.cpp
@@ -173,12 +173,31 @@ TSharedRef<SWidget> FAssetContextMenu::MakeContextMenu(const TArray<FAssetData>&
 			}
 			ContextObject->CommonClass = CommonClass;
 
+			ContextObject->bCanBeModified = true;
+
+			FAssetToolsModule& AssetToolsModule = FModuleManager::LoadModuleChecked<FAssetToolsModule>("AssetTools");
+			const TSharedRef<FBlacklistPaths>& WritableFolderFilter = AssetToolsModule.Get().GetWritableFolderBlacklist();
+			if (WritableFolderFilter->HasFiltering())
+			{
+				for (const UObject* SelectedObject : SelectedObjects)
+				{
+					if (SelectedObject)
+					{
+						UPackage* SelectedObjectPackage = SelectedObject->GetOutermost();
+						if (SelectedObjectPackage && !WritableFolderFilter->PassesStartsWithFilter(SelectedObjectPackage->GetFName()))
+						{
+							ContextObject->bCanBeModified = false;
+							break;
+						}
+					}
+				}
+			}
+
 			MenuName = UToolMenus::JoinMenuPaths(BaseMenuName, CommonClass->GetFName());
 
 			RegisterMenuHierarchy(CommonClass);
 
 			// Find asset actions for common class
-			FAssetToolsModule& AssetToolsModule = FModuleManager::LoadModuleChecked<FAssetToolsModule>("AssetTools");
 			TSharedPtr<IAssetTypeActions> CommonAssetTypeActions = AssetToolsModule.Get().GetAssetTypeActionsForClass(ContextObject->CommonClass).Pin();
 			if (CommonAssetTypeActions.IsValid() && CommonAssetTypeActions->HasActions(SelectedObjects))
 			{
@@ -272,7 +291,10 @@ void FAssetContextMenu::AddMenuOptions(UToolMenu* InMenu)
 	AddAssetTypeMenuOptions(InMenu, Context->SelectedObjects.Num() > 0);
 
 	// Add imported asset context menu options
-	AddImportedAssetMenuOptions(InMenu);
+	if (Context->bCanBeModified)
+	{
+		AddImportedAssetMenuOptions(InMenu);
+	}
 
 	// Add quick access to common commands.
 	AddCommonMenuOptions(InMenu);
@@ -283,17 +305,20 @@ void FAssetContextMenu::AddMenuOptions(UToolMenu* InMenu)
 	// Add reference options
 	AddReferenceMenuOptions(InMenu);
 
-	// Add copy full path options
-	AddCopyFilePathMenuOptions(InMenu);
-
 	// Add collection options
-	AddCollectionMenuOptions(InMenu);
+	if (Context->bCanBeModified)
+	{
+		AddCollectionMenuOptions(InMenu);
+	}
 
 	// Add documentation options
 	AddDocumentationMenuOptions(InMenu);
 
 	// Add source control options
-	AddSourceControlMenuOptions(InMenu);
+	if (Context->bCanBeModified)
+	{
+		AddSourceControlMenuOptions(InMenu);
+	}
 }
 
 void FAssetContextMenu::SetSelectedAssets(const TArray<FAssetData>& InSelectedAssets)
@@ -474,48 +499,57 @@ bool FAssetContextMenu::AddCommonMenuOptions(UToolMenu* Menu)
 	int32 NumAssetItems, NumClassItems;
 	ContentBrowserUtils::CountItemTypes(SelectedAssets, NumAssetItems, NumClassItems);
 
+	UContentBrowserAssetContextMenuContext* Context = Menu->FindContext<UContentBrowserAssetContextMenuContext>();
+	const bool bCanBeModified = !Context || Context->bCanBeModified;
+
 	{
 		FToolMenuSection& Section = Menu->AddSection("CommonAssetActions", LOCTEXT("CommonAssetActionsMenuHeading", "Common"));
 
 		// Edit
-		Section.AddMenuEntry(
-			"EditAsset",
-			LOCTEXT("EditAsset", "Edit..."),
-			LOCTEXT("EditAssetTooltip", "Opens the selected asset(s) for edit."),
-			FSlateIcon(FEditorStyle::GetStyleSetName(), "ContentBrowser.AssetActions.Edit"),
-			FUIAction( FExecuteAction::CreateSP(this, &FAssetContextMenu::ExecuteEditAsset) )
+		if (bCanBeModified)
+		{
+			Section.AddMenuEntry(
+				"EditAsset",
+				LOCTEXT("EditAsset", "Edit..."),
+				LOCTEXT("EditAssetTooltip", "Opens the selected asset(s) for edit."),
+				FSlateIcon(FEditorStyle::GetStyleSetName(), "ContentBrowser.AssetActions.Edit"),
+				FUIAction(FExecuteAction::CreateSP(this, &FAssetContextMenu::ExecuteEditAsset))
 			);
+		}
 	
 		// Only add these options if assets are selected
 		if (NumAssetItems > 0)
 		{
-			// Rename
-			Section.AddMenuEntry(FGenericCommands::Get().Rename,
-				LOCTEXT("Rename", "Rename"),
-				LOCTEXT("RenameTooltip", "Rename the selected asset."),
-				FSlateIcon(FEditorStyle::GetStyleSetName(), "ContentBrowser.AssetActions.Rename")
+			if (bCanBeModified)
+			{
+				// Rename
+				Section.AddMenuEntry(FGenericCommands::Get().Rename,
+					LOCTEXT("Rename", "Rename"),
+					LOCTEXT("RenameTooltip", "Rename the selected asset."),
+					FSlateIcon(FEditorStyle::GetStyleSetName(), "ContentBrowser.AssetActions.Rename")
 				);
 
-			// Duplicate
-			Section.AddMenuEntry(FGenericCommands::Get().Duplicate,
-				LOCTEXT("Duplicate", "Duplicate"),
-				LOCTEXT("DuplicateTooltip", "Create a copy of the selected asset(s)."),
-				FSlateIcon(FEditorStyle::GetStyleSetName(), "ContentBrowser.AssetActions.Duplicate")
+				// Duplicate
+				Section.AddMenuEntry(FGenericCommands::Get().Duplicate,
+					LOCTEXT("Duplicate", "Duplicate"),
+					LOCTEXT("DuplicateTooltip", "Create a copy of the selected asset(s)."),
+					FSlateIcon(FEditorStyle::GetStyleSetName(), "ContentBrowser.AssetActions.Duplicate")
 				);
 
-			// Save
-			Section.AddMenuEntry(FContentBrowserCommands::Get().SaveSelectedAsset,
-				LOCTEXT("SaveAsset", "Save"),
-				LOCTEXT("SaveAssetTooltip", "Saves the asset to file."),
-				FSlateIcon(FEditorStyle::GetStyleSetName(), "Level.SaveIcon16x")
+				// Save
+				Section.AddMenuEntry(FContentBrowserCommands::Get().SaveSelectedAsset,
+					LOCTEXT("SaveAsset", "Save"),
+					LOCTEXT("SaveAssetTooltip", "Saves the asset to file."),
+					FSlateIcon(FEditorStyle::GetStyleSetName(), "Level.SaveIcon16x")
 				);
 
-			// Delete
-			Section.AddMenuEntry(FGenericCommands::Get().Delete,
-				LOCTEXT("Delete", "Delete"),
-				LOCTEXT("DeleteTooltip", "Delete the selected assets."),
-				FSlateIcon(FEditorStyle::GetStyleSetName(), "ContentBrowser.AssetActions.Delete")
+				// Delete
+				Section.AddMenuEntry(FGenericCommands::Get().Delete,
+					LOCTEXT("Delete", "Delete"),
+					LOCTEXT("DeleteTooltip", "Delete the selected assets."),
+					FSlateIcon(FEditorStyle::GetStyleSetName(), "ContentBrowser.AssetActions.Delete")
 				);
+			}
 
 			// Asset Actions sub-menu
 			Section.AddSubMenu(
@@ -532,7 +566,7 @@ bool FAssetContextMenu::AddCommonMenuOptions(UToolMenu* Menu)
 				FSlateIcon(FEditorStyle::GetStyleSetName(), "ContentBrowser.AssetActions")
 				);
 
-			if (NumClassItems == 0)
+			if (NumClassItems == 0 && bCanBeModified)
 			{
 				// Asset Localization sub-menu
 				Section.AddSubMenu(
@@ -554,6 +588,8 @@ bool FAssetContextMenu::AddCommonMenuOptions(UToolMenu* Menu)
 
 void FAssetContextMenu::AddExploreMenuOptions(UToolMenu* Menu)
 {
+	UContentBrowserAssetContextMenuContext* Context = Menu->FindContext<UContentBrowserAssetContextMenuContext>();
+
 	FToolMenuSection& Section = Menu->AddSection("AssetContextExploreMenuOptions", LOCTEXT("AssetContextExploreMenuOptionsHeading", "Explore"));
 	{
 		// Find in Content Browser
@@ -563,40 +599,49 @@ void FAssetContextMenu::AddExploreMenuOptions(UToolMenu* Menu)
 			LOCTEXT("ShowInFolderViewTooltip", "Selects the folder that contains this asset in the Content Browser Sources Panel.")
 			);
 
-		// Find in Explorer
-		Section.AddMenuEntry(
-			"FindInExplorer",
-			ContentBrowserUtils::GetExploreFolderText(),
-			LOCTEXT("FindInExplorerTooltip", "Finds this asset on disk"),
-			FSlateIcon(FEditorStyle::GetStyleSetName(), "SystemWideCommands.FindInContentBrowser"),
-			FUIAction(
-				FExecuteAction::CreateSP( this, &FAssetContextMenu::ExecuteFindInExplorer ),
-				FCanExecuteAction::CreateSP( this, &FAssetContextMenu::CanExecuteFindInExplorer )
+		if (Context->bCanBeModified)
+		{
+			// Find in Explorer
+			Section.AddMenuEntry(
+				"FindInExplorer",
+				ContentBrowserUtils::GetExploreFolderText(),
+				LOCTEXT("FindInExplorerTooltip", "Finds this asset on disk"),
+				FSlateIcon(FEditorStyle::GetStyleSetName(), "SystemWideCommands.FindInContentBrowser"),
+				FUIAction(
+					FExecuteAction::CreateSP(this, &FAssetContextMenu::ExecuteFindInExplorer),
+					FCanExecuteAction::CreateSP(this, &FAssetContextMenu::CanExecuteFindInExplorer)
 				)
 			);
+		}
 	}
 }
 
 void FAssetContextMenu::MakeAssetActionsSubMenu(UToolMenu* Menu)
-{	
+{
+	UContentBrowserAssetContextMenuContext* Context = Menu->FindContext<UContentBrowserAssetContextMenuContext>();
+	const bool bCanBeModified = !Context || Context->bCanBeModified;
+
 	{
 		FToolMenuSection& Section = Menu->AddSection("AssetActionsSection");
 
-		// Create BP Using This
-		Section.AddMenuEntry(
-			"CreateBlueprintUsing",
-			LOCTEXT("CreateBlueprintUsing", "Create Blueprint Using This..."),
-			LOCTEXT("CreateBlueprintUsingTooltip", "Create a new Blueprint and add this asset to it"),
-			FSlateIcon(FEditorStyle::GetStyleSetName(), "LevelEditor.CreateClassBlueprint"),
-			FUIAction(
-				FExecuteAction::CreateSP(this, &FAssetContextMenu::ExecuteCreateBlueprintUsing),
-				FCanExecuteAction::CreateSP(this, &FAssetContextMenu::CanExecuteCreateBlueprintUsing)
-			)
-		);
+		if (bCanBeModified)
+		{
+			// Create BP Using This
+			Section.AddMenuEntry(
+				"CreateBlueprintUsing",
+				LOCTEXT("CreateBlueprintUsing", "Create Blueprint Using This..."),
+				LOCTEXT("CreateBlueprintUsingTooltip", "Create a new Blueprint and add this asset to it"),
+				FSlateIcon(FEditorStyle::GetStyleSetName(), "LevelEditor.CreateClassBlueprint"),
+				FUIAction(
+					FExecuteAction::CreateSP(this, &FAssetContextMenu::ExecuteCreateBlueprintUsing),
+					FCanExecuteAction::CreateSP(this, &FAssetContextMenu::CanExecuteCreateBlueprintUsing)
+				)
+			);
+		}
 
 		// Capture Thumbnail
 		FAssetToolsModule& AssetToolsModule = FModuleManager::LoadModuleChecked<FAssetToolsModule>("AssetTools");
-		if (SelectedAssets.Num() == 1 && AssetToolsModule.Get().AssetUsesGenericThumbnail(SelectedAssets[0]))
+		if (bCanBeModified && SelectedAssets.Num() == 1 && AssetToolsModule.Get().AssetUsesGenericThumbnail(SelectedAssets[0]))
 		{
 			Section.AddMenuEntry(
 				"CaptureThumbnail",
@@ -611,7 +656,7 @@ void FAssetContextMenu::MakeAssetActionsSubMenu(UToolMenu* Menu)
 		}
 
 		// Clear Thumbnail
-		if (CanClearCustomThumbnails())
+		if (bCanBeModified && CanClearCustomThumbnails())
 		{
 			Section.AddMenuEntry(
 				"ClearCustomThumbnail",
@@ -640,6 +685,7 @@ void FAssetContextMenu::MakeAssetActionsSubMenu(UToolMenu* Menu)
 	}
 
 	// MOVE ACTIONS
+	if (bCanBeModified)
 	{
 		FToolMenuSection& Section = Menu->AddSection("AssetContextMoveActions", LOCTEXT("AssetContextMoveActionsMenuHeading", "Move"));
 		bool bHasExportableAssets = false;
@@ -692,6 +738,7 @@ void FAssetContextMenu::MakeAssetActionsSubMenu(UToolMenu* Menu)
 	}
 
 	// ADVANCED ACTIONS
+	if (bCanBeModified)
 	{
 		FToolMenuSection& Section = Menu->AddSection("AssetContextAdvancedActions", LOCTEXT("AssetContextAdvancedActionsMenuHeading", "Advanced"));
 
@@ -791,7 +838,7 @@ void FAssetContextMenu::MakeAssetActionsSubMenu(UToolMenu* Menu)
 		}
 	}
 
-	if (GetDefault<UEditorExperimentalSettings>()->bTextAssetFormatSupport)
+	if (bCanBeModified && GetDefault<UEditorExperimentalSettings>()->bTextAssetFormatSupport)
 	{
 		FToolMenuSection& FormatActionsSection = Menu->AddSection("AssetContextTextAssetFormatActions", LOCTEXT("AssetContextTextAssetFormatActionsHeading", "Text Assets"));
 		{
@@ -1282,8 +1329,13 @@ void FAssetContextMenu::ExecuteOpenEditorsForAssets(TArray<FName> InAssets)
 
 bool FAssetContextMenu::AddReferenceMenuOptions(UToolMenu* Menu)
 {
-	FToolMenuSection& Section = Menu->AddSection("AssetContextReferences", LOCTEXT("ReferencesMenuHeading", "References"));
+	UContentBrowserAssetContextMenuContext* Context = Menu->FindContext<UContentBrowserAssetContextMenuContext>();
+
 	{
+		// If can not be modified change section name to prevent reference menu extensions
+		const FName SectionName = Context->bCanBeModified ? FName("AssetContextReferences") : FName("AssetContextReferencesReadOnly");
+		FToolMenuSection& Section = Menu->AddSection(SectionName, LOCTEXT("ReferencesMenuHeading", "References"));
+
 		Section.AddMenuEntry(
 			"CopyReference",
 			LOCTEXT("CopyReference", "Copy Reference"),
@@ -1291,22 +1343,17 @@ bool FAssetContextMenu::AddReferenceMenuOptions(UToolMenu* Menu)
 			FSlateIcon(),
 			FUIAction( FExecuteAction::CreateSP( this, &FAssetContextMenu::ExecuteCopyReference ) )
 			);
-	}
-
-	return true;
-}
-
-bool FAssetContextMenu::AddCopyFilePathMenuOptions(UToolMenu* Menu)
-{
-	FToolMenuSection& Section = Menu->AddSection("AssetContextReferences", LOCTEXT("ReferencesMenuHeading", "References"));
-	{
-		Section.AddMenuEntry(
-			"CopyFilePath",
-			LOCTEXT("CopyFilePath", "Copy File Path"),
-			LOCTEXT("CopyFilePathTooltip", "Copies the file paths on disk for the selected assets to the clipboard."),
-			FSlateIcon(),
-			FUIAction(FExecuteAction::CreateSP(this, &FAssetContextMenu::ExecuteCopyFilePath))
-		);
+	
+		if (Context->bCanBeModified)
+		{
+			Section.AddMenuEntry(
+				"CopyFilePath",
+				LOCTEXT("CopyFilePath", "Copy File Path"),
+				LOCTEXT("CopyFilePathTooltip", "Copies the file paths on disk for the selected assets to the clipboard."),
+				FSlateIcon(),
+				FUIAction(FExecuteAction::CreateSP(this, &FAssetContextMenu::ExecuteCopyFilePath))
+			);
+		}
 	}
 
 	return true;
@@ -2216,6 +2263,17 @@ void FAssetContextMenu::ExecuteShowAssetMetaData()
 	}
 }
 
+bool FAssetContextMenu::CanModifyPath(const FString& InPath) const
+{
+	FAssetToolsModule& AssetToolsModule = FModuleManager::LoadModuleChecked<FAssetToolsModule>("AssetTools");
+	if (!AssetToolsModule.Get().GetWritableFolderBlacklist()->PassesStartsWithFilter(InPath))
+	{
+		return false;
+	}
+
+	return true;
+}
+
 void FAssetContextMenu::ExecuteEditAsset()
 {
 	TMap<UClass*, TArray<UObject*> > SelectedAssetsByClass;
@@ -2298,13 +2356,29 @@ void FAssetContextMenu::ExecuteRename()
 		// Don't operate on Redirectors
 		if ( AssetViewSelectedAssets[0].AssetClass != UObjectRedirector::StaticClass()->GetFName() )
 		{
-			OnRenameRequested.ExecuteIfBound(AssetViewSelectedAssets[0]);
+			if (CanModifyPath(AssetViewSelectedAssets[0].PackageName.ToString()))
+			{
+				OnRenameRequested.ExecuteIfBound(AssetViewSelectedAssets[0]);
+			}
+			else
+			{
+				FAssetToolsModule& AssetToolsModule = FModuleManager::LoadModuleChecked<FAssetToolsModule>("AssetTools");
+				AssetToolsModule.Get().NotifyBlockedByWritableFolderFilter();
+			}
 		}
 	}
 
 	if ( AssetViewSelectedAssets.Num() == 0 && SelectedFolders.Num() == 1 )
 	{
-		OnRenameFolderRequested.ExecuteIfBound(SelectedFolders[0]);
+		if (CanModifyPath(SelectedFolders[0]))
+		{
+			OnRenameFolderRequested.ExecuteIfBound(SelectedFolders[0]);
+		}
+		else
+		{
+			FAssetToolsModule& AssetToolsModule = FModuleManager::LoadModuleChecked<FAssetToolsModule>("AssetTools");
+			AssetToolsModule.Get().NotifyBlockedByWritableFolderFilter();
+		}
 	}
 }
 
@@ -2324,6 +2398,17 @@ void FAssetContextMenu::ExecuteDelete()
 		}
 	}
 
+	FAssetToolsModule& AssetToolsModule = FModuleManager::LoadModuleChecked<FAssetToolsModule>("AssetTools");
+	const TSharedRef<FBlacklistPaths>& WritableFolderFilter = AssetToolsModule.Get().GetWritableFolderBlacklist();
+	const bool bHasWritableFolderFilter = WritableFolderFilter->HasFiltering();
+
+	TArray<FString> SelectedFolders = AssetView.Pin()->GetSelectedFolders();
+	if (SelectedFolders.Num() > 0 && !AssetToolsModule.Get().AllPassWritableFolderFilter(SelectedFolders))
+	{
+		AssetToolsModule.Get().NotifyBlockedByWritableFolderFilter();
+		return;
+	}
+
 	TArray< FAssetData > AssetViewSelectedAssets = AssetView.Pin()->GetSelectedAssets();
 	if(AssetViewSelectedAssets.Num() > 0)
 	{
@@ -2339,6 +2424,12 @@ void FAssetContextMenu::ExecuteDelete()
 				continue;
 			}
 
+			if (bHasWritableFolderFilter && !WritableFolderFilter->PassesStartsWithFilter(AssetData.PackageName))
+			{
+				AssetToolsModule.Get().NotifyBlockedByWritableFolderFilter();
+				return;
+			}
+
 			AssetsToDelete.Add( AssetData );
 		}
 
@@ -2348,7 +2439,6 @@ void FAssetContextMenu::ExecuteDelete()
 		}
 	}
 
-	TArray< FString > SelectedFolders = AssetView.Pin()->GetSelectedFolders();
 	if(SelectedFolders.Num() > 0)
 	{
 		FText Prompt;
diff --git a/Engine/Source/Editor/ContentBrowser/Private/AssetContextMenu.h b/Engine/Source/Editor/ContentBrowser/Private/AssetContextMenu.h
--- a/Engine/Source/Editor/ContentBrowser/Private/AssetContextMenu.h
+++ b/Engine/Source/Editor/ContentBrowser/Private/AssetContextMenu.h
@@ -118,6 +118,9 @@ private:
 
 private:
 
+	/** Is allowed to modify files or folders under this path */
+	bool CanModifyPath(const FString& InPath) const;
+
 	/** Registers all unregistered menus in the hierarchy for a class */
 	static void RegisterMenuHierarchy(UClass* InClass);
 
@@ -166,9 +169,6 @@ private:
 	/** Adds asset reference menu options to a menu builder. Returns true if any options were added. */
 	bool AddReferenceMenuOptions(UToolMenu* Menu);
 
-	/** Adds copy file path menu options to a menu builder. Returns true if any options were added. */
-	bool AddCopyFilePathMenuOptions(UToolMenu* Menu);
-
 	/** Adds asset documentation menu options to a menu builder. Returns true if any options were added. */
 	bool AddDocumentationMenuOptions(UToolMenu* Menu);
 	
diff --git a/Engine/Source/Editor/ContentBrowser/Private/ContentBrowserUtils.cpp b/Engine/Source/Editor/ContentBrowser/Private/ContentBrowserUtils.cpp
--- a/Engine/Source/Editor/ContentBrowser/Private/ContentBrowserUtils.cpp
+++ b/Engine/Source/Editor/ContentBrowser/Private/ContentBrowserUtils.cpp
@@ -472,6 +472,13 @@ void ContentBrowserUtils::RenameAsset(UObject* Asset, const FString& NewName, FT
 
 void ContentBrowserUtils::CopyAssets(const TArray<UObject*>& Assets, const FString& DestPath)
 {
+	FAssetToolsModule& AssetToolsModule = FModuleManager::LoadModuleChecked<FAssetToolsModule>("AssetTools");
+	if (!AssetToolsModule.Get().GetWritableFolderBlacklist()->PassesStartsWithFilter(DestPath))
+	{
+		AssetToolsModule.Get().NotifyBlockedByWritableFolderFilter();
+		return;
+	}
+
 	TArray<UObject*> NewObjects;
 	ObjectTools::DuplicateObjects(Assets, TEXT(""), DestPath, /*bOpenDialog=*/false, &NewObjects);
 
@@ -499,6 +506,12 @@ void ContentBrowserUtils::MoveAssets(const TArray<UObject*>& Assets, const FStri
 	check(DestPath.Len() > 0);
 
 	FAssetToolsModule& AssetToolsModule = FModuleManager::LoadModuleChecked<FAssetToolsModule>("AssetTools");
+	if (!AssetToolsModule.Get().GetWritableFolderBlacklist()->PassesStartsWithFilter(DestPath))
+	{
+		AssetToolsModule.Get().NotifyBlockedByWritableFolderFilter();
+		return;
+	}
+
 	TArray<FAssetRenameData> AssetsAndNames;
 	for ( auto AssetIt = Assets.CreateConstIterator(); AssetIt; ++AssetIt )
 	{
@@ -1557,7 +1570,13 @@ bool ContentBrowserUtils::IsValidPathToCreateNewFolder(const FString& InPath)
 {
 	// We can't currently make folders in class paths
 	// If we do later allow folders in class paths, they must only be created within modules (see IsValidPathToCreateNewClass above)
-	return !IsClassPath(InPath);
+	if (IsClassPath(InPath))
+	{
+		return false;
+	}
+
+	FAssetToolsModule& AssetToolsModule = FModuleManager::LoadModuleChecked<FAssetToolsModule>("AssetTools");
+	return AssetToolsModule.Get().GetWritableFolderBlacklist()->PassesStartsWithFilter(InPath);
 }
 
 const TSharedPtr<FLinearColor> ContentBrowserUtils::LoadColor(const FString& FolderPath)
diff --git a/Engine/Source/Editor/ContentBrowser/Private/DragDropHandler.cpp b/Engine/Source/Editor/ContentBrowser/Private/DragDropHandler.cpp
--- a/Engine/Source/Editor/ContentBrowser/Private/DragDropHandler.cpp
+++ b/Engine/Source/Editor/ContentBrowser/Private/DragDropHandler.cpp
@@ -8,6 +8,8 @@
 #include "Framework/Commands/UIAction.h"
 #include "Framework/MultiBox/MultiBoxBuilder.h"
 #include "EditorStyleSet.h"
+#include "AssetToolsModule.h"
+#include "Misc/BlacklistNames.h"
 
 #include "DragAndDrop/AssetDragDropOp.h"
 #include "ContentBrowserUtils.h"
@@ -113,6 +115,13 @@ void DragDropHandler::HandleDropOnAssetFolder(const TSharedRef<SWidget>& ParentW
 		return ContentBrowserUtils::IsClassPath(AssetPath);
 	});
 
+	FAssetToolsModule& AssetToolsModule = FModuleManager::LoadModuleChecked<FAssetToolsModule>("AssetTools");
+	if (!AssetToolsModule.Get().GetWritableFolderBlacklist()->PassesStartsWithFilter(TargetPath))
+	{
+		AssetToolsModule.Get().NotifyBlockedByWritableFolderFilter();
+		return;
+	}
+
 	FMenuBuilder MenuBuilder(/*bInShouldCloseWindowAfterMenuSelection=*/true, nullptr);
 	const FText MoveCopyHeaderString = FText::Format(LOCTEXT("AssetViewDropMenuHeading", "Move/Copy to {0}"), TargetDisplayName);
 	MenuBuilder.BeginSection("PathAssetMoveCopy", MoveCopyHeaderString);
diff --git a/Engine/Source/Editor/ContentBrowser/Private/NewAssetOrClassContextMenu.cpp b/Engine/Source/Editor/ContentBrowser/Private/NewAssetOrClassContextMenu.cpp
--- a/Engine/Source/Editor/ContentBrowser/Private/NewAssetOrClassContextMenu.cpp
+++ b/Engine/Source/Editor/ContentBrowser/Private/NewAssetOrClassContextMenu.cpp
@@ -218,6 +218,9 @@ void FNewAssetOrClassContextMenu::MakeContextMenu(
 	const bool bIsValidNewFolderPath = ContentBrowserUtils::IsValidPathToCreateNewFolder(FirstSelectedPath);
 	const bool bHasSinglePathSelected = InSelectedPaths.Num() == 1;
 
+	FAssetToolsModule& AssetToolsModule = FModuleManager::LoadModuleChecked<FAssetToolsModule>("AssetTools");
+	bool bCanBeModified = AssetToolsModule.Get().AllPassWritableFolderFilter(InSelectedPaths);
+
 	auto CanExecuteFolderActions = [NumAssetPaths, NumClassPaths, bIsValidNewFolderPath]() -> bool
 	{
 		// We can execute folder actions when we only have a single path selected, and that path is a valid path for creating a folder
@@ -291,7 +294,7 @@ void FNewAssetOrClassContextMenu::MakeContextMenu(
 	}
 
 	// Add Class
-	if(InOnNewClassRequested.IsBound())
+	if(InOnNewClassRequested.IsBound() && bCanBeModified)
 	{
 		FString ClassCreationPath = FirstSelectedPath;
 		FText NewClassToolTip;
@@ -328,7 +331,7 @@ void FNewAssetOrClassContextMenu::MakeContextMenu(
 	}
 
 	// Import
-	if (InOnImportAssetRequested.IsBound() && !FirstSelectedPath.IsEmpty())
+	if (InOnImportAssetRequested.IsBound() && !FirstSelectedPath.IsEmpty() && bCanBeModified)
 	{
 		{
 			FToolMenuSection& Section = Menu->AddSection("ContentBrowserImportAsset", LOCTEXT( "ImportAssetMenuHeading", "Import Asset" ));
@@ -346,7 +349,7 @@ void FNewAssetOrClassContextMenu::MakeContextMenu(
 	}
 
 	
-	if (InOnNewAssetRequested.IsBound())
+	if (InOnNewAssetRequested.IsBound() && bCanBeModified)
 	{
 		// Add Basic Asset
 		{
@@ -365,8 +368,6 @@ void FNewAssetOrClassContextMenu::MakeContextMenu(
 		{
 			FToolMenuSection& Section = Menu->AddSection("ContentBrowserNewAdvancedAsset", LOCTEXT("CreateAdvancedAssetsMenuHeading", "Create Advanced Asset"));
 
-			FAssetToolsModule& AssetToolsModule = FAssetToolsModule::GetModule();
-
 			TArray<FAdvancedAssetCategory> AdvancedAssetCategories;
 			AssetToolsModule.Get().GetAllAdvancedAssetCategories(/*out*/ AdvancedAssetCategories);
 			AdvancedAssetCategories.Sort([](const FAdvancedAssetCategory& A, const FAdvancedAssetCategory& B) {
diff --git a/Engine/Source/Editor/ContentBrowser/Private/PathContextMenu.cpp b/Engine/Source/Editor/ContentBrowser/Private/PathContextMenu.cpp
--- a/Engine/Source/Editor/ContentBrowser/Private/PathContextMenu.cpp
+++ b/Engine/Source/Editor/ContentBrowser/Private/PathContextMenu.cpp
@@ -39,6 +39,7 @@
 #include "Widgets/Notifications/SNotificationList.h"
 #include "ContentBrowserCommands.h"
 #include "ToolMenus.h"
+#include "ContentBrowserMenuContexts.h"
 
 #define LOCTEXT_NAMESPACE "ContentBrowser"
 
@@ -116,6 +117,8 @@ TSharedRef<FExtender> FPathContextMenu::MakePathViewContextMenuExtender(const TA
 
 void FPathContextMenu::MakePathViewContextMenu(UToolMenu* Menu)
 {
+	UContentBrowserFolderContext* Context = Menu->FindContext<UContentBrowserFolderContext>();
+
 	// Cache any vars that are used in determining if you can execute any actions.
 	// Useful for actions whose "CanExecute" will not change or is expensive to calculate.
 	CacheCanExecuteVars();
@@ -137,7 +140,7 @@ void FPathContextMenu::MakePathViewContextMenu(UToolMenu* Menu)
 				FText NewAssetToolTip;
 				if(SelectedPaths.Num() == 1)
 				{
-					if(CanCreateAsset())
+					if(CanCreateAsset() && Context->bCanBeModified)
 					{
 						NewAssetToolTip = FText::Format(LOCTEXT("NewAssetTooltip_CreateIn", "Create a new asset in {0}."), FText::FromString(SelectedPaths[0]));
 					}
@@ -152,22 +155,25 @@ void FPathContextMenu::MakePathViewContextMenu(UToolMenu* Menu)
 				}
 
 				// New Asset (submenu)
-				Section.AddSubMenu(
-					"NewAsset",
-					LOCTEXT( "NewAssetLabel", "New Asset" ),
-					NewAssetToolTip,
-					FNewToolMenuDelegate::CreateRaw( this, &FPathContextMenu::MakeNewAssetSubMenu ),
-					FUIAction(
-						FExecuteAction(),
-						FCanExecuteAction::CreateRaw( this, &FPathContextMenu::CanCreateAsset )
+				if (Context->bCanBeModified)
+				{
+					Section.AddSubMenu(
+						"NewAsset",
+						LOCTEXT("NewAssetLabel", "New Asset"),
+						NewAssetToolTip,
+						FNewToolMenuDelegate::CreateRaw(this, &FPathContextMenu::MakeNewAssetSubMenu),
+						FUIAction(
+							FExecuteAction(),
+							FCanExecuteAction::CreateRaw(this, &FPathContextMenu::CanCreateAsset)
 						),
-					EUserInterfaceActionType::Button,
-					false,
-					FSlateIcon()
+						EUserInterfaceActionType::Button,
+						false,
+						FSlateIcon()
 					);
+				}
 			}
 
-			if(bHasClassPaths)
+			if(bHasClassPaths && Context->bCanBeModified)
 			{
 				FText NewClassToolTip;
 				if(SelectedPaths.Num() == 1)
@@ -208,10 +214,13 @@ void FPathContextMenu::MakePathViewContextMenu(UToolMenu* Menu)
 				FUIAction( FExecuteAction::CreateSP( this, &FPathContextMenu::ExecuteExplore ) )
 				);
 
-			Section.AddMenuEntry(FGenericCommands::Get().Rename,
-				LOCTEXT("RenameFolder", "Rename"),
-				LOCTEXT("RenameFolderTooltip", "Rename the selected folder.")
-				);
+			if (Context->bCanBeModified)
+			{
+				Section.AddMenuEntry(FGenericCommands::Get().Rename,
+					LOCTEXT("RenameFolder", "Rename"),
+					LOCTEXT("RenameFolderTooltip", "Rename the selected folder.")
+					);
+			}
 
 			// If any colors have already been set, display color options as a sub menu
 			if ( ContentBrowserUtils::HasCustomColors() )
@@ -263,7 +272,7 @@ void FPathContextMenu::MakePathViewContextMenu(UToolMenu* Menu)
 			}
 		}
 
-		if(bHasAssetPaths)
+		if(bHasAssetPaths && Context->bCanBeModified)
 		{
 			// Bulk operations section //
 			{
diff --git a/Engine/Source/Editor/ContentBrowser/Private/SAssetView.cpp b/Engine/Source/Editor/ContentBrowser/Private/SAssetView.cpp
--- a/Engine/Source/Editor/ContentBrowser/Private/SAssetView.cpp
+++ b/Engine/Source/Editor/ContentBrowser/Private/SAssetView.cpp
@@ -1668,6 +1668,13 @@ FReply SAssetView::OnDrop( const FGeometry& MyGeometry, const FDragDropEvent& Dr
 		// Note: We don't test IsAssetPathSelected here as we need to prevent dropping assets on class paths
 		const FString DestPath = SourcesData.PackagePaths[0].ToString();
 
+		FAssetToolsModule& AssetToolsModule = FModuleManager::LoadModuleChecked<FAssetToolsModule>("AssetTools");
+		if (!AssetToolsModule.Get().GetWritableFolderBlacklist()->PassesStartsWithFilter(DestPath))
+		{
+			AssetToolsModule.Get().NotifyBlockedByWritableFolderFilter();
+			return FReply::Handled();
+		}
+
 		// If the DragDrop event is validated, continue trying to dock it to the Widget
 		bool bUnused = false;
 		if (DragDropHandler::ValidateDragDropOnAssetFolder(MyGeometry, DragDropEvent, DestPath, bUnused))
diff --git a/Engine/Source/Editor/ContentBrowser/Private/SContentBrowser.cpp b/Engine/Source/Editor/ContentBrowser/Private/SContentBrowser.cpp
--- a/Engine/Source/Editor/ContentBrowser/Private/SContentBrowser.cpp
+++ b/Engine/Source/Editor/ContentBrowser/Private/SContentBrowser.cpp
@@ -3218,6 +3218,10 @@ TSharedPtr<SWidget> SContentBrowser::GetFolderContextMenu(const TArray<FString>&
 	UContentBrowserFolderContext* Context = NewObject<UContentBrowserFolderContext>();
 	Context->ContentBrowser = SharedThis(this);
 	Context->OnCreateNewFolder = InOnCreateNewFolder;
+
+	FAssetToolsModule& AssetToolsModule = FModuleManager::LoadModuleChecked<FAssetToolsModule>("AssetTools");
+	Context->bCanBeModified = AssetToolsModule.Get().AllPassWritableFolderFilter(SelectedPaths);
+
 	FToolMenuContext MenuContext(Commands, Extender, Context);
 
 	return UToolMenus::Get()->GenerateWidget("ContentBrowser.FolderContextMenu", MenuContext);
@@ -3253,17 +3257,20 @@ void SContentBrowser::PopulateFolderContextMenu(UToolMenu* Menu)
 	{
 		FToolMenuSection& Section = Menu->AddSection("Section");
 
-		// New Folder
-		Section.AddMenuEntry(
-			"NewFolder",
-			LOCTEXT("NewFolder", "New Folder"),
-			NewFolderToolTip,
-			FSlateIcon(FEditorStyle::GetStyleSetName(), "ContentBrowser.NewFolderIcon"),
-			FUIAction(
-				FExecuteAction::CreateSP(this, &SContentBrowser::CreateNewFolder, SelectedPaths.Num() > 0 ? SelectedPaths[0] : FString(), Context->OnCreateNewFolder),
-				FCanExecuteAction::CreateLambda([bCanCreateNewFolder] { return bCanCreateNewFolder; })
-			)
-		);
+		if (Context->bCanBeModified)
+		{
+			// New Folder
+			Section.AddMenuEntry(
+				"NewFolder",
+				LOCTEXT("NewFolder", "New Folder"),
+				NewFolderToolTip,
+				FSlateIcon(FEditorStyle::GetStyleSetName(), "ContentBrowser.NewFolderIcon"),
+				FUIAction(
+					FExecuteAction::CreateSP(this, &SContentBrowser::CreateNewFolder, SelectedPaths.Num() > 0 ? SelectedPaths[0] : FString(), Context->OnCreateNewFolder),
+					FCanExecuteAction::CreateLambda([bCanCreateNewFolder] { return bCanCreateNewFolder; })
+				)
+			);
+		}
 
 		Section.AddMenuEntry(
 			"FolderContext",
diff --git a/Engine/Source/Editor/ContentBrowser/Public/ContentBrowserMenuContexts.h b/Engine/Source/Editor/ContentBrowser/Public/ContentBrowserMenuContexts.h
--- a/Engine/Source/Editor/ContentBrowser/Public/ContentBrowserMenuContexts.h
+++ b/Engine/Source/Editor/ContentBrowser/Public/ContentBrowserMenuContexts.h
@@ -31,6 +31,9 @@ public:
 	UPROPERTY()
 	UClass* CommonClass;
 
+	UPROPERTY()
+	bool bCanBeModified;
+
 	UFUNCTION(BlueprintCallable, Category="Tool Menus")
 	TArray<UObject*> GetSelectedObjects() const
 	{
@@ -61,6 +64,9 @@ class CONTENTBROWSER_API UContentBrowserFolderContext : public UObject
 
 public:
 
+	UPROPERTY()
+	bool bCanBeModified;
+
 	TWeakPtr<SContentBrowser> ContentBrowser;
 	FOnCreateNewFolder OnCreateNewFolder;
 };
diff --git a/Engine/Source/Editor/UnrealEd/Private/EditorEngine.cpp b/Engine/Source/Editor/UnrealEd/Private/EditorEngine.cpp
--- a/Engine/Source/Editor/UnrealEd/Private/EditorEngine.cpp
+++ b/Engine/Source/Editor/UnrealEd/Private/EditorEngine.cpp
@@ -4288,6 +4288,12 @@ bool UEditorEngine::IsPackageValidForAutoAdding(UPackage* InPackage, const FStri
 
 bool UEditorEngine::IsPackageOKToSave(UPackage* InPackage, const FString& InFilename, FOutputDevice* Error)
 {
+	FAssetToolsModule& AssetToolsModule = FModuleManager::LoadModuleChecked<FAssetToolsModule>("AssetTools");
+	if (InPackage && !AssetToolsModule.Get().GetWritableFolderBlacklist()->PassesStartsWithFilter(InPackage->GetName()))
+	{
+		return false;
+	}
+
 	return true;
 }
 
diff --git a/Engine/Source/Editor/UnrealEd/Private/FileHelpers.cpp b/Engine/Source/Editor/UnrealEd/Private/FileHelpers.cpp
--- a/Engine/Source/Editor/UnrealEd/Private/FileHelpers.cpp
+++ b/Engine/Source/Editor/UnrealEd/Private/FileHelpers.cpp
@@ -64,7 +64,7 @@
 #include "Engine/LevelStreaming.h"
 #include "AutoSaveUtils.h"
 #include "AssetRegistryModule.h"
-
+#include "Misc/BlacklistNames.h"
 
 DEFINE_LOG_CATEGORY_STATIC(LogFileHelpers, Log, All);
 
@@ -4028,11 +4028,17 @@ FString FEditorFileUtils::ExtractPackageName(const FString& ObjectPath)
 
 void FEditorFileUtils::GetDirtyWorldPackages(TArray<UPackage*>& OutDirtyPackages)
 {
+	FAssetToolsModule& AssetToolsModule = FModuleManager::LoadModuleChecked<FAssetToolsModule>("AssetTools");
+	const TSharedRef<FBlacklistPaths>& WritableFolderFilter = AssetToolsModule.Get().GetWritableFolderBlacklist();
+	const bool bHasWritableFolderFilter = WritableFolderFilter->HasFiltering();
+
 	for (TObjectIterator<UWorld> WorldIt; WorldIt; ++WorldIt)
 	{
 		UPackage* WorldPackage = WorldIt->GetOutermost();
 		if (!WorldPackage->HasAnyPackageFlags(PKG_PlayInEditor)
-			&& !WorldPackage->HasAnyFlags(RF_Transient))
+			&& !WorldPackage->HasAnyFlags(RF_Transient)
+			&& (!bHasWritableFolderFilter || WritableFolderFilter->PassesStartsWithFilter(WorldPackage->GetName()))
+			)
 		{
 			if (WorldPackage->IsDirty())
 			{
@@ -4068,6 +4074,10 @@ void FEditorFileUtils::GetDirtyWorldPackages(TArray<UPackage*>& OutDirtyPackages
 
 void FEditorFileUtils::GetDirtyContentPackages(TArray<UPackage*>& OutDirtyPackages)
 {
+	FAssetToolsModule& AssetToolsModule = FModuleManager::LoadModuleChecked<FAssetToolsModule>("AssetTools");
+	const TSharedRef<FBlacklistPaths>& WritableFolderFilter = AssetToolsModule.Get().GetWritableFolderBlacklist();
+	const bool bHasWritableFolderFilter = WritableFolderFilter->HasFiltering();
+
 	// Make a list of all content packages that we should save
 	for (TObjectIterator<UPackage> It; It; ++It)
 	{
@@ -4093,8 +4103,19 @@ void FEditorFileUtils::GetDirtyContentPackages(TArray<UPackage*>& OutDirtyPackag
 			// Ignore map packages, they are caught above.
 			bShouldIgnorePackage |= bIsMapPackage;
 
-			// Ignore packages with long, invalid names. This culls out packages with paths in read-only roots such as /Temp.
-			bShouldIgnorePackage |= (!FPackageName::IsShortPackageName(Package->GetFName()) && !FPackageName::IsValidLongPackageName(Package->GetName(), /*bIncludeReadOnlyRoots=*/false));
+			if (!bShouldIgnorePackage)
+			{
+				FString PackageName = Package->GetName();
+
+				// Ignore packages with long, invalid names. This culls out packages with paths in read-only roots such as /Temp.
+				bShouldIgnorePackage |= (!FPackageName::IsShortPackageName(Package->GetFName()) && !FPackageName::IsValidLongPackageName(PackageName, /*bIncludeReadOnlyRoots=*/false));
+
+				// Ignore packages that cannot be saved due to a custom filter
+				if (!bShouldIgnorePackage && bHasWritableFolderFilter)
+				{
+					bShouldIgnorePackage |= (!WritableFolderFilter->PassesStartsWithFilter(PackageName));
+				}
+			}
 		}
 
 		if (!bShouldIgnorePackage)
diff --git a/Engine/Source/Editor/UnrealEd/Private/Subsystems/AssetEditorSubsystem.cpp b/Engine/Source/Editor/UnrealEd/Private/Subsystems/AssetEditorSubsystem.cpp
--- a/Engine/Source/Editor/UnrealEd/Private/Subsystems/AssetEditorSubsystem.cpp
+++ b/Engine/Source/Editor/UnrealEd/Private/Subsystems/AssetEditorSubsystem.cpp
@@ -22,6 +22,7 @@
 #include "Interfaces/IAnalyticsProvider.h"
 #include "Misc/FeedbackContext.h"
 #include "Misc/ConfigCacheIni.h"
+#include "Misc/BlacklistNames.h"
 
 
 #define LOCTEXT_NAMESPACE "AssetEditorSubsystem"
@@ -274,15 +275,23 @@ bool UAssetEditorSubsystem::OpenEditorForAsset(UObject* Asset, const EToolkitMod
 	//    being edited within, we should decide whether to disallow "Edit Here" in that case, or to close the old asset
 	//    editor and summon it in the new level editor, or to just foreground the old level editor (current behavior)
 
+	FAssetToolsModule& AssetToolsModule = FModuleManager::LoadModuleChecked<FAssetToolsModule>(TEXT("AssetTools"));
+
 	const bool bBringToFrontIfOpen = true;
 
-	// Don't open asset editors for cooked packages
 	if (UPackage* Package = Asset->GetOutermost())
 	{
+		// Don't open asset editors for cooked packages
 		if (Package->bIsCookedForEditor)
 		{
 			return false;
 		}
+
+		if (!AssetToolsModule.Get().GetWritableFolderBlacklist()->PassesStartsWithFilter(Package->GetName()))
+		{
+			AssetToolsModule.Get().NotifyBlockedByWritableFolderFilter();
+			return false;
+		}
 	}
 
 	AssetEditorRequestOpenEvent.Broadcast(Asset);
@@ -302,8 +311,6 @@ bool UAssetEditorSubsystem::OpenEditorForAsset(UObject* Asset, const EToolkitMod
 
 	UE_LOG(LogAssetEditorSubsystem, Log, TEXT("Opening Asset editor for %s"), *Asset->GetFullName());
 
-	FAssetToolsModule& AssetToolsModule = FModuleManager::LoadModuleChecked<FAssetToolsModule>(TEXT("AssetTools"));
-
 	TWeakPtr<IAssetTypeActions> AssetTypeActions = AssetToolsModule.Get().GetAssetTypeActionsForClass(Asset->GetClass());
 
 	EToolkitMode::Type ActualToolkitMode = ToolkitMode;
diff --git a/Engine/Source/Editor/UnrealEd/Private/Subsystems/ImportSubsystem.cpp b/Engine/Source/Editor/UnrealEd/Private/Subsystems/ImportSubsystem.cpp
--- a/Engine/Source/Editor/UnrealEd/Private/Subsystems/ImportSubsystem.cpp
+++ b/Engine/Source/Editor/UnrealEd/Private/Subsystems/ImportSubsystem.cpp
@@ -7,7 +7,7 @@
 #include "ObjectTools.h"
 #include "FileHelpers.h"
 #include "EditorReimportHandler.h"
-
+#include "Misc/BlacklistNames.h"
 
 class FImportFilesByPath : public IImportSubsystemTask
 {
@@ -120,6 +120,13 @@ void UImportSubsystem::Deinitialize()
 
 void UImportSubsystem::ImportNextTick(const TArray<FString>& Files, const FString& DestinationPath)
 {
+	FAssetToolsModule& AssetToolsModule = FModuleManager::LoadModuleChecked<FAssetToolsModule>("AssetTools");
+	if (AssetToolsModule.Get().GetWritableFolderBlacklist()->PassesStartsWithFilter(DestinationPath))
+	{
+		AssetToolsModule.Get().NotifyBlockedByWritableFolderFilter();
+		return;
+	}
+
 	PendingTasks.Enqueue(MakeShared<FImportFilesByPath>(Files, DestinationPath));
 	GEditor->GetTimerManager()->SetTimerForNextTick(this, &UImportSubsystem::HandleNextTick);
 }
