Fixed batch Content Browser filter actions not correctly updating frontend filters

diff --git a/Engine/Source/Editor/ContentBrowser/Private/SFilterList.cpp b/Engine/Source/Editor/ContentBrowser/Private/SFilterList.cpp
--- a/Engine/Source/Editor/ContentBrowser/Private/SFilterList.cpp
+++ b/Engine/Source/Editor/ContentBrowser/Private/SFilterList.cpp
@@ -671,6 +671,10 @@ void SFilterList::EnableAllFilters()
 	for (const TSharedRef<SFilter>& Filter : Filters)
 	{
 		Filter->SetEnabled(true, false);
+		if (const TSharedPtr<FFrontendFilter>& FrontendFilter = Filter->GetFrontendFilter())
+		{
+			SetFrontendFilterActive(FrontendFilter.ToSharedRef(), true);
+		}
 	}
 
 	OnFilterChanged.ExecuteIfBound();
@@ -681,6 +685,10 @@ void SFilterList::DisableAllFilters()
 	for (const TSharedRef<SFilter>& Filter : Filters)
 	{
 		Filter->SetEnabled(false, false);
+		if (const TSharedPtr<FFrontendFilter>& FrontendFilter = Filter->GetFrontendFilter())
+		{
+			SetFrontendFilterActive(FrontendFilter.ToSharedRef(), false);
+		}
 	}
 
 	OnFilterChanged.ExecuteIfBound();
@@ -1012,6 +1020,7 @@ TSharedRef<SFilter> SFilterList::AddFilter(const TSharedRef<FFrontendFilter>& Fr
 		.FrontendFilter(FrontendFilter)
 		.OnFilterChanged( this, &SFilterList::FrontendFilterChanged, FrontendFilter )
 		.OnRequestRemove(this, &SFilterList::RemoveFilterAndUpdate)
+		.OnRequestEnableOnly(this, &SFilterList::EnableOnlyThisFilter)
 		.OnRequestEnableAll(this, &SFilterList::EnableAllFilters)
 		.OnRequestDisableAll(this, &SFilterList::DisableAllFilters)
 		.OnRequestRemoveAll(this, &SFilterList::RemoveAllFilters)
@@ -1065,6 +1074,10 @@ void SFilterList::EnableOnlyThisFilter(const TSharedRef<SFilter>& FilterToEnable
 	{
 		bool bEnable = Filter == FilterToEnable;
 		Filter->SetEnabled(bEnable, /*ExecuteOnFilterChange*/false);
+		if (const TSharedPtr<FFrontendFilter>& FrontendFilter = Filter->GetFrontendFilter())
+		{
+			SetFrontendFilterActive(FrontendFilter.ToSharedRef(), bEnable);
+		}
 	}
 
 	OnFilterChanged.ExecuteIfBound();
