Access Violation syncing Content Browser directories with pending changes to Levels under World Composition

Filter out redirectors before processing packages to reload

diff --git a/Engine/Source/Editor/ContentBrowser/Private/ContentBrowserUtils.cpp b/Engine/Source/Editor/ContentBrowser/Private/ContentBrowserUtils.cpp
--- a/Engine/Source/Editor/ContentBrowser/Private/ContentBrowserUtils.cpp
+++ b/Engine/Source/Editor/ContentBrowser/Private/ContentBrowserUtils.cpp
@@ -2225,11 +2225,15 @@ void ContentBrowserUtils::SyncPathsFromSourceControl(const TArray<FString>& Cont
 			TSet<FName> UniquePackageNames;
 			for (const FAssetData& Asset : AssetList)
 			{
-				bool bWasInSet = false;
-				UniquePackageNames.Add(Asset.PackageName, &bWasInSet);
-				if (!bWasInSet)
+				// Don't operate on Redirectors
+				if (Asset.AssetClass != UObjectRedirector::StaticClass()->GetFName())
 				{
-					PackageNames.Add(Asset.PackageName.ToString());
+					bool bWasInSet = false;
+					UniquePackageNames.Add(Asset.PackageName, &bWasInSet);
+					if (!bWasInSet)
+					{
+						PackageNames.Add(Asset.PackageName.ToString());
+					}
 				}
 			}
 		}
