Relaxed project name assumptions when estimating cook length for an asset

This was originally added to make content portable, however there's no guarantee that you're not going to migrate content to a deeper directory so it's mostly just an annoyance. The restriction is kept for project plugins though, as those may be migrated without changing their directory structure.

diff --git a/Engine/Source/Editor/ContentBrowser/Private/ContentBrowserUtils.cpp b/Engine/Source/Editor/ContentBrowser/Private/ContentBrowserUtils.cpp
--- a/Engine/Source/Editor/ContentBrowser/Private/ContentBrowserUtils.cpp
+++ b/Engine/Source/Editor/ContentBrowser/Private/ContentBrowserUtils.cpp
@@ -1846,20 +1846,17 @@ int32 ContentBrowserUtils::GetPackageLengthForCooking(const FString& PackageName
 	int32 AbsoluteCookPathToAssetLength = 0;
 
 	FString RelativePathToAsset;
-
-	const FString AbsolutePath = bIsEngineAsset ? AbsoluteEnginePath : AbsoluteGamePath;
-
-	const FString& AbsoluteCookPath = bIsEngineAsset ? AbsoluteEngineCookPath : AbsoluteGameCookPath;
-
-	if(FPackageName::TryConvertLongPackageNameToFilename(PackageName, RelativePathToAsset, FPackageName::GetAssetPackageExtension()))
+	if (FPackageName::TryConvertLongPackageNameToFilename(PackageName, RelativePathToAsset, FPackageName::GetAssetPackageExtension()))
 	{
+		const FString AbsolutePath = bIsEngineAsset ? AbsoluteEnginePath : AbsoluteGamePath;
+		const FString& AbsoluteCookPath = bIsEngineAsset ? AbsoluteEngineCookPath : AbsoluteGameCookPath;
+
 		const FString AbsolutePathToAsset = FPaths::ConvertRelativePathToFull(RelativePathToAsset);
 
 		FString AssetPathWithinCookDir = AbsolutePathToAsset;
 		FPaths::RemoveDuplicateSlashes(AssetPathWithinCookDir);
 		AssetPathWithinCookDir.RemoveFromStart(AbsolutePath, ESearchCase::CaseSensitive);
 
-
 		if (IsInternalBuild)
 		{
 			// We assume a constant size for the build machine base path, so strip either the root or game path from the start
@@ -1876,8 +1873,8 @@ int32 ContentBrowserUtils::GetPackageLengthForCooking(const FString& PackageName
 			
 			FString AbsoluteBuildMachineCookPathToAsset = FString(TEXT("D:/BuildFarm/buildmachine_++depot+UE4-Releases+4.10")) / CookDirWithoutBasePath / AssetPathWithinCookDir;
 
-			// only add game name padding if it is not an engine asset, otherwise it is considered portable already
-			if(!bIsEngineAsset)
+			// only add game name padding to project plugins so that they can ported to other projects
+			if (bIsPluginAsset && bIsProjectAsset)
 			{
 				AbsoluteBuildMachineCookPathToAsset.ReplaceInline(*GameName, *GameNamePadded, ESearchCase::CaseSensitive);
 			}
@@ -1889,8 +1886,8 @@ int32 ContentBrowserUtils::GetPackageLengthForCooking(const FString& PackageName
 			// Test that the package can be cooked based on the current project path
 			FString AbsoluteCookPathToAsset = AbsoluteCookPath / AssetPathWithinCookDir;
 
-			// only add game name padding if it is not an engine asset, otherwise it is considered portable already
-			if (!bIsEngineAsset)
+			// only add game name padding to project plugins so that they can ported to other projects
+			if (bIsPluginAsset && bIsProjectAsset)
 			{
 				AbsoluteCookPathToAsset.ReplaceInline(*GameName, *GameNamePadded, ESearchCase::CaseSensitive);
 			}
@@ -1902,6 +1899,7 @@ int32 ContentBrowserUtils::GetPackageLengthForCooking(const FString& PackageName
 	{
 		UE_LOG(LogContentBrowser, Error, TEXT("Package Name '%' is not a valid path and cannot be converted to a filename"), *PackageName);
 	}
+
 	return AbsoluteCookPathToAssetLength;
 }
 
