Fix for text asset menu items in Content browser.

diff --git a/Engine/Source/Editor/ContentBrowser/Private/AssetContextMenu.cpp b/Engine/Source/Editor/ContentBrowser/Private/AssetContextMenu.cpp
--- a/Engine/Source/Editor/ContentBrowser/Private/AssetContextMenu.cpp
+++ b/Engine/Source/Editor/ContentBrowser/Private/AssetContextMenu.cpp
@@ -793,7 +793,8 @@ void FAssetContextMenu::MakeAssetActionsSubMenu(UToolMenu* Menu)
 				FUIAction(FExecuteAction::CreateSP(this, &FAssetContextMenu::ExportSelectedAssetsToText))
 			);
 
-			MenuBuilder.AddMenuEntry(
+			FormatActionsSection.AddMenuEntry(
+				"ViewSelectedAssetAsText",
 				LOCTEXT("ViewSelectedAssetAsText", "View as text"),
 				LOCTEXT("ViewSelectedAssetAsTextTooltip", "Opens a window showing the selected asset in text format"),
 				FSlateIcon(),
@@ -801,7 +802,8 @@ void FAssetContextMenu::MakeAssetActionsSubMenu(UToolMenu* Menu)
 					FCanExecuteAction::CreateSP(this, &FAssetContextMenu::CanViewSelectedAssetAsText))
 			);
 
-			MenuBuilder.AddMenuEntry(
+			FormatActionsSection.AddMenuEntry(
+				"ViewSelectedAssetAsText",
 				LOCTEXT("TextFormatRountrip", "Run Text Asset Roundtrip"),
 				LOCTEXT("TextFormatRountripTooltip", "Save the select asset backwards or forwards between text and binary formats and check for determinism"),
 				FSlateIcon(),
