Removed redundant loop

#rnx

diff --git a/Engine/Source/Editor/ContentBrowser/Private/SAssetView.cpp b/Engine/Source/Editor/ContentBrowser/Private/SAssetView.cpp
--- a/Engine/Source/Editor/ContentBrowser/Private/SAssetView.cpp
+++ b/Engine/Source/Editor/ContentBrowser/Private/SAssetView.cpp
@@ -1,6 +1,7 @@
 // Copyright Epic Games, Inc. All Rights Reserved.
 
 #include "SAssetView.h"
+#include "Algo/Transform.h"
 #include "HAL/FileManager.h"
 #include "UObject/UnrealType.h"
 #include "Widgets/SOverlay.h"
@@ -2689,15 +2690,8 @@ void SAssetView::ProcessRecentlyAddedAssets()
 		bool bNeedsRefresh = false;
 
 		TSet<FName> ExistingObjectPaths;
-		for (const FAssetData& AssetItem : AssetItems)
-		{
-			ExistingObjectPaths.Add(AssetItem.ObjectPath);
-		}
-
-		for (const FAssetData& AssetItem : AssetItems)
-		{
-			ExistingObjectPaths.Add(AssetItem.ObjectPath);
-		}
+		ExistingObjectPaths.Reserve(AssetItems.Num());
+		Algo::Transform(AssetItems, ExistingObjectPaths, &FAssetData::ObjectPath);
 
 		int32 AssetIdx = 0;
 		for ( ; AssetIdx < FilteredRecentlyAddedAssets.Num(); ++AssetIdx )
