Fix asset view with an active filter killing the editor performance when assets are added.

In a project with a lot of assets (let's say ~400K assets), SAssetView::ProcessRecentlyAddedAssets() would process a single asset during it's alloted time slice. The recreation of the set on each call would make the editor run at >10fps until the whole array is processed, at one asset per tick.

It's easy to reproduce by starting the editor and clicking on the "Content" root in the Content Browser, then typing anything in the Filter box while the asset discovery is in progress. The asset discovery will complete but the editor performance will be degraded until the filter is cleared.

With these changes, all new assets are processed in the allowed time with a single call to ProcessRecentlyAddedAssets()

diff --git a/Engine/Source/Editor/ContentBrowser/Private/SAssetView.cpp b/Engine/Source/Editor/ContentBrowser/Private/SAssetView.cpp
--- a/Engine/Source/Editor/ContentBrowser/Private/SAssetView.cpp
+++ b/Engine/Source/Editor/ContentBrowser/Private/SAssetView.cpp
@@ -1516,22 +1516,25 @@ void SAssetView::ProcessQueriedItems(const double TickStartTime)
 	const bool bFlushFullBuffer = TickStartTime < 0;
 
 	bool ListNeedsRefresh = false;
-	int32 AssetIndex = 0;
-	for (AssetIndex = QueriedAssetItems.Num() - 1; AssetIndex >= 0; AssetIndex--)
+
+	for (auto AssetIter = QueriedAssetItems.CreateIterator(); AssetIter; ++AssetIter)
 	{
-		if (!OnShouldFilterAsset.Execute(QueriedAssetItems[AssetIndex]))
+		const FAssetData& AssetData = *AssetIter;
+
+		if (!OnShouldFilterAsset.Execute(AssetData))
 		{
-			AssetItems.Add(QueriedAssetItems[AssetIndex]);
+			AssetItems.Add(AssetData);
 
-			if (!IsFrontendFilterActive() || PassesCurrentFrontendFilter(QueriedAssetItems[AssetIndex]))
+			if (!IsFrontendFilterActive() || PassesCurrentFrontendFilter(AssetData))
 			{
-				const FAssetData& AssetData = QueriedAssetItems[AssetIndex];
 				FilteredAssetItems.Add(MakeShareable(new FAssetViewAsset(AssetData)));
 				ListNeedsRefresh = true;
 				bPendingSortFilteredItems = true;
 			}
 		}
 
+		AssetIter.RemoveCurrent();
+
 		// Check to see if we have run out of time in this tick
 		if (!bFlushFullBuffer && (FPlatformTime::Seconds() - TickStartTime) > MaxSecondsPerFrame)
 		{
@@ -1539,16 +1542,6 @@ void SAssetView::ProcessQueriedItems(const double TickStartTime)
 		}
 	}
 
-	// Trim the results array
-	if (AssetIndex > 0)
-	{
-		QueriedAssetItems.RemoveAt(AssetIndex, QueriedAssetItems.Num() - AssetIndex);
-	}
-	else
-	{
-		QueriedAssetItems.Reset();
-	}
-
 	if (ListNeedsRefresh)
 	{
 		RefreshList();
@@ -1999,7 +1992,7 @@ void SAssetView::RefreshSourceItems()
 	RelevantThumbnails.Reset();
 	Folders.Reset();
 
-	TArray<FAssetData>& Items = OnShouldFilterAsset.IsBound() ? QueriedAssetItems : AssetItems;
+	FAssetDataSet& Items = OnShouldFilterAsset.IsBound() ? QueriedAssetItems : AssetItems;
 
 	const FBlacklistPaths* FolderBlacklistToUse = (FolderBlacklist.IsValid() && FolderBlacklist->HasFiltering()) ? FolderBlacklist.Get() : nullptr;
 	const FBlacklistNames* AssetClassBlacklistToUse = (AssetClassBlacklist.IsValid() && AssetClassBlacklist->HasFiltering()) ? AssetClassBlacklist.Get() : nullptr;
@@ -2029,11 +2022,11 @@ void SAssetView::RefreshSourceItems()
 					continue;
 				}
 
-				int32 Index = Items.Emplace(*ObjIt);
+				FSetElementId Index = Items.Emplace(*ObjIt);
 				const FAssetData& AssetData = Items[Index];
 				if (!InitialAssetFilter.PassesRedirectorMainAssetFilter(AssetData))
 				{
-					Items.RemoveAtSwap(Index, 1, false);
+					Items.Remove(Index);
 					continue;
 				}
 
@@ -2107,7 +2100,9 @@ void SAssetView::RefreshSourceItems()
 		else
 		{
 			// Add assets found in the asset registry
-			AssetRegistryModule.Get().GetAssets(Filter, Items);
+			TArray<FAssetData> Assets;
+			AssetRegistryModule.Get().GetAssets(Filter, Assets);
+			Items.Append(Assets);
 		}
 
 		if ( bFilterAllowsClasses )
@@ -2137,14 +2132,16 @@ void SAssetView::RefreshSourceItems()
 		// Add any custom assets 
 		if (OnGetCustomSourceAssets.IsBound())
 		{
-			OnGetCustomSourceAssets.Execute(Filter, Items);
+			TArray<FAssetData> Assets;
+			OnGetCustomSourceAssets.Execute(Filter, Assets);
+			Items.Append(Assets);
 		}
 
-		for (int32 AssetIdx = Items.Num() - 1; AssetIdx >= 0; --AssetIdx)
+		for (auto AssetIter = QueriedAssetItems.CreateIterator(); AssetIter; ++AssetIter)
 		{
-			if (!InitialAssetFilter.PassesFilter(Items[AssetIdx]))
+			if (!InitialAssetFilter.PassesFilter(*AssetIter))
 			{
-				Items.RemoveAtSwap(AssetIdx);
+				AssetIter.RemoveCurrent();
 			}
 		}
 	}
@@ -2291,15 +2288,7 @@ void SAssetView::RefreshFilteredItems()
 
 			if (bGatherAssetTypeCount)
 			{
-				int32* TypeCount = AssetTypeCount.Find(AssetData.AssetClass);
-				if (TypeCount)
-				{
-					(*TypeCount)++;
-				}
-				else
-				{
-					AssetTypeCount.Add(AssetData.AssetClass, 1);
-				}
+				AssetTypeCount.FindOrAdd(AssetData.AssetClass)++;
 			}
 		}
 	}
@@ -2689,23 +2678,23 @@ void SAssetView::ProcessRecentlyAddedAssets()
 		double TickStartTime = FPlatformTime::Seconds();
 		bool bNeedsRefresh = false;
 
-		TSet<FName> ExistingObjectPaths;
-		ExistingObjectPaths.Reserve(AssetItems.Num());
-		Algo::Transform(AssetItems, ExistingObjectPaths, &FAssetData::ObjectPath);
-
 		int32 AssetIdx = 0;
 		for ( ; AssetIdx < FilteredRecentlyAddedAssets.Num(); ++AssetIdx )
 		{
+			if ((FPlatformTime::Seconds() - TickStartTime) > MaxSecondsPerFrame)
+			{
+				break;
+			}
+
 			const FAssetData& AssetData = FilteredRecentlyAddedAssets[AssetIdx];
-			if ( !ExistingObjectPaths.Contains(AssetData.ObjectPath) )
+			if ( !AssetItems.Find(AssetData.ObjectPath) )
 			{
 				if ( AssetData.AssetClass != UObjectRedirector::StaticClass()->GetFName() || AssetData.IsUAsset() )
 				{
 					if ( !OnShouldFilterAsset.IsBound() || !OnShouldFilterAsset.Execute(AssetData) )
 					{
 						// Add the asset to the list
-						int32 AddedAssetIdx = AssetItems.Add(AssetData);
-						ExistingObjectPaths.Add(AssetData.ObjectPath);
+						AssetItems.Add(AssetData);
 						if (!IsFrontendFilterActive() || PassesCurrentFrontendFilter(AssetData))
 						{
 							FilteredAssetItems.Add(MakeShareable(new FAssetViewAsset(AssetData)));
@@ -2715,13 +2704,6 @@ void SAssetView::ProcessRecentlyAddedAssets()
 					}
 				}
 			}
-
-			if ( (FPlatformTime::Seconds() - TickStartTime) > MaxSecondsPerFrame)
-			{
-				// Increment the index to properly trim the buffer below
-				++AssetIdx;
-				break;
-			}
 		}
 
 		// Trim the results array
@@ -2831,17 +2813,7 @@ void SAssetView::OnFolderPopulated(const FString& Path)
 
 void SAssetView::RemoveAssetByPath( const FName& ObjectPath )
 {
-	bool bFoundAsset = false;
-	for (int32 AssetIdx = 0; AssetIdx < AssetItems.Num(); ++AssetIdx)
-	{
-		if ( AssetItems[AssetIdx].ObjectPath == ObjectPath )
-		{
-			// Found the asset in the cached list, remove it
-			AssetItems.RemoveAt(AssetIdx);
-			bFoundAsset = true;
-			break;
-		}
-	}
+	bool bFoundAsset = AssetItems.Remove(ObjectPath) > 0;
 
 	if ( bFoundAsset )
 	{
@@ -2862,17 +2834,8 @@ void SAssetView::RemoveAssetByPath( const FName& ObjectPath )
 	}
 	else
 	{
-		//Make sure we don't have the item still queued up for processing
-		for (int32 AssetIdx = 0; AssetIdx < QueriedAssetItems.Num(); ++AssetIdx)
-		{
-			if ( QueriedAssetItems[AssetIdx].ObjectPath == ObjectPath )
-			{
-				// Found the asset in the cached list, remove it
-				QueriedAssetItems.RemoveAt(AssetIdx);
-				bFoundAsset = true;
-				break;
-			}
-		}
+		// Make sure we don't have the item still queued up for processing
+		QueriedAssetItems.Remove(ObjectPath);
 	}
 }
 
diff --git a/Engine/Source/Editor/ContentBrowser/Private/SAssetView.h b/Engine/Source/Editor/ContentBrowser/Private/SAssetView.h
--- a/Engine/Source/Editor/ContentBrowser/Private/SAssetView.h
+++ b/Engine/Source/Editor/ContentBrowser/Private/SAssetView.h
@@ -829,10 +829,26 @@ public:
 	FOnFolderPathChanged OnFolderPathChanged;
 
 private:
+	struct FAssetDataKeyFuncs : BaseKeyFuncs<FAssetData, FName>
+	{
+		static FORCEINLINE const FName& GetSetKey(const FAssetData& AssetData)
+		{
+			return AssetData.ObjectPath;
+		}
+		static FORCEINLINE bool Matches(const FName& A, const FName& B)
+		{
+			return A.GetDisplayIndex() == B.GetDisplayIndex();
+		}
+		static FORCEINLINE uint32 GetKeyHash(const FName& Key)
+		{
+			return GetTypeHash(Key.GetDisplayIndex());
+		}
+	};
+	typedef TSet<FAssetData, FAssetDataKeyFuncs> FAssetDataSet;
 
 	/** The asset items being displayed in the view and the filtered list */
-	TArray<FAssetData> QueriedAssetItems;
-	TArray<FAssetData> AssetItems;
+	FAssetDataSet QueriedAssetItems;
+	FAssetDataSet AssetItems;
 	TArray<TSharedPtr<FAssetViewItem>> FilteredAssetItems;
 
 	/** The folder items being displayed in the view */
