Fix content browser class blacklist not being applied

diff --git a/Engine/Source/Editor/ContentBrowser/Private/SAssetView.cpp b/Engine/Source/Editor/ContentBrowser/Private/SAssetView.cpp
--- a/Engine/Source/Editor/ContentBrowser/Private/SAssetView.cpp
+++ b/Engine/Source/Editor/ContentBrowser/Private/SAssetView.cpp
@@ -2137,7 +2137,7 @@ void SAssetView::RefreshSourceItems()
 			Items.Append(Assets);
 		}
 
-		for (auto AssetIter = QueriedAssetItems.CreateIterator(); AssetIter; ++AssetIter)
+		for (auto AssetIter = Items.CreateIterator(); AssetIter; ++AssetIter)
 		{
 			if (!InitialAssetFilter.PassesFilter(*AssetIter))
 			{
