Adding default value for ForceShowEngineContent

diff --git a/Engine/Source/Editor/ContentBrowser/Private/SAssetView.h b/Engine/Source/Editor/ContentBrowser/Private/SAssetView.h
--- a/Engine/Source/Editor/ContentBrowser/Private/SAssetView.h
+++ b/Engine/Source/Editor/ContentBrowser/Private/SAssetView.h
@@ -70,6 +70,7 @@ public:
 		, _ShowPathInColumnView(false)
 		, _ShowTypeInColumnView(true)
 		, _SortByPathInColumnView(false)
+		, _ForceShowEngineContent(false)
 		{}
 
 		/** Called to check if an asset should be filtered out by external code */
