Fix thumbnail editing mode not working

diff --git a/Engine/Source/Editor/ContentBrowser/Private/SThumbnailEditModeTools.cpp b/Engine/Source/Editor/ContentBrowser/Private/SThumbnailEditModeTools.cpp
--- a/Engine/Source/Editor/ContentBrowser/Private/SThumbnailEditModeTools.cpp
+++ b/Engine/Source/Editor/ContentBrowser/Private/SThumbnailEditModeTools.cpp
@@ -191,7 +191,7 @@ FReply SThumbnailEditModeTools::OnMouseButtonDown( const FGeometry& MyGeometry,
 				bModifiedThumbnailWhileDragging = false;
 				DragStartLocation = FIntPoint(MouseEvent.GetScreenSpacePosition().X, MouseEvent.GetScreenSpacePosition().Y);
 
-				return FReply::Handled().CaptureMouse( AsShared() ).UseHighPrecisionMouseMovement( AsShared() );
+				return FReply::Handled().CaptureMouse(AsShared()).UseHighPrecisionMouseMovement(AsShared()).PreventThrottling();
 			}
 		}
 
@@ -276,7 +276,7 @@ FReply SThumbnailEditModeTools::OnMouseMove( const FGeometry& MyGeometry, const
 			AssetThumbnail.Pin()->RefreshThumbnail();
 		}
 
-		return FReply::Handled();
+		return FReply::Handled().PreventThrottling();
 	}
 
 	return FReply::Unhandled();
