Path view widget can now optionally filter out read-only folder paths

diff --git a/Engine/Source/Editor/ContentBrowser/Private/SPathPicker.cpp b/Engine/Source/Editor/ContentBrowser/Private/SPathPicker.cpp
--- a/Engine/Source/Editor/ContentBrowser/Private/SPathPicker.cpp
+++ b/Engine/Source/Editor/ContentBrowser/Private/SPathPicker.cpp
@@ -33,6 +33,7 @@ void SPathPicker::Construct( const FArguments& InArgs )
 		.FocusSearchBoxWhenOpened(InArgs._PathPickerConfig.bFocusSearchBoxWhenOpened)
 		.AllowContextMenu(InArgs._PathPickerConfig.bAllowContextMenu)
 		.AllowClassesFolder(InArgs._PathPickerConfig.bAllowClassesFolder)
+		.AllowReadOnlyFolders(InArgs._PathPickerConfig.bAllowReadOnlyFolders)
 		.SelectionMode(ESelectionMode::Single)
 	];
 
diff --git a/Engine/Source/Editor/ContentBrowser/Private/SPathView.cpp b/Engine/Source/Editor/ContentBrowser/Private/SPathView.cpp
--- a/Engine/Source/Editor/ContentBrowser/Private/SPathView.cpp
+++ b/Engine/Source/Editor/ContentBrowser/Private/SPathView.cpp
@@ -69,6 +69,7 @@ void SPathView::Construct( const FArguments& InArgs )
 	OnGetFolderContextMenu = InArgs._OnGetFolderContextMenu;
 	OnGetPathContextMenuExtender = InArgs._OnGetPathContextMenuExtender;
 	bAllowClassesFolder = InArgs._AllowClassesFolder;
+	bAllowReadOnlyFolders = InArgs._AllowReadOnlyFolders;
 	PreventTreeItemChangedDelegateCount = 0;
 	TreeTitle = LOCTEXT("AssetTreeTitle", "Asset Tree");
 	if ( InArgs._FocusSearchBoxWhenOpened )
@@ -78,6 +79,7 @@ void SPathView::Construct( const FArguments& InArgs )
 
 	FAssetToolsModule& AssetToolsModule = FModuleManager::LoadModuleChecked<FAssetToolsModule>("AssetTools");
 	FolderBlacklist = AssetToolsModule.Get().GetFolderBlacklist();
+	WritableFolderBlacklist = AssetToolsModule.Get().GetWritableFolderBlacklist();
 
 	// Listen for when view settings are changed
 	UContentBrowserSettings::OnSettingChanged().AddSP(this, &SPathView::HandleSettingChanged);
@@ -1113,19 +1115,24 @@ void SPathView::Populate()
 	PathList.Add(UserDeveloperFolder);
 
 	// Remove paths of localized assets, if not displaying localized assets.
-	if (!(GetDefault<UContentBrowserSettings>()->GetDisplayL10NFolder()))
-	{
-		PathList.RemoveAll([](const FString& Path) -> bool
-		{
-			return ContentBrowserUtils::IsLocalizationFolder(Path);
-		});
-	}
+	const bool bDisplayL10NFolder = GetDefault<UContentBrowserSettings>()->GetDisplayL10NFolder();
+	const bool bCheckWritableFolderBlacklist = !bAllowReadOnlyFolders && WritableFolderBlacklist.IsValid() && WritableFolderBlacklist->HasFiltering();
 
 	// we have a text filter, expand all parents of matching folders
 	for ( int32 PathIdx = 0; PathIdx < PathList.Num(); ++PathIdx)
 	{
 		const FString& Path = PathList[PathIdx];
 
+		if (!bDisplayL10NFolder && ContentBrowserUtils::IsLocalizationFolder(Path))
+		{
+			continue;
+		}
+
+		if (bCheckWritableFolderBlacklist && !WritableFolderBlacklist->PassesStartsWithFilter(Path))
+		{
+			continue;
+		}
+
 		// by sending the whole path we deliberately include any children
 		// of successful hits in the filtered list. 
 		if ( SearchBoxFolderFilter->PassesFilter( Path ) )
diff --git a/Engine/Source/Editor/ContentBrowser/Private/SPathView.h b/Engine/Source/Editor/ContentBrowser/Private/SPathView.h
--- a/Engine/Source/Editor/ContentBrowser/Private/SPathView.h
+++ b/Engine/Source/Editor/ContentBrowser/Private/SPathView.h
@@ -37,6 +37,7 @@ public:
 		, _ShowSeparator(true)
 		, _AllowContextMenu(true)
 		, _AllowClassesFolder(false)
+		, _AllowReadOnlyFolders(true)
 		, _SelectionMode( ESelectionMode::Multi )
 		{}
 
@@ -70,6 +71,9 @@ public:
 		/** If false, the classes folder will be suppressed */
 		SLATE_ARGUMENT( bool, AllowClassesFolder )
 
+		/** If true, read only folders will be displayed */
+		SLATE_ARGUMENT( bool, AllowReadOnlyFolders )
+
 		/** The selection mode for the tree view */
 		SLATE_ARGUMENT( ESelectionMode::Type, SelectionMode )
 
@@ -331,6 +335,9 @@ protected:
 	/** Blacklist filter to hide folders */
 	TSharedPtr<FBlacklistPaths> FolderBlacklist;
 
+	/** Writable folder filter */
+	TSharedPtr<FBlacklistPaths> WritableFolderBlacklist;
+
 private:
 
 	/** The paths that were last reported by OnPathExpanded event. Used in preserving expansion when filtering folders */
@@ -354,6 +361,9 @@ private:
 	/** If false, the classes folder will not be added to the tree automatically */
 	bool bAllowClassesFolder;
 
+	/** If true, read only folders will be displayed */
+	bool bAllowReadOnlyFolders;
+
 	/** The title of this path view */
 	FText TreeTitle;
 };
diff --git a/Engine/Source/Editor/ContentBrowser/Public/IContentBrowserSingleton.h b/Engine/Source/Editor/ContentBrowser/Public/IContentBrowserSingleton.h
--- a/Engine/Source/Editor/ContentBrowser/Public/IContentBrowserSingleton.h
+++ b/Engine/Source/Editor/ContentBrowser/Public/IContentBrowserSingleton.h
@@ -331,6 +331,9 @@ struct FPathPickerConfig
 	/** If true, will allow class folders to be shown in the picker */
 	bool bAllowClassesFolder;
 
+	/** If true, will allow read-only folders to be shown in the picker */
+	bool bAllowReadOnlyFolders;
+
 	/** If true, will add the path specified in DefaultPath to the tree if it doesn't exist already */
 	bool bAddDefaultPath;
 
@@ -338,6 +341,7 @@ struct FPathPickerConfig
 		: bFocusSearchBoxWhenOpened(true)
 		, bAllowContextMenu(true)
 		, bAllowClassesFolder(false)
+		, bAllowReadOnlyFolders(true)
 		, bAddDefaultPath(false)
 	{}
 };
diff --git a/Engine/Source/Editor/UnrealEd/Private/Dialogs/DlgPickAssetPath.cpp b/Engine/Source/Editor/UnrealEd/Private/Dialogs/DlgPickAssetPath.cpp
--- a/Engine/Source/Editor/UnrealEd/Private/Dialogs/DlgPickAssetPath.cpp
+++ b/Engine/Source/Editor/UnrealEd/Private/Dialogs/DlgPickAssetPath.cpp
@@ -20,11 +20,13 @@ void SDlgPickAssetPath::Construct(const FArguments& InArgs)
 {
 	AssetPath = FText::FromString(FPackageName::GetLongPackagePath(InArgs._DefaultAssetPath.ToString()));
 	AssetName = FText::FromString(FPackageName::GetLongPackageAssetName(InArgs._DefaultAssetPath.ToString()));
+	bAllowReadOnlyFolders = InArgs._AllowReadOnlyFolders;
 
 	FPathPickerConfig PathPickerConfig;
 	PathPickerConfig.DefaultPath = AssetPath.ToString();
 	PathPickerConfig.OnPathSelected = FOnPathSelected::CreateSP(this, &SDlgPickAssetPath::OnPathChange);
 	PathPickerConfig.bAddDefaultPath = true;
+	PathPickerConfig.bAllowReadOnlyFolders = bAllowReadOnlyFolders;
 
 	FContentBrowserModule& ContentBrowserModule = FModuleManager::LoadModuleChecked<FContentBrowserModule>("ContentBrowser");
 
diff --git a/Engine/Source/Editor/UnrealEd/Public/Dialogs/DlgPickAssetPath.h b/Engine/Source/Editor/UnrealEd/Public/Dialogs/DlgPickAssetPath.h
--- a/Engine/Source/Editor/UnrealEd/Public/Dialogs/DlgPickAssetPath.h
+++ b/Engine/Source/Editor/UnrealEd/Public/Dialogs/DlgPickAssetPath.h
@@ -13,10 +13,15 @@ class SDlgPickAssetPath : public SWindow
 {
 public:
 	SLATE_BEGIN_ARGS(SDlgPickAssetPath)
-	{
-	}
+		: _AllowReadOnlyFolders(true)
+		{}
+
 		SLATE_ARGUMENT(FText, Title)
 		SLATE_ARGUMENT(FText, DefaultAssetPath)
+
+		/** If true, read only folders will be displayed */
+		SLATE_ARGUMENT(bool, AllowReadOnlyFolders)
+
 	SLATE_END_ARGS()
 
 	UNREALED_API SDlgPickAssetPath()
@@ -48,6 +53,9 @@ protected:
 	EAppReturnType::Type UserResponse;
 	FText AssetPath;
 	FText AssetName;
+
+	/** If true, read only folders will be displayed */
+	bool bAllowReadOnlyFolders;
 };
 
 #undef LOCTEXT_NAMESPACE
