Avoided an editor crash when the user developer directory cannot be found

#rnx

diff --git a/Engine/Source/Editor/ContentBrowser/Private/SPathView.cpp b/Engine/Source/Editor/ContentBrowser/Private/SPathView.cpp
--- a/Engine/Source/Editor/ContentBrowser/Private/SPathView.cpp
+++ b/Engine/Source/Editor/ContentBrowser/Private/SPathView.cpp
@@ -1110,9 +1110,12 @@ void SPathView::Populate()
 	// Add any class paths we discovered
 	PathList.Append(ClassFolders);
 
-	// Add the user developer folder
-	const FString UserDeveloperFolder = FPackageName::FilenameToLongPackageName(FPaths::GameUserDeveloperDir().LeftChop(1));
-	PathList.Add(UserDeveloperFolder);
+	// Add the user developer folder. This occasionally fails but we have not been able to reproduce it to find out why.
+	FString UserDeveloperFolder;
+	if (ensure(FPackageName::TryConvertFilenameToLongPackageName(FPaths::GameUserDeveloperDir().LeftChop(1), UserDeveloperFolder)))
+	{
+		PathList.Add(UserDeveloperFolder);
+	}
 
 	// Remove paths of localized assets, if not displaying localized assets.
 	const bool bDisplayL10NFolder = GetDefault<UContentBrowserSettings>()->GetDisplayL10NFolder();
