Fixed content browser splitter resize hit detection

diff --git a/Engine/Source/Editor/ContentBrowser/Private/SContentBrowser.cpp b/Engine/Source/Editor/ContentBrowser/Private/SContentBrowser.cpp
--- a/Engine/Source/Editor/ContentBrowser/Private/SContentBrowser.cpp
+++ b/Engine/Source/Editor/ContentBrowser/Private/SContentBrowser.cpp
@@ -547,6 +547,7 @@ void SContentBrowser::Construct( const FArguments& InArgs, const FName& InInstan
 							SAssignNew(PathFavoriteSplitterPtr, SSplitter)
 							.Style(FEditorStyle::Get(), "ContentBrowser.Splitter")
 							.PhysicalSplitterHandleSize(2.0f)
+							.HitDetectionSplitterHandleSize(2.0f)
 							.Orientation(EOrientation::Orient_Vertical)
 							.MinimumSlotHeight(70.0f)
 							.Visibility( this, &SContentBrowser::GetSourcesViewVisibility )
