Fix save file dialog clipping.

diff --git a/Engine/Source/Editor/ContentBrowser/Private/SAssetDialog.cpp b/Engine/Source/Editor/ContentBrowser/Private/SAssetDialog.cpp
--- a/Engine/Source/Editor/ContentBrowser/Private/SAssetDialog.cpp
+++ b/Engine/Source/Editor/ContentBrowser/Private/SAssetDialog.cpp
@@ -216,9 +216,10 @@ void SAssetDialog::Construct(const FArguments& InArgs, const FSharedAssetDialogC
 			];
 
 		ContentBox->AddSlot()
-			.FillHeight(1)
+			// 1.5 => 1px bigger than 'Save' buttons. 1.45 => 'g' gets clipped.
+			.FillHeight(1.475) 
 			.VAlign(VAlign_Center)
-			.Padding(0, 2, 0, 2)
+			.Padding(0, 0, 0, 0)
 			[
 				SAssignNew(NameEditableText, SEditableTextBox)
 				.Text(this, &SAssetDialog::GetAssetNameText)
