Fix save file dialog clipping.

diff --git a/Engine/Source/Editor/ContentBrowser/Private/SAssetDialog.cpp b/Engine/Source/Editor/ContentBrowser/Private/SAssetDialog.cpp
--- a/Engine/Source/Editor/ContentBrowser/Private/SAssetDialog.cpp
+++ b/Engine/Source/Editor/ContentBrowser/Private/SAssetDialog.cpp
@@ -216,8 +216,7 @@ void SAssetDialog::Construct(const FArguments& InArgs, const FSharedAssetDialogC
 			];
 
 		ContentBox->AddSlot()
-			// 1.5 => 1px bigger than 'Save' buttons. 1.45 => 'g' gets clipped.
-			.FillHeight(1.475) 
+			.AutoHeight() 
 			.VAlign(VAlign_Center)
 			.Padding(0, 0, 0, 0)
 			[
