Initial pass on content browser reskin

diff --git a/Engine/Plugins/Editor/AssetManagerEditor/Source/AssetManagerEditor/Private/ReferenceViewer/SReferenceViewer.cpp b/Engine/Plugins/Editor/AssetManagerEditor/Source/AssetManagerEditor/Private/ReferenceViewer/SReferenceViewer.cpp
--- a/Engine/Plugins/Editor/AssetManagerEditor/Source/AssetManagerEditor/Private/ReferenceViewer/SReferenceViewer.cpp
+++ b/Engine/Plugins/Editor/AssetManagerEditor/Source/AssetManagerEditor/Private/ReferenceViewer/SReferenceViewer.cpp
@@ -115,7 +115,7 @@ void SReferenceViewer::Construct(const FArguments& InArgs)
 				.OnClicked(this, &SReferenceViewer::BackClicked)
 				.IsEnabled(this, &SReferenceViewer::IsBackEnabled)
 				[
-					SNew(SImage) .Image(FEditorStyle::GetBrush("ContentBrowser.HistoryBack"))
+					SNew(SImage) .Image(FEditorStyle::GetBrush("Icons.CircleArrowLeft"))
 				]
 			]
 
@@ -133,7 +133,7 @@ void SReferenceViewer::Construct(const FArguments& InArgs)
 				.OnClicked(this, &SReferenceViewer::ForwardClicked)
 				.IsEnabled(this, &SReferenceViewer::IsForwardEnabled)
 				[
-					SNew(SImage) .Image(FEditorStyle::GetBrush("ContentBrowser.HistoryForward"))
+					SNew(SImage) .Image(FEditorStyle::GetBrush("Icons.CircleArrowRight"))
 				]
 			]
 
diff --git a/Engine/Plugins/Editor/MaterialAnalyzer/Source/Private/SMaterialAnalyzer.cpp b/Engine/Plugins/Editor/MaterialAnalyzer/Source/Private/SMaterialAnalyzer.cpp
--- a/Engine/Plugins/Editor/MaterialAnalyzer/Source/Private/SMaterialAnalyzer.cpp
+++ b/Engine/Plugins/Editor/MaterialAnalyzer/Source/Private/SMaterialAnalyzer.cpp
@@ -419,7 +419,7 @@ TSharedRef< ITableRow > SMaterialAnalyzer::OnGenerateSuggestionRow(TSharedPtr<FP
 							.Padding(2.0f)
 							[
 								SNew(SImage)
-								.Image(FEditorStyle::GetBrush("ContentBrowser.AddCollectionButtonIcon"))
+								.Image(FAppStyle::Get().GetBrush("Icons.Plus"))
 								.ColorAndOpacity(FSlateColor::UseForeground())
 							]
 							+SHorizontalBox::Slot()
diff --git a/Engine/Plugins/Media/MediaPlayerEditor/Source/MediaPlayerEditor/Private/Widgets/SMediaPlayerEditorPlaylist.cpp b/Engine/Plugins/Media/MediaPlayerEditor/Source/MediaPlayerEditor/Private/Widgets/SMediaPlayerEditorPlaylist.cpp
--- a/Engine/Plugins/Media/MediaPlayerEditor/Source/MediaPlayerEditor/Private/Widgets/SMediaPlayerEditorPlaylist.cpp
+++ b/Engine/Plugins/Media/MediaPlayerEditor/Source/MediaPlayerEditor/Private/Widgets/SMediaPlayerEditorPlaylist.cpp
@@ -155,7 +155,7 @@ void SMediaPlayerEditorPlaylist::Construct(const FArguments& InArgs, UMediaPlaye
 							.ToolTipText(LOCTEXT("PreviousPlaylistItemButtonToolTip", "Jump to the previous item in the playlist"))
 							[
 								SNew(SImage)
-									.Image(FEditorStyle::GetBrush("ContentBrowser.HistoryBack"))
+									.Image(FEditorStyle::GetBrush("Icons.CircleArrowLeft"))
 							]
 					]
 
@@ -177,7 +177,7 @@ void SMediaPlayerEditorPlaylist::Construct(const FArguments& InArgs, UMediaPlaye
 							.ToolTipText(LOCTEXT("NextPlaylistItemButtonToolTip", "Jump to the next item in the playlist"))
 							[
 								SNew(SImage)
-									.Image(FEditorStyle::GetBrush("ContentBrowser.HistoryForward"))
+								.Image(FEditorStyle::GetBrush("Icons.CircleArrowRight"))
 							]
 					]
 			]
diff --git a/Engine/Source/Developer/AssetTools/Private/AssetViewUtils.cpp b/Engine/Source/Developer/AssetTools/Private/AssetViewUtils.cpp
--- a/Engine/Source/Developer/AssetTools/Private/AssetViewUtils.cpp
+++ b/Engine/Source/Developer/AssetTools/Private/AssetViewUtils.cpp
@@ -1190,7 +1190,8 @@ bool AssetViewUtils::HasCustomColors( TArray< FLinearColor >* OutColors )
 FLinearColor AssetViewUtils::GetDefaultColor()
 {
 	// The default tint the folder should appear as
-	return FLinearColor::Gray;
+	static const FName FolderColorName("ContentBrowser.DefaultFolderColor");
+	return FAppStyle::Get().GetSlateColor(FolderColorName).GetSpecifiedColor();
 }
 
 static const auto CVarMaxFullPathLength = 
diff --git a/Engine/Source/Developer/LogVisualizer/Private/LogVisualizerStyle.cpp b/Engine/Source/Developer/LogVisualizer/Private/LogVisualizerStyle.cpp
--- a/Engine/Source/Developer/LogVisualizer/Private/LogVisualizerStyle.cpp
+++ b/Engine/Source/Developer/LogVisualizer/Private/LogVisualizerStyle.cpp
@@ -166,9 +166,6 @@ TSharedRef< FSlateStyleSet > FLogVisualizerStyle::Create()
 			.SetCheckedPressedImage(IMAGE_BRUSH("ContentBrowser/FilterChecked", FVector2D(10.0f, 14.0f), FLinearColor(0.5f, 0.5f, 0.5f, 1.0f)));
 		/* ... and add the new style */
 		Style.Set("ContentBrowser.FilterButton", ContentBrowserFilterButtonCheckBoxStyle);
-
-		Style.Set("ContentBrowser.FilterNameFont", DEFAULT_FONT("Regular", 9));
-		//Style.Set("ContentBrowser.FilterNameFont", FEditorStyle::Get().GetFontStyle("ContentBrowser.FilterNameFont"));
 	}
 
 	//Sequencer
diff --git a/Engine/Source/Developer/LogVisualizer/Private/SFilterWidget.cpp b/Engine/Source/Developer/LogVisualizer/Private/SFilterWidget.cpp
--- a/Engine/Source/Developer/LogVisualizer/Private/SFilterWidget.cpp
+++ b/Engine/Source/Developer/LogVisualizer/Private/SFilterWidget.cpp
@@ -50,8 +50,6 @@ void SFilterWidget::Construct(const FArguments& InArgs)
 				[
 					SNew(STextBlock)
 					.ColorAndOpacity(this, &SFilterWidget::GetFilterNameColorAndOpacity)
-					.Font(FLogVisualizerStyle::Get().GetFontStyle("ContentBrowser.FilterNameFont"))
-					.ShadowOffset(FVector2D(1.f, 1.f))
 					.Text(this, &SFilterWidget::GetCaptionString)
 				]
 			]
diff --git a/Engine/Source/Editor/AssetTagsEditor/Private/SAssetTagItem.cpp b/Engine/Source/Editor/AssetTagsEditor/Private/SAssetTagItem.cpp
--- a/Engine/Source/Editor/AssetTagsEditor/Private/SAssetTagItem.cpp
+++ b/Engine/Source/Editor/AssetTagsEditor/Private/SAssetTagItem.cpp
@@ -161,8 +161,6 @@ void SAssetTagItem::Construct(const FArguments& InArgs)
 	if (InArgs._OnNameCommitted.IsBound())
 	{
 		NameWidget = SAssignNew(InlineRenameWidget, SInlineEditableTextBlock)
-			.Font(FEditorStyle::GetFontStyle("ContentBrowser.AssetTagNameFont", StyleSpecifier))
-			.ShadowOffset(FVector2D(1.0f, 1.0f))
 			.Text(InArgs._DisplayName)
 			.HighlightText(InArgs._HighlightText)
 			.OnBeginTextEdit(InArgs._OnBeginNameEdit)
@@ -174,8 +172,6 @@ void SAssetTagItem::Construct(const FArguments& InArgs)
 	else
 	{
 		NameWidget = SNew(STextBlock)
-			.Font(FEditorStyle::GetFontStyle("ContentBrowser.AssetTagNameFont", StyleSpecifier))
-			.ShadowOffset(FVector2D(1.0f, 1.0f))
 			.Text(InArgs._DisplayName)
 			.HighlightText(InArgs._HighlightText);
 	}
@@ -184,32 +180,23 @@ void SAssetTagItem::Construct(const FArguments& InArgs)
 
 	ChildSlot
 	[
-		SNew(SBorder)
-		.Padding(0.0f)
-		.BorderBackgroundColor(FLinearColor(0.2f, 0.2f, 0.2f, 0.2f))
-		.BorderImage(FEditorStyle::GetBrush("ContentBrowser.AssetTagBackground"))
+		SAssignNew(HBox, SHorizontalBox)
 		.ToolTip(AssetTagToolTip)
+		+SHorizontalBox::Slot()
+		.AutoWidth()
 		[
-			SAssignNew(HBox, SHorizontalBox)
-
-			+SHorizontalBox::Slot()
-			.AutoWidth()
-			[
-				SAssignNew(CheckBox, SCheckBox)
-				.Style(FEditorStyle::Get(), "ContentBrowser.AssetTagButton", StyleSpecifier)
-				.IsEnabled(IsCheckBoxEnabled)
-				.IsChecked(InArgs._IsChecked)
-				.OnCheckStateChanged(InArgs._OnCheckStateChanged)
-				.ForegroundColor(this, &SAssetTagItem::GetCheckBoxForegroundColor)
-				.ToolTipText(this, &SAssetTagItem::GetCheckBoxTooltipText)
-			]
+			SAssignNew(CheckBox, SCheckBox)
+			.IsEnabled(IsCheckBoxEnabled)
+			.IsChecked(InArgs._IsChecked)
+			.OnCheckStateChanged(InArgs._OnCheckStateChanged)
+			.ToolTipText(this, &SAssetTagItem::GetCheckBoxTooltipText)
+		]
 
-			+SHorizontalBox::Slot()
-			.Padding(FEditorStyle::GetMargin("ContentBrowser.AssetTagNamePadding", StyleSpecifier))
-			.VAlign(VAlign_Center)
-			[
-				NameWidget.ToSharedRef()
-			]
+		+SHorizontalBox::Slot()
+		.Padding(FEditorStyle::GetMargin("ContentBrowser.AssetTagNamePadding", StyleSpecifier))
+		.VAlign(VAlign_Center)
+		[
+			NameWidget.ToSharedRef()
 		]
 	];
 
@@ -221,7 +208,8 @@ void SAssetTagItem::Construct(const FArguments& InArgs)
 			[
 				SNew(SImage)
 				.Visibility(this, &SAssetTagItem::GetWarningIconVisibility)
-				.Image(FEditorStyle::GetBrush("Icons.Warning"))
+				.Image(FAppStyle::Get().GetBrush("Icons.Warning"))
+				.ColorAndOpacity(FAppStyle::Get().GetSlateColor("Colors.AccentYellow"))
 				.ToolTipText(WarningText)
 			];
 	}
@@ -243,9 +231,6 @@ void SAssetTagItem::Construct(const FArguments& InArgs)
 					.MinDesiredWidth(30.0f)
 					[
 						SNew(STextBlock)
-						.Font(FEditorStyle::GetFontStyle("ContentBrowser.AssetTagCountFont", StyleSpecifier))
-						.ColorAndOpacity(FLinearColor::White)
-						.ShadowOffset(FVector2D(1.0f, 1.0f))
 						.Text(InArgs._CountText)
 						.Justification(ETextJustify::Center)
 					]
diff --git a/Engine/Source/Editor/AssetTagsEditor/Public/SAssetTagItem.h b/Engine/Source/Editor/AssetTagsEditor/Public/SAssetTagItem.h
--- a/Engine/Source/Editor/AssetTagsEditor/Public/SAssetTagItem.h
+++ b/Engine/Source/Editor/AssetTagsEditor/Public/SAssetTagItem.h
@@ -20,9 +20,6 @@ public:
 	{
 		this->ConstructInternal(InArgs, InOwnerTableView);
 
-		this->Style = &FEditorStyle::GetWidgetStyle<FTableRowStyle>("ContentBrowser.AssetTagTableRow");
-		this->BorderImage = FInvalidatableBrushAttribute(FEditorStyle::GetBrush("NoBorder"));
-
 		this->ConstructChildren(
 			InOwnerTableView->TableViewMode,
 			ApplySelectionBorderOuterPadding(InArgs._Padding),
@@ -55,7 +52,7 @@ private:
 	{
 		return
 			SNew(SBorder)
-			.Padding(2.0f)
+			.Padding(0)
 			.BorderImage(this, &SAssetTagItemTableRow::GetBorder)
 			[
 				InContent
@@ -167,15 +164,6 @@ private:
 		return BaseColor.Get().CopyWithNewOpacity(0.04f);
 	}
 
-	FSlateColor GetCheckBoxForegroundColor() const
-	{
-		return CheckBox && CheckBox->IsEnabled()
-			? CheckBox->IsChecked() 
-				? GetAssetTagBrightColor()
-				: GetAssetTagDullColor()
-			: GetAssetTagDisabledColor();
-	}
-
 	FText GetCheckBoxTooltipText() const
 	{
 		return CheckBox && CheckBox->IsEnabled()
diff --git a/Engine/Source/Editor/ContentBrowser/Private/AssetViewWidgets.cpp b/Engine/Source/Editor/ContentBrowser/Private/AssetViewWidgets.cpp
--- a/Engine/Source/Editor/ContentBrowser/Private/AssetViewWidgets.cpp
+++ b/Engine/Source/Editor/ContentBrowser/Private/AssetViewWidgets.cpp
@@ -93,10 +93,10 @@ TSharedRef<SWidget> FAssetViewItemHelper::CreateListTileItemContents(T* const In
 {
 	TSharedRef<SOverlay> ItemContentsOverlay = SNew(SOverlay);
 
+	OutItemShadowBorder = FName("ContentBrowser.AssetTileItem.DropShadow");
+
 	if (InTileOrListItem->IsFolder())
 	{
-		OutItemShadowBorder = FName("NoBorder");
-
 		// TODO: Allow items to customize their widget
 
 		const bool bDeveloperFolder = ContentBrowserUtils::IsItemDeveloperContent(InTileOrListItem->AssetItem->GetItem());
@@ -109,12 +109,8 @@ TSharedRef<SWidget> FAssetViewItemHelper::CreateListTileItemContents(T* const In
 		}
 
 		const FSlateBrush* FolderBaseImage = bDeveloperFolder
-			? FEditorStyle::GetBrush("ContentBrowser.ListViewDeveloperFolderIcon.Base") 
-			: FEditorStyle::GetBrush("ContentBrowser.ListViewFolderIcon.Base");
-
-		const FSlateBrush* FolderTintImage = bDeveloperFolder
-			? FEditorStyle::GetBrush("ContentBrowser.ListViewDeveloperFolderIcon.Mask") 
-			: FEditorStyle::GetBrush("ContentBrowser.ListViewFolderIcon.Mask");
+			? FEditorStyle::GetBrush("ContentBrowser.ListViewDeveloperFolderIcon") 
+			: FEditorStyle::GetBrush("ContentBrowser.ListViewFolderIcon");
 
 		// Folder base
 		ItemContentsOverlay->AddSlot()
@@ -168,17 +164,10 @@ TSharedRef<SWidget> FAssetViewItemHelper::CreateListTileItemContents(T* const In
 				]
 			];
 		}
-
-		// Folder tint
-		ItemContentsOverlay->AddSlot()
-		[
-			SNew(SImage)
-			.Image(FolderTintImage)
-		];
 	}
 	else
 	{
-		OutItemShadowBorder = FName("ContentBrowser.ThumbnailShadow");
+
 
 		// The actual thumbnail
 		ItemContentsOverlay->AddSlot()
@@ -311,8 +300,8 @@ FReply SAssetTileView::OnKeyDown( const FGeometry& InGeometry, const FKeyEvent&
 
 void SAssetTileView::Tick(const FGeometry& AllottedGeometry, const double InCurrentTime, const float InDeltaTime)
 {
-	// Regreshing an asset view is an intensive task. Do not do this while a user
-	// is dragging arround content for maximum responsiveness.
+	// Refreshing an asset view is an intensive task. Do not do this while a user
+	// is dragging around content for maximum responsiveness.
 	// Also prevents a re-entrancy crash caused by potentially complex thumbnail generators.
 	if (!FSlateApplication::Get().IsDragDropping())
 	{
@@ -336,8 +325,8 @@ FReply SAssetListView::OnKeyDown(const FGeometry& InGeometry, const FKeyEvent& I
 
 void SAssetListView::Tick(const FGeometry& AllottedGeometry, const double InCurrentTime, const float InDeltaTime)
 {
-	// Regreshing an asset view is an intensive task. Do not do this while a user
-	// is dragging arround content for maximum responsiveness.
+	// Refreshing an asset view is an intensive task. Do not do this while a user
+	// is dragging around content for maximum responsiveness.
 	// Also prevents a re-entrancy crash caused by potentially complex thumbnail generators.
 	if (!FSlateApplication::Get().IsDragDropping())
 	{
@@ -362,8 +351,8 @@ FReply SAssetColumnView::OnKeyDown(const FGeometry& InGeometry, const FKeyEvent&
 
 void SAssetColumnView::Tick(const FGeometry& AllottedGeometry, const double InCurrentTime, const float InDeltaTime)
 {
-	// Regreshing an asset view is an intensive task. Do not do this while a user
-	// is dragging arround content for maximum responsiveness.
+	// Refreshing an asset view is an intensive task. Do not do this while a user
+	// is dragging around content for maximum responsiveness.
 	// Also prevents a re-entrancy crash caused by potentially complex thumbnail generators.
 	if (!FSlateApplication::Get().IsDragDropping())
 	{
@@ -399,6 +388,7 @@ void SAssetViewItem::Construct( const FArguments& InArgs )
 	OnGetCustomAssetToolTip = InArgs._OnGetCustomAssetToolTip;
 	OnVisualizeAssetToolTip = InArgs._OnVisualizeAssetToolTip;
 	OnAssetToolTipClosing = InArgs._OnAssetToolTipClosing;
+	IsSelected = InArgs._IsSelected;
 
 	bDraggedOver = false;
 
@@ -677,7 +667,7 @@ EVisibility SAssetViewItem::GetThumbnailEditModeUIVisibility() const
 	return !IsFolder() && ThumbnailEditMode.Get() ? EVisibility::Visible : EVisibility::Collapsed;
 }
 
-BEGIN_SLATE_FUNCTION_BUILD_OPTIMIZATION
+
 TSharedRef<SWidget> SAssetViewItem::CreateToolTipWidget() const
 {
 	if ( AssetItem.IsValid() )
@@ -962,7 +952,6 @@ TSharedRef<SWidget> SAssetViewItem::CreateToolTipWidget() const
 		return SNullWidget::NullWidget;
 	}
 }
-END_SLATE_FUNCTION_BUILD_OPTIMIZATION
 
 EVisibility SAssetViewItem::GetCheckedOutByOtherTextVisibility() const
 {
@@ -1409,6 +1398,7 @@ void SAssetListItem::Construct( const FArguments& InArgs )
 		.OnGetCustomAssetToolTip(InArgs._OnGetCustomAssetToolTip)
 		.OnVisualizeAssetToolTip(InArgs._OnVisualizeAssetToolTip)
 		.OnAssetToolTipClosing( InArgs._OnAssetToolTipClosing )
+		.IsSelected(InArgs._IsSelected)
 		);
 
 	AssetThumbnail = InArgs._AssetThumbnail;
@@ -1499,7 +1489,7 @@ void SAssetListItem::Construct( const FArguments& InArgs )
 					.OnTextCommitted(this, &SAssetListItem::HandleNameCommitted)
 					.OnVerifyTextChanged(this, &SAssetListItem::HandleVerifyNameChanged)
 					.HighlightText(InArgs._HighlightText)
-					.IsSelected(InArgs._IsSelected)
+					.IsSelected(InArgs._IsSelectedExclusively)
 					.IsReadOnly(this, &SAssetListItem::IsNameReadOnly)
 				]
 
@@ -1578,7 +1568,6 @@ SAssetTileItem::~SAssetTileItem()
 
 }
 
-BEGIN_SLATE_FUNCTION_BUILD_OPTIMIZATION
 void SAssetTileItem::Construct( const FArguments& InArgs )
 {
 	SAssetViewItem::Construct( SAssetViewItem::FArguments()
@@ -1594,6 +1583,7 @@ void SAssetTileItem::Construct( const FArguments& InArgs )
 		.OnGetCustomAssetToolTip(InArgs._OnGetCustomAssetToolTip)
 		.OnVisualizeAssetToolTip(InArgs._OnVisualizeAssetToolTip)
 		.OnAssetToolTipClosing( InArgs._OnAssetToolTipClosing )
+		.IsSelected(InArgs._IsSelected)
 		);
 
 	AssetThumbnail = InArgs._AssetThumbnail;
@@ -1611,7 +1601,7 @@ void SAssetTileItem::Construct( const FArguments& InArgs )
 		ThumbnailConfig.ThumbnailLabel = InArgs._ThumbnailLabel;
 		ThumbnailConfig.HighlightedText = InArgs._HighlightText;
 		ThumbnailConfig.HintColorAndOpacity = InArgs._ThumbnailHintColorAndOpacity;
-
+		ThumbnailConfig.Padding= FMargin(2.0f);
 		{
 			FContentBrowserItemDataAttributeValue ColorAttributeValue = AssetItem->GetItem().GetItemAttribute(ContentBrowserItemAttributes::ItemColor);
 			if (ColorAttributeValue.IsValid())
@@ -1630,62 +1620,93 @@ void SAssetTileItem::Construct( const FArguments& InArgs )
 	}
 	else
 	{
-		Thumbnail = SNew(SImage) .Image( FEditorStyle::GetDefaultBrush() );
+		Thumbnail = SNew(SImage).Image(FEditorStyle::GetDefaultBrush());
 	}
 
 	FName ItemShadowBorderName;
 	TSharedRef<SWidget> ItemContents = FAssetViewItemHelper::CreateTileItemContents(this, Thumbnail.ToSharedRef(), ItemShadowBorderName);
 
 	ChildSlot
-	[
+	.Padding(FMargin(0.0f, 0.0f, 4.0f, 4.0f))
+	[				
+		// Drop shadow border
 		SNew(SBorder)
-		.BorderImage(this, &SAssetViewItem::GetBorderImage)
-		.Padding(0)
-		.AddMetaData<FTagMetaData>(FTagMetaData(AssetItem->GetItem().GetVirtualPath()))
+		.Padding(FMargin(0.0f, 0.0f, 5.0f, 5.0f))
+		.BorderImage(FAppStyle::Get().GetBrush(ItemShadowBorderName))
 		[
-			SNew(SVerticalBox)
-
-			// Thumbnail
-			+SVerticalBox::Slot()
-			.AutoHeight()
-			.HAlign(HAlign_Center)
+			SNew(SOverlay)
+			.AddMetaData<FTagMetaData>(FTagMetaData(AssetItem->GetItem().GetVirtualPath()))
+			+SOverlay::Slot()
 			[
-				// The remainder of the space is reserved for the name.
-				SNew(SBox)
-				.Padding(ThumbnailPadding - 4.f)
-				.WidthOverride(this, &SAssetTileItem::GetThumbnailBoxSize)
-				.HeightOverride( this, &SAssetTileItem::GetThumbnailBoxSize )
+				SNew(SBorder)
+				.Padding(0)
+				.BorderImage(FAppStyle::Get().GetBrush("ContentBrowser.AssetTileItem.ThumbnailAreaBackground"))
 				[
-					// Drop shadow border
-					SNew(SBorder)
-					.Padding(4.f)
-					.BorderImage(FEditorStyle::GetBrush(ItemShadowBorderName))
+					SNew(SVerticalBox)
+					// Thumbnail
+					+SVerticalBox::Slot()
+					.AutoHeight()
 					[
-						ItemContents
+						// The remainder of the space is reserved for the name.
+						SNew(SBox)
+						.Padding(0)
+						.WidthOverride(this, &SAssetTileItem::GetThumbnailBoxSize)
+						.HeightOverride(this, &SAssetTileItem::GetThumbnailBoxSize)
+						[
+							ItemContents
+						]
+					]
+
+					+SVerticalBox::Slot()
+					// note this negative padding is intentional to cover up the top part of the rounded border
+					.Padding(FMargin(0.0f, -3.0f, 0.0f, 0.0f))
+					[
+						SNew(SBorder)
+						.Padding(FMargin(2.0f, 3.0f))
+						.BorderImage(FAppStyle::Get().GetBrush("ContentBrowser.AssetTileItem.NameAreaBackground"))
+						[
+							SNew(SVerticalBox)
+							+ SVerticalBox::Slot()
+							.HAlign(HAlign_Left)
+							.VAlign(VAlign_Top)
+							.Padding(2.0f)
+							.AutoHeight()
+							[
+								SNew(STextBlock)
+								.Visibility(this, &SAssetTileItem::GetAssetClassLabelVisibility)
+								.TextStyle(FAppStyle::Get(), "ButtonText")
+								.TransformPolicy(ETextTransformPolicy::ToUpper)
+								.Text(this, &SAssetTileItem::GetAssetClassText)
+							]
+							+ SVerticalBox::Slot()
+							.Padding(2.0f)
+							.VAlign(VAlign_Center)
+							[
+								SAssignNew(InlineRenameWidget, SInlineEditableTextBlock)
+								.Font(this, &SAssetTileItem::GetThumbnailFont)
+								.Text(GetNameText())
+								.OnBeginTextEdit(this, &SAssetTileItem::HandleBeginNameChange)
+								.OnTextCommitted(this, &SAssetTileItem::HandleNameCommitted)
+								.OnVerifyTextChanged(this, &SAssetTileItem::HandleVerifyNameChanged)
+								.HighlightText(InArgs._HighlightText)
+								.IsSelected(InArgs._IsSelectedExclusively)
+								.IsReadOnly(this, &SAssetTileItem::IsNameReadOnly)
+								.Justification(ETextJustify::Left)
+								.LineBreakPolicy(FBreakIterator::CreateCamelCaseBreakIterator())
+
+
+							]
+						]
 					]
 				]
 			]
-
-			+SVerticalBox::Slot()
-			.Padding(FMargin(1.f, 0))
-			.HAlign(HAlign_Center)
-			.VAlign(VAlign_Center)
-			.FillHeight(1.f)
+			+SOverlay::Slot()
 			[
-				SAssignNew(InlineRenameWidget, SInlineEditableTextBlock)
-					.Font( this, &SAssetTileItem::GetThumbnailFont )
-					.Text( GetNameText() )
-					.OnBeginTextEdit(this, &SAssetTileItem::HandleBeginNameChange)
-					.OnTextCommitted(this, &SAssetTileItem::HandleNameCommitted)
-					.OnVerifyTextChanged(this, &SAssetTileItem::HandleVerifyNameChanged)
-					.HighlightText(InArgs._HighlightText)
-					.IsSelected(InArgs._IsSelected)
-					.IsReadOnly(this, &SAssetTileItem::IsNameReadOnly)
-					.Justification(ETextJustify::Center)
-					.LineBreakPolicy(FBreakIterator::CreateCamelCaseBreakIterator())
+				SNew(SImage)
+				.Image(this, &SAssetViewItem::GetBorderImage)
+				.Visibility(EVisibility::HitTestInvisible)	
 			]
 		]
-	
 	];
 
 	if(AssetItem.IsValid())
@@ -1694,7 +1715,6 @@ void SAssetTileItem::Construct( const FArguments& InArgs )
 		AssetItem->OnRenameCanceled().BindSP(InlineRenameWidget.Get(), &SInlineEditableTextBlock::ExitEditingMode);
 	}
 }
-END_SLATE_FUNCTION_BUILD_OPTIMIZATION
 
 void SAssetTileItem::OnAssetDataChanged()
 {
@@ -1714,6 +1734,32 @@ void SAssetTileItem::OnAssetDataChanged()
 	}
 }
 
+const FSlateBrush* SAssetTileItem::GetBorderImage() const
+{
+	const bool bIsSelected = IsSelected.IsBound() ? IsSelected.Execute() : false;
+	const bool bIsHoveredOrDraggedOver = IsHovered() || bDraggedOver;
+
+	if (bIsSelected && bIsHoveredOrDraggedOver)
+	{
+		static const FName SelectedHover("ContentBrowser.AssetTileItem.SelectedHoverBorder");
+		return FAppStyle::Get().GetBrush(SelectedHover);
+	}
+	else if (bIsSelected)
+	{
+		static const FName Selected("ContentBrowser.AssetTileItem.SelectedBorder");
+		return FAppStyle::Get().GetBrush(Selected);
+	}
+	else if (bIsHoveredOrDraggedOver)
+	{
+		static const FName Hovered("ContentBrowser.AssetTileItem.HoverBorder");
+		return FAppStyle::Get().GetBrush(Hovered);
+	}
+	else
+	{
+		return FStyleDefaults::GetNoBrush();
+	}
+}
+
 float SAssetTileItem::GetExtraStateIconWidth() const
 {
 	return GetStateIconImageSize().Get();
@@ -1732,7 +1778,18 @@ FOptionalSize SAssetTileItem::GetStateIconImageSize() const
 
 FOptionalSize SAssetTileItem::GetThumbnailBoxSize() const
 {
-	return FOptionalSize(ItemWidth.Get());
+	return FOptionalSize(ItemWidth.Get()- ThumbnailPadding);
+}
+
+EVisibility SAssetTileItem::GetAssetClassLabelVisibility() const
+{
+	FOptionalSize ThumbSize = GetThumbnailBoxSize();
+	if (ThumbSize.IsSet())
+	{
+		float Size = ThumbSize.Get() *GetTickSpaceGeometry().Scale;
+		return Size < 100 ? EVisibility::Collapsed : EVisibility::Visible;
+	}
+	return EVisibility::Collapsed;
 }
 
 FSlateFontInfo SAssetTileItem::GetThumbnailFont() const
@@ -1834,7 +1891,6 @@ void SAssetColumnItem::Construct( const FArguments& InArgs )
 	HighlightText = InArgs._HighlightText;
 }
 
-BEGIN_SLATE_FUNCTION_BUILD_OPTIMIZATION
 TSharedRef<SWidget> SAssetColumnItem::GenerateWidgetForColumn( const FName& ColumnName, FIsSelected InIsSelected )
 {
 	TSharedPtr<SWidget> Content;
@@ -1985,7 +2041,7 @@ TSharedRef<SWidget> SAssetColumnItem::GenerateWidgetForColumn( const FName& Colu
 			Content.ToSharedRef()
 		];
 }
-END_SLATE_FUNCTION_BUILD_OPTIMIZATION
+
 
 void SAssetColumnItem::OnAssetDataChanged()
 {
diff --git a/Engine/Source/Editor/ContentBrowser/Private/AssetViewWidgets.h b/Engine/Source/Editor/ContentBrowser/Private/AssetViewWidgets.h
--- a/Engine/Source/Editor/ContentBrowser/Private/AssetViewWidgets.h
+++ b/Engine/Source/Editor/ContentBrowser/Private/AssetViewWidgets.h
@@ -96,25 +96,25 @@ public:
 		{}
 
 		/** Data for the asset this item represents */
-		SLATE_ARGUMENT( TSharedPtr<FAssetViewItem>, AssetItem )
+	SLATE_ARGUMENT(TSharedPtr<FAssetViewItem>, AssetItem)
 
 		/** Delegate for when an asset name has entered a rename state */
-		SLATE_EVENT( FOnRenameBegin, OnRenameBegin )
+		SLATE_EVENT(FOnRenameBegin, OnRenameBegin)
 
 		/** Delegate for when an asset name has been entered for an item that is in a rename state */
-		SLATE_EVENT( FOnRenameCommit, OnRenameCommit )
+		SLATE_EVENT(FOnRenameCommit, OnRenameCommit)
 
 		/** Delegate for when an asset name has been entered for an item to verify the name before commit */
-		SLATE_EVENT( FOnVerifyRenameCommit, OnVerifyRenameCommit )
+		SLATE_EVENT(FOnVerifyRenameCommit, OnVerifyRenameCommit)
 
 		/** Called when any asset item is destroyed. Used in thumbnail management */
-		SLATE_EVENT( FOnItemDestroyed, OnItemDestroyed )
+		SLATE_EVENT(FOnItemDestroyed, OnItemDestroyed)
 
 		/** If false, the tooltip will not be displayed */
-		SLATE_ATTRIBUTE( bool, ShouldAllowToolTip )
+		SLATE_ATTRIBUTE(bool, ShouldAllowToolTip)
 
 		/** If true, display the thumbnail edit mode UI */
-		SLATE_ATTRIBUTE( bool, ThumbnailEditMode )
+		SLATE_ATTRIBUTE(bool, ThumbnailEditMode)
 
 		/** The string in the title to highlight (used when searching by string) */
 		SLATE_ATTRIBUTE(FText, HighlightText)
@@ -123,13 +123,16 @@ public:
 		SLATE_EVENT(FOnIsAssetValidForCustomToolTip, OnIsAssetValidForCustomToolTip)
 
 		/** Delegate to call (if bound) to get a custom tooltip for this view item */
-		SLATE_EVENT( FOnGetCustomAssetToolTip, OnGetCustomAssetToolTip )
+		SLATE_EVENT(FOnGetCustomAssetToolTip, OnGetCustomAssetToolTip)
 
 		/** Delegate for when an item is about to show a tool tip */
-		SLATE_EVENT( FOnVisualizeAssetToolTip, OnVisualizeAssetToolTip)
+		SLATE_EVENT(FOnVisualizeAssetToolTip, OnVisualizeAssetToolTip)
 
 		/** Delegate for when an item's tooltip is about to close */
-		SLATE_EVENT( FOnAssetToolTipClosing, OnAssetToolTipClosing )
+		SLATE_EVENT(FOnAssetToolTipClosing, OnAssetToolTipClosing)
+
+		/** Delegate for getting the selection state of this item */ 
+		SLATE_ARGUMENT(FIsSelected, IsSelected)
 
 	SLATE_END_ARGS()
 
@@ -153,7 +156,7 @@ public:
 	virtual FSlateColor GetAssetColor() const;
 
 	/** Get the border image to display */
-	const FSlateBrush* GetBorderImage() const;
+	virtual const FSlateBrush* GetBorderImage() const;
 
 	/** Get the name text to be displayed for this item */
 	FText GetNameText() const;
@@ -283,8 +286,10 @@ protected:
 
 	/** Called if bound when a tooltip is closing */
 	FOnAssetToolTipClosing OnAssetToolTipClosing;
+	
+	/** Delegate for getting the selection state of this item */ 
+	FIsSelected IsSelected;
 
-	/** The geometry last frame. Used when telling popup messages where to appear. */
 	FGeometry LastGeometry;
 
 	/** If false, the tooltip will not be displayed */
@@ -378,7 +383,10 @@ public:
 		SLATE_ARGUMENT( bool, AllowThumbnailHintLabel )
 
 		/** Whether the item is selected in the view */
-		SLATE_ARGUMENT( FIsSelected, IsSelected )
+		SLATE_ARGUMENT(FIsSelected, IsSelected)
+
+		/** Whether the item is selected in the view without anything else being selected*/
+		SLATE_ARGUMENT(FIsSelected, IsSelectedExclusively)
 
 		/** Delegate to call (if bound) to check if it is valid to get a custom tooltip for this view item */
 		SLATE_EVENT(FOnIsAssetValidForCustomToolTip, OnIsAssetValidForCustomToolTip)
@@ -489,7 +497,10 @@ public:
 		SLATE_ATTRIBUTE( bool, ThumbnailEditMode )
 
 		/** Whether the item is selected in the view */
-		SLATE_ARGUMENT( FIsSelected, IsSelected )
+		SLATE_ARGUMENT(FIsSelected, IsSelected)
+
+		/** Whether the item is selected in the view without anything else being selected*/
+		SLATE_ARGUMENT(FIsSelected, IsSelectedExclusively)
 
 		/** Delegate to call (if bound) to check if it is valid to get a custom tooltip for this view item */
 		SLATE_EVENT(FOnIsAssetValidForCustomToolTip, OnIsAssetValidForCustomToolTip)
@@ -517,9 +528,11 @@ public:
 	/** Whether the widget should allow primitive tools to be displayed */
 	bool CanDisplayPrimitiveTools() const { return true; }
 
+	/** Get the border image to display */
+	virtual const FSlateBrush* GetBorderImage() const;
 protected:
 	/** SAssetViewItem interface */
-	virtual float GetNameTextWrapWidth() const override { return LastGeometry.GetLocalSize().X - 2.f; }
+	virtual float GetNameTextWrapWidth() const override { return LastGeometry.GetLocalSize().X - 15.f; }
 
 	/** Get the expected width of an extra state icon. */
 	float GetExtraStateIconWidth() const;
@@ -533,6 +546,9 @@ protected:
 	/** Returns the size of the thumbnail widget */
 	FOptionalSize GetThumbnailBoxSize() const;
 
+	/** Gets the visibility of the asset class label in thumbnails */
+	EVisibility GetAssetClassLabelVisibility() const;
+
 	/** Returns the font to use for the thumbnail label */
 	FSlateFontInfo GetThumbnailFont() const;
 
@@ -627,7 +643,7 @@ public:
 
 		SMultiColumnTableRow< TSharedPtr<FAssetViewItem> >::Construct( 
 			FSuperRowType::FArguments()
-				.Style(FEditorStyle::Get(), "ContentBrowser.AssetListView.TableRow")
+				.Style(FEditorStyle::Get(), "ContentBrowser.AssetListView.ColumnListTableRow")
 				.OnDragDetected(InArgs._OnDragDetected), 
 			InOwnerTableView);
 		Content = this->AssetColumnItem;
diff --git a/Engine/Source/Editor/ContentBrowser/Private/SAssetPicker.cpp b/Engine/Source/Editor/ContentBrowser/Private/SAssetPicker.cpp
--- a/Engine/Source/Editor/ContentBrowser/Private/SAssetPicker.cpp
+++ b/Engine/Source/Editor/ContentBrowser/Private/SAssetPicker.cpp
@@ -295,7 +295,6 @@ void SAssetPicker::Construct( const FArguments& InArgs )
 		.InitialBackendFilter(CurrentBackendFilter)
 		.InitialViewType(InArgs._AssetPickerConfig.InitialAssetViewType)
 		.InitialAssetSelection(InArgs._AssetPickerConfig.InitialAssetSelection)
-		.ThumbnailScale(InArgs._AssetPickerConfig.ThumbnailScale)
 		.ShowBottomToolbar(InArgs._AssetPickerConfig.bShowBottomToolbar)
 		.OnAssetTagWantsToBeDisplayed(InArgs._AssetPickerConfig.OnAssetTagWantsToBeDisplayed)
 		.OnGetCustomSourceAssets(InArgs._AssetPickerConfig.OnGetCustomSourceAssets)
diff --git a/Engine/Source/Editor/ContentBrowser/Private/SAssetView.cpp b/Engine/Source/Editor/ContentBrowser/Private/SAssetView.cpp
--- a/Engine/Source/Editor/ContentBrowser/Private/SAssetView.cpp
+++ b/Engine/Source/Editor/ContentBrowser/Private/SAssetView.cpp
@@ -78,6 +78,26 @@ namespace
 	const double JumpDelaySeconds = 2.0;
 }
 
+
+static FText ThumbnailSizeToDisplayName(EThumbnailSize InSize)
+{
+	switch (InSize)
+	{
+	case EThumbnailSize::Tiny:
+		return LOCTEXT("TinyThumbnailSize", "Tiny");
+	case EThumbnailSize::Small:
+		return LOCTEXT("SmallThumbnailSize", "Small");
+	case EThumbnailSize::Medium:
+		return LOCTEXT("MediumThumbnailSize", "Medium");
+	case EThumbnailSize::Large:
+		return LOCTEXT("LargeThumbnailSize", "Large");
+	case EThumbnailSize::Huge:
+		return LOCTEXT("HugeThumbnailSize", "Huge");
+	default:
+		return FText::GetEmpty();
+	}
+}
+
 SAssetView::~SAssetView()
 {
 	if (IContentBrowserDataModule* ContentBrowserDataModule = IContentBrowserDataModule::GetPtr())
@@ -102,7 +122,7 @@ SAssetView::~SAssetView()
 	AssetThumbnailPool.Reset();
 }
 
-BEGIN_SLATE_FUNCTION_BUILD_OPTIMIZATION
+
 void SAssetView::Construct( const FArguments& InArgs )
 {
 	bIsWorking = false;
@@ -129,6 +149,8 @@ void SAssetView::Construct( const FArguments& InArgs )
 	// Listen for when view settings are changed
 	UContentBrowserSettings::OnSettingChanged().AddSP(this, &SAssetView::HandleSettingChanged);
 
+	ThumbnailSize = EThumbnailSize::Medium;
+
 	// Get desktop metrics
 	FDisplayMetrics DisplayMetrics;
 	FSlateApplication::Get().GetCachedDisplayMetrics( DisplayMetrics );
@@ -137,7 +159,7 @@ void SAssetView::Construct( const FArguments& InArgs )
 		DisplayMetrics.PrimaryDisplayWorkAreaRect.Right - DisplayMetrics.PrimaryDisplayWorkAreaRect.Left,
 		DisplayMetrics.PrimaryDisplayWorkAreaRect.Bottom - DisplayMetrics.PrimaryDisplayWorkAreaRect.Top );
 
-	const float ThumbnailScaleRangeScalar = ( DisplaySize.Y / 1080 );
+	ThumbnailScaleRangeScalar = ( DisplaySize.Y / 2160 );
 
 	// Create a thumbnail pool for rendering thumbnails	
 	AssetThumbnailPool = MakeShareable( new FAssetThumbnailPool(1024, InArgs._AreRealTimeThumbnailsAllowed) );
@@ -146,19 +168,13 @@ void SAssetView::Construct( const FArguments& InArgs )
 	ListViewThumbnailSize = 64;
 	ListViewThumbnailPadding = 4;
 	TileViewThumbnailResolution = 256;
-	TileViewThumbnailSize = 128;
-	TileViewThumbnailPadding = 5;
+	TileViewThumbnailSize = 150;
+	TileViewThumbnailPadding = 9;
 
-	TileViewNameHeight = 36;
-	ThumbnailScaleSliderValue = InArgs._ThumbnailScale; 
-
-	if ( !ThumbnailScaleSliderValue.IsBound() )
-	{
-		ThumbnailScaleSliderValue = FMath::Clamp<float>(ThumbnailScaleSliderValue.Get(), 0.0f, 1.0f);
-	}
+	TileViewNameHeight = 115;
 
 	MinThumbnailScale = 0.2f * ThumbnailScaleRangeScalar;
-	MaxThumbnailScale = 2.0f * ThumbnailScaleRangeScalar;
+	MaxThumbnailScale = 1.9f * ThumbnailScaleRangeScalar;
 
 	bCanShowClasses = InArgs._CanShowClasses;
 
@@ -241,6 +257,8 @@ void SAssetView::Construct( const FArguments& InArgs )
 
 	NumVisibleColumns = 0;
 
+	OwningContentBrowser = InArgs._OwningContentBrowser;
+
 	FAssetToolsModule& AssetToolsModule = FModuleManager::LoadModuleChecked<FAssetToolsModule>("AssetTools");
 	AssetClassBlacklist = AssetToolsModule.Get().GetAssetClassBlacklist();
 	FolderBlacklist = AssetToolsModule.Get().GetFolderBlacklist();
@@ -251,8 +269,14 @@ void SAssetView::Construct( const FArguments& InArgs )
 	TSharedRef<SVerticalBox> VerticalBox = SNew(SVerticalBox);
 
 	ChildSlot
+	.Padding(0.0f, 1.0f, 0.0f, 0.0f)
 	[
-		VerticalBox
+		SNew(SBorder)
+		.Padding(0)
+		.BorderImage(FAppStyle::Get().GetBrush("Brushes.Background"))
+		[
+			VerticalBox
+		]
 	];
 
 	// Assets area
@@ -285,9 +309,8 @@ void SAssetView::Construct( const FArguments& InArgs )
 			.VAlign(VAlign_Fill)
 			[
 				// Container for the view types
-				SAssignNew(ViewContainer, SBorder)
+				SAssignNew(ViewContainer, SBox)
 				.Padding(0)
-				.BorderImage(FEditorStyle::GetBrush("NoBorder"))
 			]
 
 			+ SOverlay::Slot()
@@ -336,11 +359,10 @@ void SAssetView::Construct( const FArguments& InArgs )
 	[
 		SNew(SBorder)
 		.Visibility( this, &SAssetView::GetEditModeLabelVisibility )
-		.BorderImage( FEditorStyle::GetBrush("ContentBrowser.EditModeLabelBorder") )
+		.BorderImage(FAppStyle::Get().GetBrush("Brushes.Background"))
 		.Content()
 		[
-			SNew( SHorizontalBox )
-
+			SNew(SHorizontalBox)
 			+SHorizontalBox::Slot()
 			.VAlign(VAlign_Center)
 			.Padding(4, 0, 0, 0)
@@ -348,7 +370,7 @@ void SAssetView::Construct( const FArguments& InArgs )
 			[
 				SNew(STextBlock)
 				.Text(LOCTEXT("ThumbnailEditModeLabel", "Editing Thumbnails. Drag a thumbnail to rotate it if there is a 3D environment."))
-				.TextStyle( FEditorStyle::Get(), "ContentBrowser.EditModeLabelFont" )
+				.ColorAndOpacity(FAppStyle::Get().GetSlateColor("Colors.Primary"))
 			]
 
 			+SHorizontalBox::Slot()
@@ -356,8 +378,9 @@ void SAssetView::Construct( const FArguments& InArgs )
 			.VAlign(VAlign_Center)
 			[
 				SNew(SButton)
-				.Text( LOCTEXT("EndThumbnailEditModeButton", "Done Editing") )
-				.OnClicked( this, &SAssetView::EndThumbnailEditModeClicked )
+				.ButtonStyle(FAppStyle::Get(), "PrimaryButton")
+				.Text(LOCTEXT("EndThumbnailEditModeButton", "Done Editing"))
+				.OnClicked(this, &SAssetView::EndThumbnailEditModeClicked)
 			]
 		]
 	];
@@ -374,7 +397,7 @@ void SAssetView::Construct( const FArguments& InArgs )
 			+SHorizontalBox::Slot()
 			.FillWidth(1.f)
 			.VAlign(VAlign_Center)
-			.Padding(8, 0)
+			.Padding(8, 5)
 			[
 				SNew(STextBlock)
 				.Text(this, &SAssetView::GetAssetCountText)
@@ -384,9 +407,9 @@ void SAssetView::Construct( const FArguments& InArgs )
 			+SHorizontalBox::Slot()
 			.AutoWidth()
 			[
-				SAssignNew( ViewOptionsComboButton, SComboButton )
+				SNew(SComboButton)
+				.Visibility(InArgs._ShowViewOptions ? EVisibility::Visible : EVisibility::Collapsed)
 				.ContentPadding(0)
-				.ForegroundColor( this, &SAssetView::GetViewButtonForegroundColor )
 				.ButtonStyle( FEditorStyle::Get(), "ToggleButton" ) // Use the tool bar item style for this button
 				.OnGetMenuContent( this, &SAssetView::GetViewButtonContent )
 				.ButtonContent()
@@ -397,7 +420,8 @@ void SAssetView::Construct( const FArguments& InArgs )
 					.AutoWidth()
 					.VAlign(VAlign_Center)
 					[
-						SNew(SImage).Image( FEditorStyle::GetBrush("GenericViewButton") )
+						SNew(SImage)
+						.Image( FEditorStyle::GetBrush("GenericViewButton") )
 					]
  
 					+SHorizontalBox::Slot()
@@ -405,7 +429,8 @@ void SAssetView::Construct( const FArguments& InArgs )
 					.Padding(2, 0, 0, 0)
 					.VAlign(VAlign_Center)
 					[
-						SNew(STextBlock).Text( LOCTEXT("ViewButton", "View Options") )
+						SNew(STextBlock)
+						.Text( LOCTEXT("ViewButton", "View Options") )
 					]
 				]
 			]
@@ -435,8 +460,6 @@ void SAssetView::Construct( const FArguments& InArgs )
 	}
 }
 
-END_SLATE_FUNCTION_BUILD_OPTIMIZATION
-
 TOptional< float > SAssetView::GetIsWorkingProgressBarState() const
 {
 	if (bIsWorking)
@@ -744,7 +767,7 @@ void SAssetView::RequestQuickFrontendListRefresh()
 
 FString SAssetView::GetThumbnailScaleSettingPath(const FString& SettingsString) const
 {
-	return SettingsString + TEXT(".ThumbnailSizeScale");
+	return SettingsString + TEXT(".ThumbnailSize");
 }
 
 FString SAssetView::GetCurrentViewTypeSettingPath(const FString& SettingsString) const
@@ -754,7 +777,7 @@ FString SAssetView::GetCurrentViewTypeSettingPath(const FString& SettingsString)
 
 void SAssetView::SaveSettings(const FString& IniFilename, const FString& IniSection, const FString& SettingsString) const
 {
-	GConfig->SetFloat(*IniSection, *GetThumbnailScaleSettingPath(SettingsString), ThumbnailScaleSliderValue.Get(), IniFilename);
+	GConfig->SetInt(*IniSection, *GetThumbnailScaleSettingPath(SettingsString), (int32)ThumbnailSize, IniFilename);
 	GConfig->SetInt(*IniSection, *GetCurrentViewTypeSettingPath(SettingsString), CurrentViewType, IniFilename);
 	
 	GConfig->SetArray(*IniSection, *(SettingsString + TEXT(".HiddenColumns")), HiddenColumnNames, IniFilename);
@@ -762,12 +785,12 @@ void SAssetView::SaveSettings(const FString& IniFilename, const FString& IniSect
 
 void SAssetView::LoadSettings(const FString& IniFilename, const FString& IniSection, const FString& SettingsString)
 {
-	float Scale = 0.f;
-	if ( GConfig->GetFloat(*IniSection, *GetThumbnailScaleSettingPath(SettingsString), Scale, IniFilename) )
+	int32 ThumbnailSizeConfig = (int32)EThumbnailSize::Medium;
+	if ( GConfig->GetInt(*IniSection, *GetThumbnailScaleSettingPath(SettingsString), ThumbnailSizeConfig, IniFilename) )
 	{
 		// Clamp value to normal range and update state
-		Scale = FMath::Clamp<float>(Scale, 0.f, 1.f);
-		SetThumbnailScale(Scale);
+		ThumbnailSizeConfig = FMath::Clamp<int32>(ThumbnailSizeConfig, 0, (int32)EThumbnailSize::MAX-1);
+		OnThumbnailSizeChanged((EThumbnailSize)ThumbnailSizeConfig);
 	}
 
 	int32 ViewType = EAssetViewType::Tile;
@@ -831,6 +854,10 @@ void SAssetView::AdjustActiveSelection(int32 SelectionDelta)
 
 void SAssetView::Tick( const FGeometry& AllottedGeometry, const double InCurrentTime, const float InDeltaTime )
 {
+	// Adjust min and max thumbnail scale based on dpi
+	MinThumbnailScale = (0.2f * ThumbnailScaleRangeScalar)/AllottedGeometry.Scale;
+	MaxThumbnailScale = (1.9f * ThumbnailScaleRangeScalar)/AllottedGeometry.Scale;
+
 	CalculateFillScale( AllottedGeometry );
 
 	CurrentTime = InCurrentTime;
@@ -1028,7 +1055,7 @@ void SAssetView::CalculateFillScale( const FGeometry& AllottedGeometry )
 
 		// Scrollbars are 16, but we add 1 to deal with half pixels.
 		const float ScrollbarWidth = 16 + 1;
-		float TotalWidth = AllottedGeometry.GetLocalSize().X - ( ScrollbarWidth / AllottedGeometry.Scale );
+		float TotalWidth = AllottedGeometry.GetLocalSize().X -(ScrollbarWidth);
 		float Coverage = TotalWidth / ItemWidth;
 		int32 Items = (int)( TotalWidth / ItemWidth );
 
@@ -1500,6 +1527,7 @@ FReply SAssetView::OnKeyDown( const FGeometry& MyGeometry, const FKeyEvent& InKe
 
 FReply SAssetView::OnMouseWheel( const FGeometry& MyGeometry, const FPointerEvent& MouseEvent )
 {
+/*
 	if( MouseEvent.IsControlDown() )
 	{
 		const float DesiredScale = FMath::Clamp<float>(GetThumbnailScale() + ( MouseEvent.GetWheelDelta() * 0.05f ), 0.0f, 1.0f);
@@ -1508,7 +1536,7 @@ FReply SAssetView::OnMouseWheel( const FGeometry& MyGeometry, const FPointerEven
 			SetThumbnailScale( DesiredScale );
 		}		
 		return FReply::Handled();
-	}
+	}*/
 	return FReply::Unhandled();
 }
 
@@ -2172,14 +2200,6 @@ FLinearColor SAssetView::GetThumbnailHintColorAndOpacity() const
 	return ThumbnailHintColorAndOpacity;
 }
 
-FSlateColor SAssetView::GetViewButtonForegroundColor() const
-{
-	static const FName InvertedForegroundName("InvertedForeground");
-	static const FName DefaultForegroundName("DefaultForeground");
-
-	return ViewOptionsComboButton->IsHovered() ? FEditorStyle::GetSlateColor(InvertedForegroundName) : FEditorStyle::GetSlateColor(DefaultForegroundName);
-}
-
 TSharedRef<SWidget> SAssetView::GetViewButtonContent()
 {
 	SAssetView::RegisterGetViewButtonMenu();
@@ -2200,6 +2220,8 @@ TSharedRef<SWidget> SAssetView::GetViewButtonContent()
 
 	UContentBrowserAssetViewContextMenuContext* Context = NewObject<UContentBrowserAssetViewContextMenuContext>();
 	Context->AssetView = SharedThis(this);
+	Context->OwningContentBrowser = OwningContentBrowser;
+
 	FToolMenuContext MenuContext(nullptr, MenuExtender, Context);
 	return UToolMenus::Get()->GenerateWidget("ContentBrowser.AssetViewOptions", MenuContext);
 }
@@ -2451,16 +2473,33 @@ void SAssetView::PopulateViewButtonMenu(UToolMenu* Menu)
 
 	{
 		FToolMenuSection& Section = Menu->AddSection("AssetThumbnails", LOCTEXT("ThumbnailsHeading", "Thumbnails"));
-		Section.AddEntry(FToolMenuEntry::InitWidget(
-			"ThumbnailScale",
-			SNew(SSlider)
-				.ToolTipText( LOCTEXT("ThumbnailScaleToolTip", "Adjust the size of thumbnails.") )
-				.Value( this, &SAssetView::GetThumbnailScale )
-				.OnValueChanged( this, &SAssetView::SetThumbnailScale )
-				.Locked( this, &SAssetView::IsThumbnailScalingLocked ),
-			LOCTEXT("ThumbnailScaleLabel", "Scale"),
-			/*bNoIndent=*/true
-			));
+
+		auto CreateThumbnailSizeSubMenu = [this](UToolMenu* SubMenu)
+		{
+			FToolMenuSection& SizeSection = SubMenu->AddSection("ThumbnailSizes");
+			
+			for (int32 EnumValue = (int32)EThumbnailSize::Tiny; EnumValue < (int32)EThumbnailSize::MAX; ++EnumValue)
+			{
+				SizeSection.AddMenuEntry(
+					NAME_None,
+					ThumbnailSizeToDisplayName((EThumbnailSize)EnumValue),
+					FText::GetEmpty(),
+					FSlateIcon(),
+					FUIAction(
+						FExecuteAction::CreateSP(this, &SAssetView::OnThumbnailSizeChanged, (EThumbnailSize)EnumValue),
+						FCanExecuteAction::CreateSP(this, &SAssetView::IsThumbnailScalingAllowed),
+						FIsActionChecked::CreateSP(this, &SAssetView::IsThumbnailSizeChecked, (EThumbnailSize)EnumValue)
+					),
+					EUserInterfaceActionType::RadioButton
+				);
+			}
+		};
+		Section.AddEntry(FToolMenuEntry::InitSubMenu(
+			"ThumbnailSize",
+			LOCTEXT("ThumbnailSize", "Thumbnail Size"),
+			LOCTEXT("ThumbnailSizeToolTip", "Adjust the size of thumbnails."),
+			FNewToolMenuDelegate::CreateLambda(CreateThumbnailSizeSubMenu)
+		));
 
 		Section.AddMenuEntry(
 			"ThumbnailEditMode",
@@ -2961,7 +3000,7 @@ TSharedRef<ITableRow> SAssetView::MakeListViewWidget(TSharedPtr<FAssetViewItem>
 	{
 		TSharedPtr< STableRow<TSharedPtr<FAssetViewItem>> > TableRowWidget;
 		SAssignNew( TableRowWidget, STableRow<TSharedPtr<FAssetViewItem>>, OwnerTable )
-			.Style(FEditorStyle::Get(), "ContentBrowser.AssetListView.TableRow")
+			.Style(FEditorStyle::Get(), "ContentBrowser.AssetListView.ColumnListTableRow")
 			.Cursor( bAllowDragging ? EMouseCursor::GrabHand : EMouseCursor::Default )
 			.OnDragDetected( this, &SAssetView::OnDraggingAssetItem );
 
@@ -2994,7 +3033,7 @@ TSharedRef<ITableRow> SAssetView::MakeListViewWidget(TSharedPtr<FAssetViewItem>
 
 		TSharedPtr< STableRow<TSharedPtr<FAssetViewItem>> > TableRowWidget;
 		SAssignNew( TableRowWidget, STableRow<TSharedPtr<FAssetViewItem>>, OwnerTable )
-		.Style(FEditorStyle::Get(), "ContentBrowser.AssetListView.TableRow")
+		.Style(FEditorStyle::Get(), "ContentBrowser.AssetListView.ColumnListTableRow")
 		.Cursor( bAllowDragging ? EMouseCursor::GrabHand : EMouseCursor::Default )
 		.OnDragDetected( this, &SAssetView::OnDraggingAssetItem );
 
@@ -3040,7 +3079,7 @@ TSharedRef<ITableRow> SAssetView::MakeTileViewWidget(TSharedPtr<FAssetViewItem>
 	{
 		TSharedPtr< STableRow<TSharedPtr<FAssetViewItem>> > TableRowWidget;
 		SAssignNew( TableRowWidget, STableRow<TSharedPtr<FAssetViewItem>>, OwnerTable )
-			.Style( FEditorStyle::Get(), "ContentBrowser.AssetListView.TableRow" )
+			.Style( FEditorStyle::Get(), "ContentBrowser.AssetListView.TileTableRow" )
 			.Cursor( bAllowDragging ? EMouseCursor::GrabHand : EMouseCursor::Default )
 			.OnDragDetected( this, &SAssetView::OnDraggingAssetItem );
 
@@ -3054,7 +3093,8 @@ TSharedRef<ITableRow> SAssetView::MakeTileViewWidget(TSharedPtr<FAssetViewItem>
 			.OnItemDestroyed(this, &SAssetView::AssetItemWidgetDestroyed)
 			.ShouldAllowToolTip(this, &SAssetView::ShouldAllowToolTips)
 			.HighlightText( HighlightedText )
-			.IsSelected( FIsSelected::CreateSP(TableRowWidget.Get(), &STableRow<TSharedPtr<FAssetViewItem>>::IsSelectedExclusively) );
+			.IsSelected(FIsSelected::CreateSP(TableRowWidget.Get(), &STableRow<TSharedPtr<FAssetViewItem>>::IsSelected))
+			.IsSelectedExclusively(FIsSelected::CreateSP(TableRowWidget.Get(), &STableRow<TSharedPtr<FAssetViewItem>>::IsSelectedExclusively));
 
 		TableRowWidget->SetContent(Item);
 
@@ -3073,7 +3113,7 @@ TSharedRef<ITableRow> SAssetView::MakeTileViewWidget(TSharedPtr<FAssetViewItem>
 
 		TSharedPtr< STableRow<TSharedPtr<FAssetViewItem>> > TableRowWidget;
 		SAssignNew( TableRowWidget, STableRow<TSharedPtr<FAssetViewItem>>, OwnerTable )
-		.Style(FEditorStyle::Get(), "ContentBrowser.AssetListView.TableRow")
+		.Style(FEditorStyle::Get(), "ContentBrowser.AssetListView.TileTableRow")
 		.Cursor( bAllowDragging ? EMouseCursor::GrabHand : EMouseCursor::Default )
 		.OnDragDetected( this, &SAssetView::OnDraggingAssetItem );
 
@@ -3093,7 +3133,8 @@ TSharedRef<ITableRow> SAssetView::MakeTileViewWidget(TSharedPtr<FAssetViewItem>
 			.ThumbnailLabel( ThumbnailLabel )
 			.ThumbnailHintColorAndOpacity( this, &SAssetView::GetThumbnailHintColorAndOpacity )
 			.AllowThumbnailHintLabel( AllowThumbnailHintLabel )
-			.IsSelected( FIsSelected::CreateSP(TableRowWidget.Get(), &STableRow<TSharedPtr<FAssetViewItem>>::IsSelectedExclusively) )
+			.IsSelected(FIsSelected::CreateSP(TableRowWidget.Get(), &STableRow<TSharedPtr<FAssetViewItem>>::IsSelected))
+			.IsSelectedExclusively( FIsSelected::CreateSP(TableRowWidget.Get(), &STableRow<TSharedPtr<FAssetViewItem>>::IsSelectedExclusively) )
 			.OnIsAssetValidForCustomToolTip(OnIsAssetValidForCustomToolTip)
 			.OnGetCustomAssetToolTip(OnGetCustomAssetToolTip)
 			.OnVisualizeAssetToolTip( OnVisualizeAssetToolTip )
@@ -3110,7 +3151,7 @@ TSharedRef<ITableRow> SAssetView::MakeColumnViewWidget(TSharedPtr<FAssetViewItem
 	if ( !ensure(AssetItem.IsValid()) )
 	{
 		return SNew( STableRow<TSharedPtr<FAssetViewItem>>, OwnerTable )
-			.Style(FEditorStyle::Get(), "ContentBrowser.AssetListView.TableRow");
+			.Style(FEditorStyle::Get(), "ContentBrowser.AssetListView.ColumnListTableRow");
 	}
 
 	// Update the cached custom data
@@ -3676,20 +3717,49 @@ void SAssetView::ToggleThumbnailEditMode()
 	bThumbnailEditMode = !bThumbnailEditMode;
 }
 
-float SAssetView::GetThumbnailScale() const
+
+void SAssetView::OnThumbnailSizeChanged(EThumbnailSize NewThumbnailSize)
 {
-	return ThumbnailScaleSliderValue.Get();
+	ThumbnailSize = NewThumbnailSize;
+	RefreshList();
 }
 
-void SAssetView::SetThumbnailScale( float NewValue )
+bool SAssetView::IsThumbnailSizeChecked(EThumbnailSize InThumbnailSize) const
 {
-	ThumbnailScaleSliderValue = NewValue;
-	RefreshList();
+	return ThumbnailSize == InThumbnailSize;
+}
+
+float SAssetView::GetThumbnailScale() const
+{
+	float BaseScale;
+	switch (ThumbnailSize)
+	{
+	case EThumbnailSize::Tiny:
+		BaseScale = 0.1f;
+		break;
+	case EThumbnailSize::Small:
+		BaseScale = 0.25f;
+		break;
+	case EThumbnailSize::Medium:
+		BaseScale = 0.5f;
+		break;
+	case EThumbnailSize::Large:
+		BaseScale = 0.75f;
+		break;
+	case EThumbnailSize::Huge:
+		BaseScale = 1.0f;
+		break;
+	default:
+		BaseScale = 0.5f;
+		break;
+	}
+
+	return BaseScale * GetTickSpaceGeometry().Scale;
 }
 
-bool SAssetView::IsThumbnailScalingLocked() const
+bool SAssetView::IsThumbnailScalingAllowed() const
 {
-	return GetCurrentViewType() == EAssetViewType::Column;
+	return GetCurrentViewType() != EAssetViewType::Column;
 }
 
 float SAssetView::GetListViewItemHeight() const
@@ -3699,7 +3769,7 @@ float SAssetView::GetListViewItemHeight() const
 
 float SAssetView::GetTileViewItemHeight() const
 {
-	return TileViewNameHeight + GetTileViewItemBaseHeight() * FillScale;
+	return (TileViewNameHeight * FMath::Lerp(MinThumbnailScale, MaxThumbnailScale, GetThumbnailScale())) + GetTileViewItemBaseHeight() * FillScale;
 }
 
 float SAssetView::GetTileViewItemBaseHeight() const
diff --git a/Engine/Source/Editor/ContentBrowser/Private/SAssetView.h b/Engine/Source/Editor/ContentBrowser/Private/SAssetView.h
--- a/Engine/Source/Editor/ContentBrowser/Private/SAssetView.h
+++ b/Engine/Source/Editor/ContentBrowser/Private/SAssetView.h
@@ -34,6 +34,7 @@ class SAssetTileView;
 class SComboButton;
 class UFactory;
 struct FPropertyChangedEvent;
+class SContentBrowser;
 
 /** Fires whenever the asset view is asked to start to create a temporary item */
 DECLARE_DELEGATE_OneParam(FOnAssetViewNewItemRequested, const FContentBrowserItem& /*NewItem*/);
@@ -41,6 +42,19 @@ DECLARE_DELEGATE_OneParam(FOnAssetViewNewItemRequested, const FContentBrowserIte
 /** Fires whenever one of the "Search" options changes, useful for modifying search criteria to match */
 DECLARE_DELEGATE(FOnSearchOptionChanged);
 
+enum class EThumbnailSize : uint8
+{
+	Tiny = 0,
+	Small,
+	Medium,
+	Large,
+	Huge,
+
+	// Not a size
+	MAX
+};
+
+
 /**
  * A widget to display a list of filtered assets
  */
@@ -53,8 +67,8 @@ public:
 		, _ThumbnailLabel( EThumbnailLabel::ClassName )
 		, _AllowThumbnailHintLabel(true)
 		, _InitialViewType(EAssetViewType::Tile)
-		, _ThumbnailScale(0.1f)
 		, _ShowBottomToolbar(true)
+		, _ShowViewOptions(true)
 		, _AllowThumbnailEditMode(false)
 		, _CanShowClasses(true)
 		, _CanShowFolders(false)
@@ -136,13 +150,12 @@ public:
 
 		/** The initial view type */
 		SLATE_ARGUMENT( EAssetViewType::Type, InitialViewType )
-
-		/** The thumbnail scale. [0-1] where 0.5 is no scale */
-		SLATE_ATTRIBUTE( float, ThumbnailScale )
-
 		/** Should the toolbar indicating number of selected assets, mode switch buttons, etc... be shown? */
 		SLATE_ARGUMENT( bool, ShowBottomToolbar )
 
+		/** Should view options be shown. Note: If ShowBottomToolbar is false, then view options are not shown regardless of this setting */
+		SLATE_ARGUMENT( bool, ShowViewOptions)
+
 		/** True if the asset view may edit thumbnails */
 		SLATE_ARGUMENT( bool, AllowThumbnailEditMode )
 
@@ -209,6 +222,8 @@ public:
 		/** Custom columns that can be use specific */
 		SLATE_EVENT( FOnSearchOptionChanged, OnSearchOptionsChanged)
 
+		/** The content browser that owns this view if any */
+		SLATE_ARGUMENT(TSharedPtr<SContentBrowser>, OwningContentBrowser)
 	SLATE_END_ARGS()
 
 	~SAssetView();
@@ -332,6 +347,12 @@ public:
 	/** @return true when we are including collection names in search criteria */
 	bool IsIncludingCollectionNames() const;
 
+	/** Handler for when the view combo button is clicked */
+	TSharedRef<SWidget> GetViewButtonContent();
+
+	/** Gets the text for the asset count label */
+	FText GetAssetCountText() const;
+
 private:
 
 	/** Sets the pending selection to the current selection (used when changing views or refreshing the view). */
@@ -406,12 +427,6 @@ private:
 	/** Returns the thumbnails hint color and opacity */
 	FLinearColor GetThumbnailHintColorAndOpacity() const;
 
-	/** Returns the foreground color for the view button */
-	FSlateColor GetViewButtonForegroundColor() const;
-
-	/** Handler for when the view combo button is clicked */
-	TSharedRef<SWidget> GetViewButtonContent();
-
 	/** Register menu for when the view combo button is clicked */
 	static void RegisterGetViewButtonMenu();
 
@@ -621,9 +636,6 @@ private:
 	/** The "Done Editing" button was pressed in the thumbnail edit mode strip */
 	FReply EndThumbnailEditModeClicked();
 
-	/** Gets the text for the asset count label */
-	FText GetAssetCountText() const;
-
 	/** Gets the visibility of the Thumbnail Edit Mode label */
 	EVisibility GetEditModeLabelVisibility() const;
 
@@ -639,14 +651,17 @@ private:
 	/** Toggles thumbnail editing mode */
 	void ToggleThumbnailEditMode();
 
-	/** Gets the current value for the scale slider (0 to 1) */
-	float GetThumbnailScale() const;
+	/** Called when  thumbnail size is changed */
+	void OnThumbnailSizeChanged(EThumbnailSize NewThumbnailSize);
+
+	/** Is the current thumbnail size menu option checked */
+	bool IsThumbnailSizeChecked(EThumbnailSize InThumbnailSize) const;
 
-	/** Sets the current scale value (0 to 1) */
-	void SetThumbnailScale( float NewValue );
+	/** Are thumbnails allowed to be scaled by the user */
+	bool IsThumbnailScalingAllowed() const;
 
-	/** Is thumbnail scale slider locked? */
-	bool IsThumbnailScalingLocked() const;
+	/** Gets the current thumbnail scale */
+	float GetThumbnailScale() const;
 
 	/** Gets the scaled item height for the list view */
 	float GetListViewItemHeight() const;
@@ -779,10 +794,10 @@ private:
 	TSharedPtr<class SAssetListView> ListView;
 	TSharedPtr<class SAssetTileView> TileView;
 	TSharedPtr<class SAssetColumnView> ColumnView;
-	TSharedPtr<SBorder> ViewContainer;
+	TSharedPtr<SBox> ViewContainer;
 
-	/** The button that displays view options */
-	TSharedPtr<SComboButton> ViewOptionsComboButton;
+	/** The content browser that created this asset view if any */
+	TWeakPtr<SContentBrowser> OwningContentBrowser;
 
 	/** The current base source filter for the view */
 	FSourcesData SourcesData;
@@ -884,13 +899,16 @@ private:
 	int32 TileViewThumbnailPadding;
 	int32 TileViewNameHeight;
 
-	/** The current value for the thumbnail scale from the thumbnail slider */
-	TAttribute< float > ThumbnailScaleSliderValue;
-
 	/** The max and min thumbnail scales as a fraction of the rendered size */
 	float MinThumbnailScale;
 	float MaxThumbnailScale;
 
+	/** Scalar applied to thumbnail sizes so that users thumbnails are scaled based on users display area size*/
+	float ThumbnailScaleRangeScalar;
+
+	/** Current thumbnail size */
+	EThumbnailSize ThumbnailSize;
+
 	/** Flag indicating if we will be filling the empty space in the tile view. */
 	bool bFillEmptySpaceInTileView;
 
diff --git a/Engine/Source/Editor/ContentBrowser/Private/SCollectionPicker.cpp b/Engine/Source/Editor/ContentBrowser/Private/SCollectionPicker.cpp
--- a/Engine/Source/Editor/ContentBrowser/Private/SCollectionPicker.cpp
+++ b/Engine/Source/Editor/ContentBrowser/Private/SCollectionPicker.cpp
@@ -12,7 +12,6 @@ void SCollectionPicker::Construct( const FArguments& InArgs )
 		SNew(SCollectionView)
 		.AllowCollectionButtons(InArgs._CollectionPickerConfig.AllowCollectionButtons)
 		.OnCollectionSelected(InArgs._CollectionPickerConfig.OnCollectionSelected)
-		.AllowCollapsing( false )
 		.AllowContextMenu( InArgs._CollectionPickerConfig.AllowRightClickMenu )
 	];
 }
diff --git a/Engine/Source/Editor/ContentBrowser/Private/SCollectionView.cpp b/Engine/Source/Editor/ContentBrowser/Private/SCollectionView.cpp
--- a/Engine/Source/Editor/ContentBrowser/Private/SCollectionView.cpp
+++ b/Engine/Source/Editor/ContentBrowser/Private/SCollectionView.cpp
@@ -89,6 +89,8 @@ void SCollectionView::Construct( const FArguments& InArgs )
 	bQueueCollectionItemsUpdate = false;
 	bQueueSCCRefresh = true;
 
+	IsDocked = InArgs._IsDocked;
+
 	FCollectionManagerModule& CollectionManagerModule = FCollectionManagerModule::GetModule();
 	CollectionManagerModule.Get().OnCollectionCreated().AddSP( this, &SCollectionView::HandleCollectionCreated );
 	CollectionManagerModule.Get().OnCollectionRenamed().AddSP( this, &SCollectionView::HandleCollectionRenamed );
@@ -122,65 +124,45 @@ void SCollectionView::Construct( const FArguments& InArgs )
 		CollectionListContextMenuOpening = FOnContextMenuOpening::CreateSP( this, &SCollectionView::MakeCollectionTreeContextMenu );
 	}
 
+	SearchPtr = InArgs._ExternalSearch;
+	if (SearchPtr)
+	{
+		SearchPtr->OnSearchChanged().AddSP(this, &SCollectionView::SetCollectionsSearchFilterText);
+	}
+
 	ExternalSearchPtr = InArgs._ExternalSearch;
 	TitleContent = SNew(SHorizontalBox);
-	SetAllowExternalSearch(true);
+
 
 	PreventSelectionChangedDelegateCount = 0;
 
 	TSharedRef< SWidget > HeaderContent = SNew(SHorizontalBox)
+			.Visibility(this, &SCollectionView::GetHeaderVisibility)
 			+ SHorizontalBox::Slot()
 			.FillWidth(1.0f)
 			.Padding(0.0f)
 			[
 				TitleContent.ToSharedRef()
 			]
-
-			+SHorizontalBox::Slot()
-			.AutoWidth()
-			.VAlign(VAlign_Center)
-			.Padding(2.0f, 0.0f, 0.0f, 0.0f)
-			[
-				SNew(SButton)
-				.ButtonStyle(FEditorStyle::Get(), "FlatButton")
-				.ToolTipText(this, &SCollectionView::GetSwitchCollectionViewModeToolTipText)
-				.OnClicked(this, &SCollectionView::OnSwitchCollectionViewMode)
-				.ContentPadding(FMargin(2, 2))
-				.ForegroundColor(FEditorStyle::GetSlateColor("DefaultForeground"))
-				.Visibility(this, &SCollectionView::GetCollectionButtonsVisibility)
-				[
-					SNew(SImage)
-					.Image(this, &SCollectionView::GetSwitchCollectionViewModeIcon)
-				]
-			]
-
 			+SHorizontalBox::Slot()
 			.AutoWidth()
 			.VAlign(VAlign_Center)
 			.Padding(2.0f, 0.0f, 0.0f, 0.0f)
 			[
 				SNew(SButton)
-				.ButtonStyle(FEditorStyle::Get(), "FlatButton")
+				.ButtonStyle(FEditorStyle::Get(), "SimpleButton")
 				.ToolTipText(LOCTEXT("AddCollectionButtonTooltip", "Add a collection."))
 				.OnClicked(this, &SCollectionView::MakeAddCollectionMenu)
 				.ContentPadding( FMargin(2, 2) )
-				.Visibility(this, &SCollectionView::GetCollectionButtonsVisibility)
+				.Visibility(bAllowCollectionButtons ? EVisibility::Visible : EVisibility::Collapsed)
 				[
 					SNew(SImage)
-					.Image( FEditorStyle::GetBrush("ContentBrowser.AddCollectionButtonIcon") )
+					.Image(FAppStyle::Get().GetBrush("Icons.PlusCircle"))
+					.ColorAndOpacity(FSlateColor::UseForeground())
 				]
 			];
 
 	TSharedRef< SWidget > BodyContent = SNew(SVerticalBox)
-			// Separator
-			+SVerticalBox::Slot()
-			.AutoHeight()
-			.Padding(0, 0, 0, 1)
-			[
-				SNew(SSeparator)
-				.Visibility(InArgs._ShowSeparator ? EVisibility::Visible : EVisibility::Collapsed)
-			]
-
 			// Collections tree
 			+SVerticalBox::Slot()
 			.FillHeight(1.f)
@@ -198,42 +180,6 @@ void SCollectionView::Construct( const FArguments& InArgs )
 				.Visibility(this, &SCollectionView::GetCollectionTreeVisibility)
 			];
 
-	TSharedPtr< SWidget > Content;
-	if ( InArgs._AllowCollapsing )
-	{
-		Content = SAssignNew(CollectionsExpandableAreaPtr, SExpandableArea)
-			.MaxHeight(200)
-			.BorderImage( FEditorStyle::GetBrush("NoBorder") )
-			.HeaderPadding( FMargin(4.0f, 0.0f, 0.0f, 0.0f) )
-			.HeaderContent()
-			[
-				SNew(SBox)
-				.Padding(FMargin(6.0f, 0.0f, 0.0f, 0.0f))
-				[
-					HeaderContent
-				]
-			]
-			.BodyContent()
-			[
-				BodyContent
-			];
-	}
-	else
-	{
-		Content = SNew( SVerticalBox )
-		+SVerticalBox::Slot()
-		.AutoHeight()
-		.Padding(FMargin(12.0f, 0.0f, 0.0f, 0.0f))
-		[
-			HeaderContent
-		]
-
-		+SVerticalBox::Slot()
-		[
-			BodyContent
-		];
-	}
-
 	ChildSlot
 	[
 		SNew(SOverlay)
@@ -241,9 +187,18 @@ void SCollectionView::Construct( const FArguments& InArgs )
 		// Main content
 		+SOverlay::Slot()
 		[
-			Content.ToSharedRef()
+			SNew(SVerticalBox)
+			+ SVerticalBox::Slot()
+			.AutoHeight()
+			.Padding(FMargin(12.0f, 0.0f, 0.0f, 0.0f))
+			[
+				HeaderContent
+			]
+			+ SVerticalBox::Slot()
+			[
+				BodyContent
+			]
 		]
-
 		// Drop target overlay
 		+SOverlay::Slot()
 		[
@@ -261,61 +216,6 @@ void SCollectionView::Construct( const FArguments& InArgs )
 	UpdateCollectionItems();
 }
 
-void SCollectionView::SetAllowExternalSearch(const bool InAllowExternalSearch)
-{
-	if (bAllowExternalSearch == InAllowExternalSearch)
-	{
-		return;
-	}
-	bAllowExternalSearch = InAllowExternalSearch;
-
-	if (SearchPtr)
-	{
-		SearchPtr->OnSearchChanged().RemoveAll(this);
-	}
-	TitleContent->ClearChildren();
-
-	SearchPtr = bAllowExternalSearch ? ExternalSearchPtr : nullptr;
-	if (!SearchPtr)
-	{
-		SearchPtr = MakeShared<FSourcesSearch>();
-		SearchPtr->Initialize();
-		SearchPtr->SetHintText(LOCTEXT("CollectionsViewSearchBoxHint", "Search Collections"));
-	}
-	SearchPtr->OnSearchChanged().AddSP(this, &SCollectionView::SetCollectionsSearchFilterText);
-
-	if (SearchPtr == ExternalSearchPtr)
-	{
-		// If using an external search then just show the title and don't hide it when collapsing
-		TitleContent->AddSlot()
-		[
-			SNew(STextBlock)
-			.Font(FEditorStyle::GetFontStyle("ContentBrowser.SourceTitleFont"))
-			.Text(LOCTEXT("CollectionsListTitle", "Collections"))
-		];
-	}
-	else
-	{
-		// If using an internal search then show the title or search box depending on whether we're collapsed or not
-		TitleContent->AddSlot()
-		[
-			SNew(STextBlock)
-			.Font(FEditorStyle::GetFontStyle("ContentBrowser.SourceTitleFont"))
-			.Text(LOCTEXT("CollectionsListTitle", "Collections"))
-			.Visibility(this, &SCollectionView::GetCollectionsTitleTextVisibility)
-		];
-
-		TitleContent->AddSlot()
-		[
-			SNew(SBox)
-			.VAlign(VAlign_Center)
-			.Visibility(this, &SCollectionView::GetCollectionsSearchBoxVisibility)
-			[
-				SearchPtr->GetWidget()
-			]
-		];
-	}
-}
 
 void SCollectionView::HandleCollectionCreated( const FCollectionNameType& Collection )
 {
@@ -628,11 +528,6 @@ void SCollectionView::SetSelectedCollections(const TArray<FCollectionNameType>&
 	// Prevent the selection changed delegate since the invoking code requested it
 	FScopedPreventSelectionChangedDelegate DelegatePrevention( SharedThis(this) );
 
-	// Expand the collections area if we are indeed selecting at least one collection
-	if ( bEnsureVisible && CollectionsToSelect.Num() > 0 && CollectionsExpandableAreaPtr.IsValid() )
-	{
-		CollectionsExpandableAreaPtr->SetExpanded(true);
-	}
 
 	// Clear the selection to start, then add the selected items as they are found
 	CollectionTreePtr->ClearSelection();
@@ -744,16 +639,12 @@ void SCollectionView::SaveSettings(const FString& IniFilename, const FString& In
 		GConfig->SetString(*IniSection, *(SettingsString + InSubKey), *CollectionsString, IniFilename);
 	};
 
-	const bool IsCollectionsExpanded = CollectionsExpandableAreaPtr.IsValid() ? CollectionsExpandableAreaPtr->IsExpanded() : true;
-	GConfig->SetBool(*IniSection, *(SettingsString + TEXT(".CollectionsExpanded")), IsCollectionsExpanded, IniFilename);
 	SaveCollectionsArrayToIni(TEXT(".SelectedCollections"), CollectionTreePtr->GetSelectedItems());
 	{
 		TSet<TSharedPtr<FCollectionItem>> ExpandedCollectionItems;
 		CollectionTreePtr->GetExpandedItems(ExpandedCollectionItems);
 		SaveCollectionsArrayToIni(TEXT(".ExpandedCollections"), ExpandedCollectionItems.Array());
 	}
-
-	GConfig->SetInt(*IniSection, *(SettingsString + TEXT(".ViewMode")), (int32)CollectionViewMode, IniFilename);
 }
 
 void SCollectionView::LoadSettings(const FString& IniFilename, const FString& IniSection, const FString& SettingsString)
@@ -786,13 +677,6 @@ void SCollectionView::LoadSettings(const FString& IniFilename, const FString& In
 		return RetCollectionsArray;
 	};
 
-	// Collection expansion state
-	bool bCollectionsExpanded = false;
-	if (CollectionsExpandableAreaPtr.IsValid() && GConfig->GetBool(*IniSection, *(SettingsString + TEXT(".CollectionsExpanded")), bCollectionsExpanded, IniFilename))
-	{
-		CollectionsExpandableAreaPtr->SetExpanded(bCollectionsExpanded);
-	}
-
 	// Selected Collections
 	TArray<FCollectionNameType> NewSelectedCollections = LoadCollectionsArrayFromIni(TEXT(".SelectedCollections"));
 	if (NewSelectedCollections.Num() > 0)
@@ -812,15 +696,6 @@ void SCollectionView::LoadSettings(const FString& IniFilename, const FString& In
 	{
 		SetExpandedCollections(NewExpandedCollections);
 	}
-
-	// View Mode
-	{
-		int32 CollectionViewModeInt = 0;
-		GConfig->GetInt(*IniSection, *(SettingsString + TEXT(".ViewMode")), CollectionViewModeInt, IniFilename);
-		CollectionViewMode = (EAssetTagItemViewMode)CollectionViewModeInt;
-
-		CollectionTreePtr->RebuildList();
-	}
 }
 
 void SCollectionView::Tick( const FGeometry& AllottedGeometry, const double InCurrentTime, const float InDeltaTime )
@@ -975,78 +850,6 @@ FReply SCollectionView::MakeAddCollectionMenu()
 	return FReply::Handled();
 }
 
-EVisibility SCollectionView::GetCollectionsTitleTextVisibility() const
-{
-	// Only show the title text if we have an expansion area, but are collapsed
-	return (CollectionsExpandableAreaPtr.IsValid() && !CollectionsExpandableAreaPtr->IsExpanded()) ? EVisibility::Visible : EVisibility::Collapsed;
-}
-
-EVisibility SCollectionView::GetCollectionsSearchBoxVisibility() const
-{
-	// Only show the search box if we have an expanded expansion area, or aren't currently using an expansion area
-	return (!CollectionsExpandableAreaPtr.IsValid() || CollectionsExpandableAreaPtr->IsExpanded()) ? EVisibility::Visible : EVisibility::Collapsed;
-}
-
-EVisibility SCollectionView::GetCollectionButtonsVisibility() const
-{
-	return (bAllowCollectionButtons && ( !CollectionsExpandableAreaPtr.IsValid() || CollectionsExpandableAreaPtr->IsExpanded() ) ) ? EVisibility::Visible : EVisibility::Collapsed;
-}
-
-const FSlateBrush* SCollectionView::GetSwitchCollectionViewModeIcon() const
-{
-	switch (CollectionViewMode)
-	{
-	case EAssetTagItemViewMode::Standard:
-		return FEditorStyle::GetBrush("ContentBrowser.Sources.Collections.Compact");
-
-	case EAssetTagItemViewMode::Compact:
-		return FEditorStyle::GetBrush("ContentBrowser.Sources.Collections");
-
-	default:
-		break;
-	}
-	check(false);
-	return nullptr;
-}
-
-FText SCollectionView::GetSwitchCollectionViewModeToolTipText() const
-{
-	switch (CollectionViewMode)
-	{
-	case EAssetTagItemViewMode::Standard:
-		return LOCTEXT("SwitchToCompactView_ToolTip", "Switch to compact view");
-
-	case EAssetTagItemViewMode::Compact:
-		return LOCTEXT("SwitchToStanardView_ToolTip", "Switch to standard view");
-
-	default:
-		break;
-	}
-	check(false);
-	return FText();
-}
-
-FReply SCollectionView::OnSwitchCollectionViewMode()
-{
-	switch (CollectionViewMode)
-	{
-	case EAssetTagItemViewMode::Standard:
-		CollectionViewMode = EAssetTagItemViewMode::Compact;
-		break;
-
-	case EAssetTagItemViewMode::Compact:
-		CollectionViewMode = EAssetTagItemViewMode::Standard;
-		break;
-
-	default:
-		check(false);
-		break;
-	}
-
-	CollectionTreePtr->RebuildList();
-	return FReply::Handled();
-}
-
 void SCollectionView::CreateCollectionItem( ECollectionShareType::Type CollectionType, ECollectionStorageMode::Type StorageMode, const FCreateCollectionPayload& InCreationPayload )
 {
 	if ( ensure(CollectionType != ECollectionShareType::CST_All) )
@@ -1232,6 +1035,11 @@ EVisibility SCollectionView::GetCollectionTreeVisibility() const
 	return AvailableCollections.Num() > 0 ? EVisibility::Visible : EVisibility::Collapsed;
 }
 
+EVisibility SCollectionView::GetHeaderVisibility() const
+{
+	return IsDocked.Get() ? EVisibility::Collapsed : EVisibility::SelfHitTestInvisible;
+}
+
 const FSlateBrush* SCollectionView::GetCollectionViewDropTargetBorder() const
 {
 	return bDraggedOver ? FEditorStyle::GetBrush("ContentBrowser.CollectionTreeDragDropBorder") : FEditorStyle::GetBrush("NoBorder");
@@ -1271,7 +1079,6 @@ TSharedRef<ITableRow> SCollectionView::GenerateCollectionRow( TSharedPtr<FCollec
 			SAssignNew(CollectionTreeItem, SCollectionTreeItem)
 			.ParentWidget(SharedThis(this))
 			.CollectionItem(CollectionItem)
-			.ViewMode(CollectionViewMode)
 			.OnNameChangeCommit(this, &SCollectionView::CollectionNameChangeCommit)
 			.OnVerifyRenameCommit(this, &SCollectionView::CollectionVerifyRenameCommit)
 			.OnValidateDragDrop(this, &SCollectionView::ValidateDragDropOnCollectionItem)
diff --git a/Engine/Source/Editor/ContentBrowser/Private/SCollectionView.h b/Engine/Source/Editor/ContentBrowser/Private/SCollectionView.h
--- a/Engine/Source/Editor/ContentBrowser/Private/SCollectionView.h
+++ b/Engine/Source/Editor/ContentBrowser/Private/SCollectionView.h
@@ -37,8 +37,8 @@ public:
 	SLATE_BEGIN_ARGS( SCollectionView )
 		: _AllowCollectionButtons(true)
 		, _AllowRightClickMenu(true)
-		, _AllowCollapsing(true)
 		, _AllowContextMenu(true)
+		, _IsDocked(false)
 		, _AllowCollectionDrag(false)
 		, _AllowQuickAssetManagement(false)
 		, _ShowSeparator(true)
@@ -50,8 +50,8 @@ public:
 		/** If true, collection buttons will be displayed */
 		SLATE_ARGUMENT( bool, AllowCollectionButtons )
 		SLATE_ARGUMENT( bool, AllowRightClickMenu )
-		SLATE_ARGUMENT( bool, AllowCollapsing )
 		SLATE_ARGUMENT( bool, AllowContextMenu )
+		SLATE_ATTRIBUTE( bool, IsDocked)
 
 		/** If true, the user will be able to drag collections from this view */
 		SLATE_ARGUMENT( bool, AllowCollectionDrag )
@@ -70,9 +70,6 @@ public:
 	/** Constructs this widget with InArgs */
 	void Construct( const FArguments& InArgs );
 
-	/** Set whether we're using an external search or not */
-	void SetAllowExternalSearch(const bool InAllowExternalSearch);
-
 	/** Selects the specified collections */
 	void SetSelectedCollections(const TArray<FCollectionNameType>& CollectionsToSelect, const bool bEnsureVisible = true);
 
@@ -107,6 +104,9 @@ public:
 	/** Creates the menu for the save dynamic collection button */
 	void MakeSaveDynamicCollectionMenu(FText InQueryString);
 
+	/** Creates the menu for the add collection button */
+	FReply MakeAddCollectionMenu();
+
 private:
 
 	/** Payload data for CreateCollectionItem (likely from a delegate binding) */
@@ -146,27 +146,6 @@ private:
 	/** True if the selection changed delegate is allowed at the moment */
 	bool ShouldAllowSelectionChangedDelegate() const;
 
-	/** Creates the menu for the add collection button */
-	FReply MakeAddCollectionMenu();
-
-	/** Gets the visibility of the collections title text */
-	EVisibility GetCollectionsTitleTextVisibility() const;
-
-	/** Gets the visibility of the collections search box */
-	EVisibility GetCollectionsSearchBoxVisibility() const;
-
-	/** Gets the visibility of the collection buttons */
-	EVisibility GetCollectionButtonsVisibility() const;
-	
-	/** Get the icon used on the collection view mode switcher button */
-	const FSlateBrush* GetSwitchCollectionViewModeIcon() const;
-
-	/** Get the tooltip text used on the collection view mode switcher button */
-	FText GetSwitchCollectionViewModeToolTipText() const;
-
-	/** Handle the collection view mode switcher button being clicked */
-	FReply OnSwitchCollectionViewMode();
-
 	/** Sets up an inline creation process for a new collection of the specified type */
 	void CreateCollectionItem( ECollectionShareType::Type CollectionType, ECollectionStorageMode::Type StorageMode, const FCreateCollectionPayload& InCreationPayload );
 
@@ -179,6 +158,9 @@ private:
 	/** Returns the visibility of the collection tree */
 	EVisibility GetCollectionTreeVisibility() const;
 
+	/** Returns the visibility of the collection tree header. Collapsed if the collection tree is docked */
+	EVisibility GetHeaderVisibility() const;
+
 	/** Get the border of the collection tree */
 	const FSlateBrush* GetCollectionViewDropTargetBorder() const;
 
@@ -274,7 +256,6 @@ private:
 
 	/** Get the active filter text */
 	FText GetCollectionsSearchFilterText() const;
-
 private:
 
 	/** A helper class to manage PreventSelectionChangedDelegateCount by incrementing it when constructed (on the stack) and decrementing when destroyed */
@@ -326,12 +307,12 @@ private:
 	/** The context menu logic and data */
 	TSharedPtr<class FCollectionContextMenu> CollectionContextMenu;
 
-	/** The collections SExpandableArea */
-	TSharedPtr< SExpandableArea > CollectionsExpandableAreaPtr;
-
 	/** Delegate to invoke when selection changes. */
 	FOnCollectionSelected OnCollectionSelected;
 
+	/** Whether or not the collection view is docked under the sources panel */
+	TAttribute<bool> IsDocked;
+
 	/** If true, collection buttons (such as add) are allowed */
 	bool bAllowCollectionButtons;
 
@@ -347,9 +328,6 @@ private:
 	/** True when a drag is over this view with a valid operation for drop */
 	bool bDraggedOver;
 
-	/** The view mode that collections within this view should use */
-	EAssetTagItemViewMode CollectionViewMode = EAssetTagItemViewMode::Standard;
-
 	/** If > 0, the selection changed delegate will not be called. Used to update the tree from an external source or in certain bulk operations. */
 	int32 PreventSelectionChangedDelegateCount;
 
diff --git a/Engine/Source/Editor/ContentBrowser/Private/SContentBrowser.cpp b/Engine/Source/Editor/ContentBrowser/Private/SContentBrowser.cpp
--- a/Engine/Source/Editor/ContentBrowser/Private/SContentBrowser.cpp
+++ b/Engine/Source/Editor/ContentBrowser/Private/SContentBrowser.cpp
@@ -67,8 +67,10 @@
 #include "IContentBrowserDataModule.h"
 #include "ContentBrowserDataSource.h"
 #include "ContentBrowserDataSubsystem.h"
-
 #include "Brushes/SlateColorBrush.h"
+#include "ToolMenu.h"
+#include "Widgets/Input/SExpandableButton.h"
+#include "SExpandableSearchArea.h"
 
 
 #define LOCTEXT_NAMESPACE "ContentBrowser"
@@ -131,306 +133,304 @@ void SContentBrowser::Construct( const FArguments& InArgs, const FName& InInstan
 
 	SourcesSearch = MakeShared<FSourcesSearch>();
 	SourcesSearch->Initialize();
-	SourcesSearch->SetHintText(MakeAttributeSP(this, &SContentBrowser::GetSourcesSearchHintText));
+	SourcesSearch->SetHintText(LOCTEXT("SearchPathsHint", "Search Paths"));
+
+	CollectionSearch = MakeShared<FSourcesSearch>();
+	CollectionSearch->Initialize();
+	CollectionSearch->SetHintText(LOCTEXT("CollectionsViewSearchBoxHint", "Search Collections"));
 
 	CollectionViewPtr = SNew(SCollectionView)
 		.OnCollectionSelected(this, &SContentBrowser::CollectionSelected)
 		.AddMetaData<FTagMetaData>(FTagMetaData(TEXT("ContentBrowserCollections")))
-		.AllowCollapsing(false)
 		.AllowCollectionDrag(true)
 		.AllowQuickAssetManagement(true)
-		.ExternalSearch(SourcesSearch);
+		.IsDocked(this, &SContentBrowser::IsCollectionViewDocked)
+		.ExternalSearch(CollectionSearch);
 
 	static const FName DefaultForegroundName("DefaultForeground");
 
 	BindCommands();
 	UContentBrowserSettings::OnSettingChanged().AddSP(this, &SContentBrowser::OnContentBrowserSettingsChanged);
 
+	// Currently this controls the asset count 
+	const bool bShowBottomToolbar = Config != nullptr ? Config->bShowBottomToolbar : true;
+
+	AssetViewPtr = SNew(SAssetView)
+		.ThumbnailLabel(Config != nullptr ? Config->ThumbnailLabel : EThumbnailLabel::ClassName)
+		//.ThumbnailScale(Config != nullptr ? Config->ThumbnailScale : 0.18f)
+		.InitialViewType(Config != nullptr ? Config->InitialAssetViewType : EAssetViewType::Tile)
+		.OnNewItemRequested(this, &SContentBrowser::OnNewItemRequested)
+		.OnItemSelectionChanged(this, &SContentBrowser::OnItemSelectionChanged, EContentBrowserViewContext::AssetView)
+		.OnItemsActivated(this, &SContentBrowser::OnItemsActivated)
+		.OnGetItemContextMenu(this, &SContentBrowser::GetItemContextMenu, EContentBrowserViewContext::AssetView)
+		.OnItemRenameCommitted(this, &SContentBrowser::OnItemRenameCommitted)
+		.AreRealTimeThumbnailsAllowed(this, &SContentBrowser::IsHovered)
+		.FrontendFilters(FrontendFilters)
+		.HighlightedText(this, &SContentBrowser::GetHighlightedText)
+		.ShowBottomToolbar(bShowBottomToolbar)
+		.ShowViewOptions(false)  // We control this for the main content browser
+		.AllowThumbnailEditMode(true)
+		.AllowThumbnailHintLabel(false)
+		.CanShowFolders(Config != nullptr ? Config->bCanShowFolders : true)
+		.CanShowClasses(Config != nullptr ? Config->bCanShowClasses : true)
+		.CanShowRealTimeThumbnails(Config != nullptr ? Config->bCanShowRealTimeThumbnails : true)
+		.CanShowDevelopersFolder(Config != nullptr ? Config->bCanShowDevelopersFolder : true)
+		.CanShowFavorites(true)
+		.CanDockCollections(true)
+		.AddMetaData<FTagMetaData>(FTagMetaData(TEXT("ContentBrowserAssets")))
+		.OwningContentBrowser(SharedThis(this))
+		.OnSearchOptionsChanged(this, &SContentBrowser::HandleAssetViewSearchOptionsChanged)
+		.FillEmptySpaceInTileView(true);
+
+
+	TSharedRef<SWidget> ViewOptions = SNullWidget::NullWidget;
+
+	// Note, for backwards compatibility ShowBottomToolbar controls the visibility of view options so we respect that here
+	if (bShowBottomToolbar)
+	{
+		ViewOptions =
+			SNew(SComboButton)
+			.ContentPadding(0)
+			.ComboButtonStyle(&FAppStyle::Get().GetWidgetStyle<FComboButtonStyle>("SimpleComboButton"))
+			.OnGetMenuContent(AssetViewPtr.ToSharedRef(), &SAssetView::GetViewButtonContent)
+			.HasDownArrow(false)
+			.ButtonContent()
+			[
+				SNew(SImage)
+				.ColorAndOpacity(FSlateColor::UseForeground())
+				.Image(FAppStyle::Get().GetBrush("Icons.Settings"))
+			];
+	}
+
 	ChildSlot
 	[
 		SNew(SVerticalBox)
-
-		// Path and history
 		+SVerticalBox::Slot()
 		.AutoHeight()
 		.Padding( 0, 0, 0, 0 )
 		[
-			SNew( SWrapBox )
-			.UseAllottedSize( true )
-			.InnerSlotPadding( FVector2D( 5, 2 ) )
-
-			+ SWrapBox::Slot()
-			.FillLineWhenSizeLessThan( 600 )
-			.FillEmptySpace( true )
+			SNew( SBorder )
+			.Padding( FMargin( 3 ) )
+			.BorderImage(FAppStyle::Get().GetBrush("Brushes.Background"))
 			[
-				SNew( SHorizontalBox )
-
-				+ SHorizontalBox::Slot()
-				.FillWidth(1.0f)
+				SNew( SWrapBox )
+				.UseAllottedSize( true )
+				.InnerSlotPadding( FVector2D( 5, 2 ) )
+				+ SWrapBox::Slot()
+				.FillLineWhenSizeLessThan( 600 )
+				.FillEmptySpace( true )
+				.HAlign(HAlign_Left)
 				[
-					SNew( SBorder )
-					.Padding( FMargin( 3 ) )
-					.BorderImage( FEditorStyle::GetBrush( "ContentBrowser.TopBar.GroupBorder" ) )
+					SNew( SHorizontalBox )
+					// Add
+					+ SHorizontalBox::Slot()
+					.Padding(5.0f, 0.0f)
+					.AutoWidth()
+					.VAlign( VAlign_Center )
+					.HAlign( HAlign_Left )
 					[
-						SNew( SHorizontalBox )
-
-						// Add/Import
-						+ SHorizontalBox::Slot()
-						.AutoWidth()
-						.VAlign( VAlign_Center )
-						.HAlign( HAlign_Left )
+						SNew( SComboButton )
+						.ComboButtonStyle(&FAppStyle::Get().GetWidgetStyle<FComboButtonStyle>("SimpleComboButton"))
+						.ForegroundColor(FSlateColor::UseStyle())
+						.ContentPadding(2)
+						.OnGetMenuContent_Lambda( [this]{ return MakeAddNewContextMenu(EContentBrowserDataMenuContext_AddNewMenuDomain::Toolbar); } )
+						.ToolTipText( this, &SContentBrowser::GetAddNewToolTipText )
+						.IsEnabled( this, &SContentBrowser::IsAddNewEnabled )
+						.AddMetaData<FTagMetaData>(FTagMetaData(TEXT("ContentBrowserNewAsset")))
+						.HasDownArrow(false)
+						.ButtonContent()
 						[
-							SNew( SComboButton )
-							.ComboButtonStyle( FEditorStyle::Get(), "ToolbarComboButton" )
-							.ButtonStyle(FEditorStyle::Get(), "FlatButton.Success")
-							.ForegroundColor(FLinearColor::White)
-							.ContentPadding(FMargin(6, 2))
-							.OnGetMenuContent_Lambda( [this]{ return MakeAddNewContextMenu(EContentBrowserDataMenuContext_AddNewMenuDomain::Toolbar); } )
-							.ToolTipText( this, &SContentBrowser::GetAddNewToolTipText )
-							.IsEnabled( this, &SContentBrowser::IsAddNewEnabled )
-							.AddMetaData<FTagMetaData>(FTagMetaData(TEXT("ContentBrowserNewAsset")))
-							.HasDownArrow(false)
-							.ButtonContent()
+							SNew(SHorizontalBox)
+							+ SHorizontalBox::Slot()
+							.AutoWidth()
+							.HAlign(HAlign_Center)
+							.VAlign(VAlign_Center)
 							[
-								SNew( SHorizontalBox )
-
-								// New Icon
-								+ SHorizontalBox::Slot()
-								.VAlign(VAlign_Center)
-								.AutoWidth()
-								[
-									SNew(STextBlock)
-									.TextStyle(FEditorStyle::Get(), "ContentBrowser.TopBar.Font")
-									.Font(FEditorStyle::Get().GetFontStyle("FontAwesome.11"))
-									.Text(FEditorFontGlyphs::File)
-								]
-
-								// New Text
-								+ SHorizontalBox::Slot()
-								.AutoWidth()
-								.VAlign(VAlign_Center)
-								.Padding(4, 0, 0, 0)
-								[
-									SNew( STextBlock )
-									.TextStyle( FEditorStyle::Get(), "ContentBrowser.TopBar.Font" )
-									.Text( LOCTEXT( "AddImportButton", "Add/Import" ) )
-								]
-
-								// Down Arrow
-								+ SHorizontalBox::Slot()
-								.VAlign(VAlign_Center)
-								.AutoWidth()
-								.Padding(4, 0, 0, 0)
-								[
-									SNew(STextBlock)
-									.TextStyle(FEditorStyle::Get(), "ContentBrowser.TopBar.Font")
-									.Font(FEditorStyle::Get().GetFontStyle("FontAwesome.10"))
-									.Text(FEditorFontGlyphs::Caret_Down)
-								]
+								SNew(SImage)
+								.Image(FAppStyle::Get().GetBrush("Icons.Plus"))
+								.ColorAndOpacity(FSlateColor::UseForeground())
 							]
-						]
-
-						// Save
-						+ SHorizontalBox::Slot()
-						.FillWidth(1.0f)
-						.VAlign(VAlign_Center)
-						.HAlign(HAlign_Left)
-						[
-							SNew( SButton )
-							.ButtonStyle(FEditorStyle::Get(), "FlatButton")
-							.ToolTipText( LOCTEXT( "SaveDirtyPackagesTooltip", "Save all modified assets." ) )
-							.ContentPadding(FMargin(6, 2))
-							.OnClicked( this, &SContentBrowser::OnSaveClicked )
-							.AddMetaData<FTagMetaData>(FTagMetaData(TEXT("ContentBrowserSaveDirtyPackages")))
+							+ SHorizontalBox::Slot()
+							.Padding(FMargin(3, 0, 0, 0))
+							.VAlign(VAlign_Center)
+							.AutoWidth()
 							[
-								SNew( SHorizontalBox )
-
-								// Save All Icon
-								+ SHorizontalBox::Slot()
-								.VAlign(VAlign_Center)
-								.AutoWidth()
-								[
-									SNew(STextBlock)
-									.TextStyle(FEditorStyle::Get(), "ContentBrowser.TopBar.Font")
-									.Font(FEditorStyle::Get().GetFontStyle("FontAwesome.11"))
-									.Text(FEditorFontGlyphs::Floppy_O)
-								]
-
-								// Save All Text
-								+ SHorizontalBox::Slot()
-								.AutoWidth()
-								.VAlign(VAlign_Center)
-								.Padding(4, 0, 0, 0)
-								[
-									SNew( STextBlock )
-									.TextStyle( FEditorStyle::Get(), "ContentBrowser.TopBar.Font" )
-									.Text( LOCTEXT( "SaveAll", "Save All" ) )
-								]
+								SNew(STextBlock)
+								.TextStyle(FAppStyle::Get(), "NormalText")
+								.Text(LOCTEXT("NewAssetButton", "New"))
 							]
 						]
 					]
-				]
-			]
-
-			+ SWrapBox::Slot()
-			.FillEmptySpace( true )
-			[
-				SNew(SBorder)
-				.Padding(FMargin(3))
-				.BorderImage( FEditorStyle::GetBrush("ContentBrowser.TopBar.GroupBorder") )
-				[
-					SNew(SHorizontalBox)
-
-					// History Back Button
-					+SHorizontalBox::Slot()
+					+ SHorizontalBox::Slot()
+					.Padding(5.0f, 0.0f)
 					.AutoWidth()
+					.VAlign(VAlign_Center)
+					.HAlign(HAlign_Left)
 					[
-						SNew(SVerticalBox)
-
-						+ SVerticalBox::Slot()
-						.FillHeight(1.0f)
+						SNew(SButton)
+						.ButtonStyle(FAppStyle::Get(), "SimpleButton")
+						.ToolTipText( LOCTEXT( "ImportTooltip", "Import assets from files to the currently selected folder" ) )
+						.ContentPadding(2)
+						.OnClicked(this, &SContentBrowser::OnImportClicked)
 						[
-							SNew(SButton)
-							.VAlign(EVerticalAlignment::VAlign_Center)
-							.ButtonStyle(FEditorStyle::Get(), "FlatButton")
-							.ForegroundColor( FEditorStyle::GetSlateColor(DefaultForegroundName) )
-							.ToolTipText( this, &SContentBrowser::GetHistoryBackTooltip )
-							.ContentPadding( FMargin(1, 0) )
-							.OnClicked(this, &SContentBrowser::BackClicked)
-							.IsEnabled(this, &SContentBrowser::IsBackEnabled)
-							.AddMetaData<FTagMetaData>(FTagMetaData(TEXT("ContentBrowserHistoryBack")))
+							SNew( SHorizontalBox )
+							+ SHorizontalBox::Slot()
+							.AutoWidth()
+							.HAlign(HAlign_Center)
+							.VAlign(VAlign_Center)
+							[
+								SNew(SImage)
+								.Image(FAppStyle::Get().GetBrush("Icons.Import"))
+								.ColorAndOpacity(FSlateColor::UseForeground())
+							]
+							+ SHorizontalBox::Slot()
+							.Padding(FMargin(3, 0, 0, 0))
+							.VAlign(VAlign_Center)
+							.AutoWidth()
 							[
 								SNew(STextBlock)
-								.TextStyle(FEditorStyle::Get(), "ContentBrowser.TopBar.Font")
-								.Font(FEditorStyle::Get().GetFontStyle("FontAwesome.11"))
-								.Text(FText::FromString(FString(TEXT("\xf060"))) /*fa-arrow-left*/)
+								.TextStyle(FAppStyle::Get(), "NormalText")
+								.Text( LOCTEXT( "Import", "Import" ) )
 							]
 						]
 					]
-
-					// History Forward Button
+					// Save
 					+ SHorizontalBox::Slot()
+					.Padding(5.0f, 0.0f)
 					.AutoWidth()
+					.VAlign(VAlign_Center)
+					.HAlign(HAlign_Left)
 					[
-						SNew(SVerticalBox)
-
-						+ SVerticalBox::Slot()
-						.FillHeight(1.0f)
+						SNew( SButton )
+						.ButtonStyle(FAppStyle::Get(), "SimpleButton")
+						.ToolTipText( LOCTEXT( "SaveDirtyPackagesTooltip", "Save all modified assets." ) )
+						.ContentPadding(2)
+						.OnClicked(this, &SContentBrowser::OnSaveClicked)
+						.AddMetaData<FTagMetaData>(FTagMetaData(TEXT("ContentBrowserSaveDirtyPackages")))
 						[
-							SNew(SButton)
-							.VAlign(EVerticalAlignment::VAlign_Center)
-							.ButtonStyle(FEditorStyle::Get(), "FlatButton")
-							.ForegroundColor( FEditorStyle::GetSlateColor(DefaultForegroundName) )
-							.ToolTipText( this, &SContentBrowser::GetHistoryForwardTooltip )
-							.ContentPadding( FMargin(1, 0) )
-							.OnClicked(this, &SContentBrowser::ForwardClicked)
-							.IsEnabled(this, &SContentBrowser::IsForwardEnabled)
-							.AddMetaData<FTagMetaData>(FTagMetaData(TEXT("ContentBrowserHistoryForward")))
+							SNew( SHorizontalBox )
+
+							// Save All Icon
+							+ SHorizontalBox::Slot()
+							.AutoWidth()
+							.HAlign(HAlign_Center)
+							.VAlign(VAlign_Center)
+							[
+								SNew(SImage)
+								.Image(FAppStyle::Get().GetBrush("Icons.Save"))
+								.ColorAndOpacity(FSlateColor::UseForeground())
+							]
+
+							// Save All Text
+							+ SHorizontalBox::Slot()
+							.Padding(FMargin(3, 0, 0, 0))
+							.VAlign(VAlign_Center)
+							.AutoWidth()
 							[
 								SNew(STextBlock)
-								.TextStyle(FEditorStyle::Get(), "ContentBrowser.TopBar.Font")
-								.Font(FEditorStyle::Get().GetFontStyle("FontAwesome.11"))
-								.Text(FText::FromString(FString(TEXT("\xf061"))) /*fa-arrow-right*/)
+								.TextStyle(FAppStyle::Get(), "NormalText")
+								.Text( LOCTEXT( "SaveAll", "Save All" ) )
 							]
 						]
 					]
-
-					// Separator
+				]
+				+ SWrapBox::Slot()
+				.FillEmptySpace(true)
+				[
+					SNew(SHorizontalBox)
+					// Search
+					+ SHorizontalBox::Slot()
+					.VAlign(VAlign_Center)
+					.Padding(5, 0, 0, 0)
+					.FillWidth(1.0f)
+					[
+						SAssignNew(SearchBoxPtr, SAssetSearchBox)
+						.HintText(this, &SContentBrowser::GetSearchAssetsHintText)
+						.OnTextChanged(this, &SContentBrowser::OnSearchBoxChanged)
+						.OnTextCommitted(this, &SContentBrowser::OnSearchBoxCommitted)
+						.OnKeyDownHandler(this, &SContentBrowser::OnSearchKeyDown)
+						.OnAssetSearchBoxSuggestionFilter(this, &SContentBrowser::OnAssetSearchSuggestionFilter)
+						.OnAssetSearchBoxSuggestionChosen(this, &SContentBrowser::OnAssetSearchSuggestionChosen)
+						.DelayChangeNotificationsWhileTyping(true)
+						.Visibility((Config != nullptr ? Config->bCanShowAssetSearch : true) ? EVisibility::Visible : EVisibility::Collapsed)
+						.AddMetaData<FTagMetaData>(FTagMetaData(TEXT("ContentBrowserSearchAssets")))
+					]
+					
+					// Save Search
 					+ SHorizontalBox::Slot()
 					.AutoWidth()
-					.Padding(3, 0)
+					.VAlign(VAlign_Center)
+					.Padding(2.0f, 0.0f, 0.0f, 0.0f)
 					[
-						SNew(SSeparator)
-						.Orientation(Orient_Vertical)
+						SNew(SButton)
+						.ButtonStyle(FAppStyle::Get(), "SimpleButton")
+						.ToolTipText(LOCTEXT("SaveSearchButtonTooltip", "Save the current search as a dynamic collection."))
+						.IsEnabled(this, &SContentBrowser::IsSaveSearchButtonEnabled)
+						.OnClicked(this, &SContentBrowser::OnSaveSearchButtonClicked)
+						.ContentPadding(FMargin(1, 1))
+						.Visibility((Config != nullptr ? Config->bCanShowAssetSearch : true) ? EVisibility::Visible : EVisibility::Collapsed)
+						[
+							SNew(STextBlock)
+							.TextStyle(FEditorStyle::Get(), "GenericFilters.TextStyle")
+							.Font(FEditorStyle::Get().GetFontStyle("FontAwesome.10"))
+							.Text(FEditorFontGlyphs::Floppy_O)
+						]
 					]
-
-					// Path picker
 					+ SHorizontalBox::Slot()
 					.AutoWidth()
-					.VAlign( VAlign_Fill )
+					.HAlign(HAlign_Left)
+					.Padding(5, 0, 0, 0)
 					[
-						SAssignNew( PathPickerButton, SComboButton )
-						.Visibility( ( Config != nullptr ? Config->bUsePathPicker : true ) ? EVisibility::Visible : EVisibility::Collapsed )
-						.ButtonStyle(FEditorStyle::Get(), "FlatButton")
-						.ForegroundColor(FLinearColor::White)
-						.ToolTipText( LOCTEXT( "PathPickerTooltip", "Choose a path" ) )
-						.OnGetMenuContent( this, &SContentBrowser::GetPathPickerContent )
-						.HasDownArrow( false )
-						.AddMetaData<FTagMetaData>(FTagMetaData(TEXT("ContentBrowserPathPicker")))
-						.ContentPadding(FMargin(3, 3))
+						SNew(SComboButton)
+						.ComboButtonStyle(&FAppStyle::Get().GetWidgetStyle<FComboButtonStyle>("SimpleComboButton"))
+						.ForegroundColor(FSlateColor::UseStyle())
+						.ToolTipText(LOCTEXT("AddFilterToolTip", "Add an asset filter."))
+						.OnGetMenuContent(this, &SContentBrowser::MakeAddFilterMenu)
+						.ContentPadding(FMargin(1, 0))
+						.AddMetaData<FTagMetaData>(FTagMetaData(TEXT("ContentBrowserFiltersCombo")))
+						.Visibility((Config != nullptr ? Config->bCanShowFilters : true) ? EVisibility::Visible : EVisibility::Collapsed)
 						.ButtonContent()
 						[
-							SNew(STextBlock)
-							.TextStyle(FEditorStyle::Get(), "ContentBrowser.TopBar.Font")
-							.Font(FEditorStyle::Get().GetFontStyle("FontAwesome.11"))
-							.Text(FText::FromString(FString(TEXT("\xf07c"))) /*fa-folder-open*/)
+							SNew(SImage)
+							.Image(FAppStyle::Get().GetBrush("Icons.Filter"))
+							.ColorAndOpacity(FSlateColor::UseForeground())
 						]
 					]
-
-					// Path
 					+ SHorizontalBox::Slot()
+					.Padding(5.0f,0.0f,0.0f,0.0f)
 					.VAlign(VAlign_Center)
-					.HAlign(HAlign_Left)
-					.FillWidth(1.0f)
-					.Padding(FMargin(0))
+					.FillWidth(1.5f)
 					[
-						SAssignNew(PathBreadcrumbTrail, SBreadcrumbTrail<FString>)
-						.ButtonContentPadding(FMargin(2, 2))
-						.ButtonStyle(FEditorStyle::Get(), "FlatButton")
-						.DelimiterImage(FEditorStyle::GetBrush("ContentBrowser.PathDelimiter"))
-						.TextStyle(FEditorStyle::Get(), "ContentBrowser.PathText")
-						.ShowLeadingDelimiter(false)
-						.InvertTextColorOnHover(false)
-						.OnCrumbClicked(this, &SContentBrowser::OnPathClicked)
-						.HasCrumbMenuContent(this, &SContentBrowser::OnHasCrumbDelimiterContent)
-						.GetCrumbMenuContent(this, &SContentBrowser::OnGetCrumbDelimiterContent)
-						.AddMetaData<FTagMetaData>(FTagMetaData(TEXT("ContentBrowserPath")))
+						SAssignNew(FilterListPtr, SFilterList)
+						.OnFilterChanged(this, &SContentBrowser::OnFilterChanged)
+						.OnGetContextMenu(this, &SContentBrowser::GetFilterContextMenu)
+						.Visibility((Config != nullptr ? Config->bCanShowFilters : true) ? EVisibility::Visible : EVisibility::Collapsed)
+						.FrontendFilters(FrontendFilters)
+						.AddMetaData<FTagMetaData>(FTagMetaData(TEXT("ContentBrowserFilters")))
 					]
-
-					// Lock button
 					+ SHorizontalBox::Slot()
 					.AutoWidth()
+					.HAlign(HAlign_Right)
 					.VAlign(VAlign_Center)
 					[
-						SNew(SVerticalBox)
-						.Visibility( ( Config != nullptr ? Config->bCanShowLockButton : true ) ? EVisibility::SelfHitTestInvisible : EVisibility::Collapsed )
-
-						+ SVerticalBox::Slot()
-						.FillHeight(1.0f)
-						[
-							SNew(SButton)
-							.VAlign(EVerticalAlignment::VAlign_Center)
-							.ButtonStyle(FEditorStyle::Get(), "FlatButton")
-							.ToolTipText( LOCTEXT("LockToggleTooltip", "Toggle lock. If locked, this browser will ignore Find in Content Browser requests.") )
-							.ContentPadding( FMargin(1, 0) )
-							.OnClicked(this, &SContentBrowser::ToggleLockClicked)
-							.AddMetaData<FTagMetaData>(FTagMetaData(TEXT("ContentBrowserLock")))
-							[
-								SNew(SImage)
-								.Image( this, &SContentBrowser::GetToggleLockImage)
-							]
-						]
+						CreateLockButton(Config)
+					]
+					+ SHorizontalBox::Slot()
+					.AutoWidth()
+					.Padding(5.0f, 0.0f, 0.0f, 0.0f)
+					.HAlign(HAlign_Right)
+					[
+						ViewOptions
 					]
 				]
 			]
 		]
 
-		+ SVerticalBox::Slot()
-		.AutoHeight()
-		.Padding(0,0,0,0)
-		[
-			SNew(SBox)
-			.HeightOverride(2.0f)
-			[
-				SNew(SImage)
-				.Image(new FSlateColorBrush(FLinearColor( FColor( 34, 34, 34) ) ) )
-			]
-		]
 
 		// Assets/tree
 		+ SVerticalBox::Slot()
 		.FillHeight(1.0f)
-		.Padding(0,2,0,0)
+		.Padding(0)
 		[
 			// The tree/assets splitter
 			SAssignNew(PathAssetSplitterPtr, SSplitter)
@@ -443,61 +443,12 @@ void SContentBrowser::Construct( const FArguments& InArgs, const FName& InInstan
 			[
 				SNew(SBorder)
 				.Padding(FMargin(3))
-				.BorderImage(FEditorStyle::GetBrush("ContentBrowser.TopBar.GroupBorder"))
+				.BorderImage(FEditorStyle::GetBrush("Brushes.Background"))
 				.Visibility(this, &SContentBrowser::GetSourcesViewVisibility)
 				[
-					SNew(SVerticalBox)
-
-					// Paths expansion/search
-					+SVerticalBox::Slot()
-					.AutoHeight()
-					[
-						SNew(SHorizontalBox)
-						.AddMetaData<FTagMetaData>(FTagMetaData(TEXT("ContentBrowserSourcesToggle1")))
-					
-						+SHorizontalBox::Slot()
-						.AutoWidth()
-						.Padding(0, 0, 2, 0)
-						[
-							SNew(SButton)
-							.VAlign(EVerticalAlignment::VAlign_Center)
-							.ButtonStyle(FEditorStyle::Get(), "ToggleButton")
-							.ToolTipText(LOCTEXT("SourcesTreeToggleTooltip", "Show or hide the sources panel"))
-							.ContentPadding(FMargin(1, 0))
-							.ForegroundColor(FEditorStyle::GetSlateColor(DefaultForegroundName))
-							.OnClicked(this, &SContentBrowser::SourcesViewExpandClicked)
-							[
-								SNew(SImage)
-								.Image(this, &SContentBrowser::GetSourcesToggleImage)
-							]
-						]
-
-						+SHorizontalBox::Slot()
-						[
-							SourcesSearch->GetWidget()
-						]
-
-						+SHorizontalBox::Slot()
-						.AutoWidth()
-						.Padding(2, 0, 0, 0)
-						[
-							SNew(SButton)
-							.Visibility(this, &SContentBrowser::GetSourcesSwitcherVisibility)
-							.VAlign(EVerticalAlignment::VAlign_Center)
-							.ButtonStyle(FEditorStyle::Get(), "ToggleButton")
-							.ToolTipText(this, &SContentBrowser::GetSourcesSwitcherToolTipText)
-							.ContentPadding(FMargin(1, 0))
-							.ForegroundColor(FEditorStyle::GetSlateColor(DefaultForegroundName))
-							.OnClicked(this, &SContentBrowser::OnSourcesSwitcherClicked)
-							[
-								SNew(SImage)
-								.Image(this, &SContentBrowser::GetSourcesSwitcherIcon)
-							]
-						]
-					]
-
-					+SVerticalBox::Slot()
-					.FillHeight(1.0f)
+					SNew(SBorder)
+					.Padding(FMargin(3))
+					.BorderImage(FEditorStyle::GetBrush("Brushes.Input"))
 					[
 						// Note: If adding more widgets here, fix ContentBrowserSourcesWidgetSwitcherIndex and the code that uses it!
 						SAssignNew(SourcesWidgetSwitcher, SWidgetSwitcher)
@@ -507,88 +458,30 @@ void SContentBrowser::Construct( const FArguments& InArgs, const FName& InInstan
 						[
 							SAssignNew(PathFavoriteSplitterPtr, SSplitter)
 							.Style(FEditorStyle::Get(), "ContentBrowser.Splitter")
-							.PhysicalSplitterHandleSize(2.0f)
-							.HitDetectionSplitterHandleSize(2.0f)
+							.PhysicalSplitterHandleSize(5.0f)
+							.HitDetectionSplitterHandleSize(5.0f)
 							.Orientation(EOrientation::Orient_Vertical)
-							.MinimumSlotHeight(70.0f)
+							.MinimumSlotHeight(26.0f)
 							.Visibility( this, &SContentBrowser::GetSourcesViewVisibility )
-							
 							+SSplitter::Slot()
+							.SizeRule(TAttribute<SSplitter::ESizeRule>(this, &SContentBrowser::GetFavoritesAreaSizeRule))
 							.Value(0.2f)
 							[
-								SNew(SBox)
-								.Visibility(this, &SContentBrowser::GetFavoriteFolderVisibility)
-								[
-									SNew(SExpandableArea)
-									.BorderImage(FEditorStyle::GetBrush("NoBorder"))
-									.HeaderPadding(FMargin(0.0f, 3.0f, 0.0f, 0.0f))
-									.HeaderContent()
-									[
-										SNew(SHorizontalBox)
-
-										+SHorizontalBox::Slot()
-										.AutoWidth()
-										.Padding(0, 0, 2, 0)
-										.VAlign(VAlign_Center)
-										[
-											SNew(SImage) 
-											.Image(FEditorStyle::GetBrush("PropertyWindow.Favorites_Enabled"))
-										]
-
-										+SHorizontalBox::Slot()
-										.AutoWidth()
-										.VAlign(VAlign_Center)
-										[
-											SNew(STextBlock)
-											.Text(LOCTEXT("Favorites", "Favorites"))
-											.Font(FEditorStyle::GetFontStyle("ContentBrowser.SourceTreeRootItemFont"))
-										]
-									]
-									.BodyContent()
-									[
-										SNew(SBox)
-										.Padding(FMargin(9, 0, 0, 0))
-										[
-											SAssignNew(FavoritePathViewPtr, SFavoritePathView)
-											.OnItemSelectionChanged(this, &SContentBrowser::OnItemSelectionChanged, EContentBrowserViewContext::FavoriteView)
-											.OnGetItemContextMenu(this, &SContentBrowser::GetItemContextMenu, EContentBrowserViewContext::FavoriteView)
-											.FocusSearchBoxWhenOpened(false)
-											.ShowTreeTitle(false)
-											.ShowSeparator(false)
-											.AllowClassesFolder(true)
-											.AddMetaData<FTagMetaData>(FTagMetaData(TEXT("ContentBrowserFavorites")))
-											.ExternalSearch(SourcesSearch)
-										]
-									]
-								]
+								CreateFavoritesView(Config)
 							]
 							
 							+SSplitter::Slot()
+							.SizeRule(TAttribute<SSplitter::ESizeRule>(this, &SContentBrowser::GetPathAreaSizeRule))
 							.Value(0.8f)
 							[
-								SNew(SBox)
-								.Padding(FMargin(0.0f, 1.0f, 0.0f, 0.0f))
-								[
-									SAssignNew( PathViewPtr, SPathView )
-									.OnItemSelectionChanged( this, &SContentBrowser::OnItemSelectionChanged, EContentBrowserViewContext::PathView )
-									.OnGetItemContextMenu( this, &SContentBrowser::GetItemContextMenu, EContentBrowserViewContext::PathView )
-									.FocusSearchBoxWhenOpened( false )
-									.ShowTreeTitle( false )
-									.ShowSeparator( false )
-									.AllowClassesFolder( true )
-									.AddMetaData<FTagMetaData>(FTagMetaData(TEXT("ContentBrowserSources")))
-									.ExternalSearch(SourcesSearch)
-								]
+								CreatePathView(Config)
 							]
 
 							+SSplitter::Slot()
+							.SizeRule(TAttribute<SSplitter::ESizeRule>(this, &SContentBrowser::GetCollectionsAreaSizeRule))
 							.Value(0.4f)
 							[
-								SNew(SBox)
-								.Visibility(this, &SContentBrowser::GetDockedCollectionsVisibility)
-								[
-									CollectionViewPtr.ToSharedRef()
-								]
+								CreateDockedCollectionsView(Config)
 							]
 						]
 
@@ -609,168 +502,13 @@ void SContentBrowser::Construct( const FArguments& InArgs, const FName& InInstan
 			+ SSplitter::Slot()
 			.Value(0.75f)
 			[
-				SNew(SBorder)
-				.Padding(FMargin(3))
-				.BorderImage( FEditorStyle::GetBrush("ContentBrowser.TopBar.GroupBorder") )
-				[
-					SNew(SVerticalBox)
-
-					// Search and commands
-					+ SVerticalBox::Slot()
-					.AutoHeight()
-					.Padding(0, 0, 0, 2)
-					[
-						SNew(SHorizontalBox)
-
-						// Expand/collapse sources button
-						+ SHorizontalBox::Slot()
-						.AutoWidth()
-						.Padding( 0, 0, 4, 0 )
-						[
-							SNew( SVerticalBox )
-							.Visibility(( Config != nullptr ? Config->bUseSourcesView : true ) ? EVisibility::SelfHitTestInvisible : EVisibility::Collapsed)
-							.AddMetaData<FTagMetaData>(FTagMetaData(TEXT("ContentBrowserSourcesToggle2")))
-							+ SVerticalBox::Slot()
-							.FillHeight( 1.0f )
-							[
-								SNew( SButton )
-								.VAlign( EVerticalAlignment::VAlign_Center )
-								.ButtonStyle( FEditorStyle::Get(), "ToggleButton" )
-								.ToolTipText( LOCTEXT( "SourcesTreeToggleTooltip", "Show or hide the sources panel" ) )
-								.ContentPadding( FMargin( 1, 0 ) )
-								.ForegroundColor( FEditorStyle::GetSlateColor(DefaultForegroundName) )
-								.OnClicked( this, &SContentBrowser::SourcesViewExpandClicked )
-								.Visibility( this, &SContentBrowser::GetPathExpanderVisibility )
-								[
-									SNew( SImage )
-									.Image( this, &SContentBrowser::GetSourcesToggleImage )
-								]
-							]
-						]
-
-						// Filter
-						+ SHorizontalBox::Slot()
-						.AutoWidth()
-						[
-							SNew( SComboButton )
-							.ComboButtonStyle( FEditorStyle::Get(), "GenericFilters.ComboButtonStyle" )
-							.ForegroundColor(FLinearColor::White)
-							.ContentPadding(0)
-							.ToolTipText( LOCTEXT( "AddFilterToolTip", "Add an asset filter." ) )
-							.OnGetMenuContent( this, &SContentBrowser::MakeAddFilterMenu )
-							.HasDownArrow( true )
-							.ContentPadding( FMargin( 1, 0 ) )
-							.AddMetaData<FTagMetaData>(FTagMetaData(TEXT("ContentBrowserFiltersCombo")))
-							.Visibility( ( Config != nullptr ? Config->bCanShowFilters : true ) ? EVisibility::Visible : EVisibility::Collapsed )
-							.ButtonContent()
-							[
-								SNew(SHorizontalBox)
-
-								+ SHorizontalBox::Slot()
-								.AutoWidth()
-								[
-									SNew(STextBlock)
-									.TextStyle(FEditorStyle::Get(), "GenericFilters.TextStyle")
-									.Font(FEditorStyle::Get().GetFontStyle("FontAwesome.9"))
-									.Text(FText::FromString(FString(TEXT("\xf0b0"))) /*fa-filter*/)
-								]
-
-								+ SHorizontalBox::Slot()
-								.AutoWidth()
-								.Padding(2,0,0,0)
-								[
-									SNew(STextBlock)
-									.TextStyle(FEditorStyle::Get(), "GenericFilters.TextStyle")
-									.Text(LOCTEXT("Filters", "Filters"))
-								]
-							]
-						]
-
-						// Search
-						+SHorizontalBox::Slot()
-						.Padding(4, 1, 0, 0)
-						.FillWidth(1.0f)
-						[
-							SAssignNew(SearchBoxPtr, SAssetSearchBox)
-							.HintText( this, &SContentBrowser::GetSearchAssetsHintText )
-							.OnTextChanged( this, &SContentBrowser::OnSearchBoxChanged )
-							.OnTextCommitted( this, &SContentBrowser::OnSearchBoxCommitted )
-							.OnKeyDownHandler(this, &SContentBrowser::OnSearchKeyDown)
-							.OnAssetSearchBoxSuggestionFilter( this, &SContentBrowser::OnAssetSearchSuggestionFilter )
-							.OnAssetSearchBoxSuggestionChosen( this, &SContentBrowser::OnAssetSearchSuggestionChosen )
-							.DelayChangeNotificationsWhileTyping( true )
-							.Visibility( ( Config != nullptr ? Config->bCanShowAssetSearch : true ) ? EVisibility::Visible : EVisibility::Collapsed )
-							.AddMetaData<FTagMetaData>(FTagMetaData(TEXT("ContentBrowserSearchAssets")))
-						]
-
-						// Save Search
-						+SHorizontalBox::Slot()
-						.AutoWidth()
-						.VAlign(VAlign_Center)
-						.Padding(2.0f, 0.0f, 0.0f, 0.0f)
-						[
-							SNew(SButton)
-							.ButtonStyle(FEditorStyle::Get(), "FlatButton")
-							.ToolTipText(LOCTEXT("SaveSearchButtonTooltip", "Save the current search as a dynamic collection."))
-							.IsEnabled(this, &SContentBrowser::IsSaveSearchButtonEnabled)
-							.OnClicked(this, &SContentBrowser::OnSaveSearchButtonClicked)
-							.ContentPadding( FMargin(1, 1) )
-							.Visibility( ( Config != nullptr ? Config->bCanShowAssetSearch : true ) ? EVisibility::Visible : EVisibility::Collapsed )
-							[
-								SNew(STextBlock)
-								.TextStyle(FEditorStyle::Get(), "GenericFilters.TextStyle")
-								.Font(FEditorStyle::Get().GetFontStyle("FontAwesome.10"))
-								.Text(FEditorFontGlyphs::Floppy_O)
-							]
-						]
-					]
-
-					// Filters
-					+ SVerticalBox::Slot()
-					.AutoHeight()
-					[
-						SAssignNew(FilterListPtr, SFilterList)
-						.OnFilterChanged(this, &SContentBrowser::OnFilterChanged)
-						.OnGetContextMenu(this, &SContentBrowser::GetFilterContextMenu)
-						.Visibility( ( Config != nullptr ? Config->bCanShowFilters : true ) ? EVisibility::Visible : EVisibility::Collapsed )
-						.FrontendFilters(FrontendFilters)
-						.AddMetaData<FTagMetaData>(FTagMetaData(TEXT("ContentBrowserFilters")))
-					]
-
-					// Assets
-					+ SVerticalBox::Slot()
-					.FillHeight( 1.0f )
-					.Padding( 0 )
-					[
-						SAssignNew(AssetViewPtr, SAssetView)
-						.ThumbnailLabel( Config != nullptr ? Config->ThumbnailLabel : EThumbnailLabel::ClassName )
-						.ThumbnailScale( Config != nullptr ? Config->ThumbnailScale : 0.18f )
-						.InitialViewType( Config != nullptr ? Config->InitialAssetViewType : EAssetViewType::Tile )
-						.ShowBottomToolbar( Config != nullptr ? Config->bShowBottomToolbar : true )
-						.OnNewItemRequested(this, &SContentBrowser::OnNewItemRequested)
-						.OnItemSelectionChanged(this, &SContentBrowser::OnItemSelectionChanged, EContentBrowserViewContext::AssetView)
-						.OnItemsActivated(this, &SContentBrowser::OnItemsActivated)
-						.OnGetItemContextMenu(this, &SContentBrowser::GetItemContextMenu, EContentBrowserViewContext::AssetView)
-						.OnItemRenameCommitted(this, &SContentBrowser::OnItemRenameCommitted)
-						.AreRealTimeThumbnailsAllowed(this, &SContentBrowser::IsHovered)
-						.FrontendFilters(FrontendFilters)
-						.HighlightedText(this, &SContentBrowser::GetHighlightedText)
-						.AllowThumbnailEditMode(true)
-						.AllowThumbnailHintLabel(false)
-						.CanShowFolders(Config != nullptr ? Config->bCanShowFolders : true)
-						.CanShowClasses(Config != nullptr ? Config->bCanShowClasses : true)
-						.CanShowRealTimeThumbnails( Config != nullptr ? Config->bCanShowRealTimeThumbnails : true)
-						.CanShowDevelopersFolder( Config != nullptr ? Config->bCanShowDevelopersFolder : true)
-						.CanShowFavorites(true)
-						.CanDockCollections(true)
-						.AddMetaData<FTagMetaData>(FTagMetaData(TEXT("ContentBrowserAssets")))
-						.OnSearchOptionsChanged(this, &SContentBrowser::HandleAssetViewSearchOptionsChanged)
-					]
-				]
+				CreateAssetView(Config)
 			]
 		]
 	];
 
+	ExtendViewOptionsMenu(Config);
+
 	AssetContextMenu = MakeShared<FAssetContextMenu>(AssetViewPtr);
 	AssetContextMenu->BindCommands(Commands);
 	AssetContextMenu->SetOnShowInPathsViewRequested( FAssetContextMenu::FOnShowInPathsViewRequested::CreateSP(this, &SContentBrowser::OnShowInPathsViewRequested) );
@@ -899,7 +637,17 @@ EVisibility SContentBrowser::GetFavoriteFolderVisibility() const
 
 EVisibility SContentBrowser::GetDockedCollectionsVisibility() const
 {
-	return GetDefault<UContentBrowserSettings>()->GetDockCollections() ? EVisibility::Visible : EVisibility::Collapsed;
+	return IsCollectionViewDocked() ? EVisibility::Visible : EVisibility::Collapsed;
+}
+
+EVisibility SContentBrowser::GetLockButtonVisibility() const
+{
+	return IsLocked() ? EVisibility::Visible : EVisibility::Collapsed;
+}
+
+bool SContentBrowser::IsCollectionViewDocked() const
+{
+	return GetDefault<UContentBrowserSettings>()->GetDockCollections();
 }
 
 void SContentBrowser::ToggleFolderFavorite(const TArray<FString>& FolderPaths)
@@ -938,6 +686,317 @@ void SContentBrowser::HandleAssetViewSearchOptionsChanged()
 	TextFilter->SetIncludeCollectionNames(AssetViewPtr->IsIncludingCollectionNames());
 }
 
+TSharedRef<SWidget> SContentBrowser::CreateLockButton(const FContentBrowserConfig* Config)
+{
+	if(Config == nullptr || Config->bCanShowLockButton)
+	{
+		return
+			SNew(SButton)
+			.ButtonStyle(FAppStyle::Get(), "SimpleButton")
+			.ToolTipText(LOCTEXT("LockToggleTooltip", "Toggle lock. If locked, this browser will ignore Find in Content Browser requests."))
+			.ContentPadding(FMargin(1, 0))
+			.OnClicked(this, &SContentBrowser::ToggleLockClicked)
+			.AddMetaData<FTagMetaData>(FTagMetaData(TEXT("ContentBrowserLock")))
+			.Visibility(this, &SContentBrowser::GetLockButtonVisibility)
+			[
+				SNew(SImage)
+				.Image(this, &SContentBrowser::GetLockIcon)
+				.ColorAndOpacity(FSlateColor::UseStyle())
+			];
+	}
+
+	return SNullWidget::NullWidget;
+}
+
+TSharedRef<SWidget> SContentBrowser::CreateAssetView(const FContentBrowserConfig* Config)
+{
+	return
+
+		SNew(SVerticalBox)
+		+ SVerticalBox::Slot()
+		.AutoHeight()
+		[
+			SNew(SBorder)
+			.BorderImage(FAppStyle::Get().GetBrush("Brushes.Background"))
+			.Padding(FMargin(0.0f, 5.0f))
+			[
+				SNew(SHorizontalBox)
+				// Sources panel toggle
+				+ SHorizontalBox::Slot()
+				.AutoWidth()
+				.HAlign(HAlign_Left)
+				.Padding(5, 0, 0, 0)
+				[
+					SNew(SButton)
+					.VAlign(EVerticalAlignment::VAlign_Center)
+					.ButtonStyle(FAppStyle::Get(), "SimpleButton")
+					.ToolTipText(LOCTEXT("SourcesTreeToggleTooltip", "Show or hide the sources panel"))
+					.ContentPadding(FMargin(1, 0))
+					.OnClicked(this, &SContentBrowser::SourcesViewExpandClicked)
+					[
+						SNew(SImage)
+						.Image(this, &SContentBrowser::GetSourcesToggleImage)
+						.ColorAndOpacity(FSlateColor::UseForeground())
+					]
+				]
+				// History Back Button
+				+SHorizontalBox::Slot()
+				.Padding(2, 0, 0, 0)
+				.AutoWidth()
+				[
+					SNew(SButton)
+					.VAlign(EVerticalAlignment::VAlign_Center)
+					.ButtonStyle(FEditorStyle::Get(), "SimpleButton")
+					.ToolTipText( this, &SContentBrowser::GetHistoryBackTooltip )
+					.ContentPadding( FMargin(1, 0) )
+					.OnClicked(this, &SContentBrowser::BackClicked)
+					.IsEnabled(this, &SContentBrowser::IsBackEnabled)
+					.AddMetaData<FTagMetaData>(FTagMetaData(TEXT("ContentBrowserHistoryBack")))
+					[
+						SNew(SImage)
+						.Image(FAppStyle::Get().GetBrush("Icons.CircleArrowLeft"))
+						.ColorAndOpacity(FSlateColor::UseForeground())
+					]
+				]
+				// History Forward Button
+				+ SHorizontalBox::Slot()
+				.Padding(2, 0, 0, 0)
+				.AutoWidth()
+				[
+					SNew(SButton)
+					.VAlign(EVerticalAlignment::VAlign_Center)
+					.ButtonStyle(FEditorStyle::Get(), "SimpleButton")
+					.ToolTipText( this, &SContentBrowser::GetHistoryForwardTooltip )
+					.ContentPadding( FMargin(1, 0) )
+					.OnClicked(this, &SContentBrowser::ForwardClicked)
+					.IsEnabled(this, &SContentBrowser::IsForwardEnabled)
+					.AddMetaData<FTagMetaData>(FTagMetaData(TEXT("ContentBrowserHistoryForward")))
+					[
+						SNew(SImage)
+						.Image(FAppStyle::Get().GetBrush("Icons.CircleArrowRight"))
+						.ColorAndOpacity(FSlateColor::UseForeground())
+					]
+				]
+				// Path picker
+				+ SHorizontalBox::Slot()
+				.Padding(2, 0, 0, 0)
+				.AutoWidth()
+				.VAlign(VAlign_Fill)
+				[
+					SAssignNew( PathPickerButton, SComboButton )
+					.Visibility( ( Config != nullptr ? Config->bUsePathPicker : true ) ? EVisibility::Visible : EVisibility::Collapsed )
+					.ButtonStyle(FEditorStyle::Get(), "SimpleButton")
+					.ToolTipText( LOCTEXT( "PathPickerTooltip", "Choose a path" ) )
+					.OnGetMenuContent( this, &SContentBrowser::GetPathPickerContent )
+					.HasDownArrow( false )
+					.AddMetaData<FTagMetaData>(FTagMetaData(TEXT("ContentBrowserPathPicker")))
+					.ContentPadding(FMargin(1, 0))
+					.ButtonContent()
+					[
+						SNew(SImage)
+						.Image(FAppStyle::Get().GetBrush("Icons.FolderClosed"))
+						.ColorAndOpacity(FSlateColor::UseForeground())
+					]
+				]
+
+				// Path
+				+ SHorizontalBox::Slot()
+				.VAlign(VAlign_Center)
+				.HAlign(HAlign_Left)
+				.FillWidth(1.0f)
+				.Padding(2, 0, 0, 0)
+				[
+					SAssignNew(PathBreadcrumbTrail, SBreadcrumbTrail<FString>)
+					.ButtonContentPadding(FMargin(2, 2))
+					.ButtonStyle(FAppStyle::Get(), "SimpleButton")
+					.DelimiterImage(FAppStyle::Get().GetBrush("Icons.ChevronRight"))
+					.TextStyle(FAppStyle::Get(), "NormalText")
+					.ShowLeadingDelimiter(false)
+					.InvertTextColorOnHover(false)
+					.OnCrumbClicked(this, &SContentBrowser::OnPathClicked)
+					.HasCrumbMenuContent(this, &SContentBrowser::OnHasCrumbDelimiterContent)
+					.GetCrumbMenuContent(this, &SContentBrowser::OnGetCrumbDelimiterContent)
+					.AddMetaData<FTagMetaData>(FTagMetaData(TEXT("ContentBrowserPath")))
+				]
+			]
+		]
+		// Assets
+		+ SVerticalBox::Slot()
+		.FillHeight( 1.0f )
+		.Padding( 0 )
+		[
+			AssetViewPtr.ToSharedRef()
+		];
+}
+
+TSharedRef<SWidget> SContentBrowser::CreateFavoritesView(const FContentBrowserConfig* Config)
+{
+	return
+		SAssignNew(FavoritesArea, SExpandableArea)
+		.BorderImage(FAppStyle::Get().GetBrush("Brushes.Background"))
+		.BodyBorderImage(FAppStyle::Get().GetBrush("Brushes.Input"))
+		.HeaderPadding(FMargin(5.0f, 5.0f))
+		.Visibility(this, &SContentBrowser::GetFavoriteFolderVisibility)
+		.AllowAnimatedTransition(false)
+		.HeaderContent()
+		[
+			SNew(STextBlock)
+			.Text(LOCTEXT("Favorites", "Favorites"))
+			.TextStyle(FAppStyle::Get(), "ButtonText")
+			.Font(FAppStyle::Get().GetFontStyle("NormalFontBold"))
+			.TransformPolicy(ETextTransformPolicy::ToUpper)
+		]
+		.BodyContent()
+		[
+			SAssignNew(FavoritePathViewPtr, SFavoritePathView)
+			.OnItemSelectionChanged(this, &SContentBrowser::OnItemSelectionChanged, EContentBrowserViewContext::FavoriteView)
+			.OnGetItemContextMenu(this, &SContentBrowser::GetItemContextMenu, EContentBrowserViewContext::FavoriteView)
+			.FocusSearchBoxWhenOpened(false)
+			.ShowTreeTitle(false)
+			.ShowSeparator(false)
+			.AllowClassesFolder(true)
+			.AddMetaData<FTagMetaData>(FTagMetaData(TEXT("ContentBrowserFavorites")))
+			.ExternalSearch(SourcesSearch)
+		];
+}
+
+TSharedRef<SWidget> SContentBrowser::CreatePathView(const FContentBrowserConfig* Config)
+{	
+	return
+		SAssignNew(PathArea, SExpandableArea)
+		.BorderImage(FAppStyle::Get().GetBrush("Brushes.Background"))
+		.BodyBorderImage(FAppStyle::Get().GetBrush("Brushes.Input"))
+		.HeaderPadding(FMargin(5.0f, 1.0f))
+		.AllowAnimatedTransition(false)
+		.HeaderContent()
+		[
+			SNew(SHorizontalBox)
+			+ SHorizontalBox::Slot()
+			.VAlign(VAlign_Center)
+			.AutoWidth()
+			[
+				SNew(STextBlock)
+				.Text(FText::FromString(FApp::GetProjectName()))
+				.TextStyle(FAppStyle::Get(), "ButtonText")
+				.Font(FAppStyle::Get().GetFontStyle("NormalFontBold"))
+				.TransformPolicy(ETextTransformPolicy::ToUpper)
+			]
+			+SHorizontalBox::Slot()
+			.VAlign(VAlign_Center)
+			.HAlign(HAlign_Right)
+			.Padding(5.0f, 0.0f)
+			[
+				SNew(SExpandableSearchArea, SourcesSearch->GetWidget())
+			]
+		]
+		.BodyContent()
+		[
+			SAssignNew(PathViewPtr, SPathView)
+			.OnItemSelectionChanged(this, &SContentBrowser::OnItemSelectionChanged, EContentBrowserViewContext::PathView)
+			.OnGetItemContextMenu(this, &SContentBrowser::GetItemContextMenu, EContentBrowserViewContext::PathView)
+			.FocusSearchBoxWhenOpened(false)
+			.ShowTreeTitle(false)
+			.ShowSeparator(false)
+			.AllowClassesFolder(true)
+			.AddMetaData<FTagMetaData>(FTagMetaData(TEXT("ContentBrowserSources")))
+			.ExternalSearch(SourcesSearch)
+		];
+}
+
+TSharedRef<SWidget> SContentBrowser::CreateDockedCollectionsView(const FContentBrowserConfig* Config)
+{
+	return
+		SAssignNew(CollectionArea, SExpandableArea)
+		.BorderImage(FAppStyle::Get().GetBrush("Brushes.Background"))
+		.BodyBorderImage(FAppStyle::Get().GetBrush("Brushes.Input"))
+		.HeaderPadding(FMargin(5.0f, 1.0f))
+		.Visibility(this, &SContentBrowser::GetDockedCollectionsVisibility)
+		.AllowAnimatedTransition(false)
+		.HeaderContent()
+		[
+			SNew(SHorizontalBox)
+			+ SHorizontalBox::Slot()
+			.VAlign(VAlign_Center)
+			.AutoWidth()
+			[
+				SNew(STextBlock)
+				.Text(LOCTEXT("CollectionsTitle", "Collections"))
+				.TextStyle(FAppStyle::Get(), "ButtonText")
+				.Font(FAppStyle::Get().GetFontStyle("NormalFontBold"))
+				.TransformPolicy(ETextTransformPolicy::ToUpper)
+			]
+			+ SHorizontalBox::Slot()
+			.VAlign(VAlign_Center)
+			.HAlign(HAlign_Right)
+			.Padding(5.0f, 0.0f)
+			[
+				SNew(SButton)
+				.ButtonStyle(FEditorStyle::Get(), "SimpleButton")
+				.ToolTipText(LOCTEXT("AddCollectionButtonTooltip", "Add a collection."))
+				.OnClicked(CollectionViewPtr.ToSharedRef(), &SCollectionView::MakeAddCollectionMenu)
+				.ContentPadding(FMargin(1, 0))
+				[
+					SNew(SImage)
+					.Image(FAppStyle::Get().GetBrush("Icons.PlusCircle"))
+					.ColorAndOpacity(FSlateColor::UseForeground())
+				]
+			]
+			+ SHorizontalBox::Slot()
+			.AutoWidth()
+			.VAlign(VAlign_Center)
+			.HAlign(HAlign_Right)
+			.Padding(5.0f, 0.0f)
+			[
+				SNew(SExpandableSearchArea, CollectionSearch->GetWidget())
+			]
+		]
+		.BodyContent()
+		[
+			CollectionViewPtr.ToSharedRef()
+		];
+}
+
+void SContentBrowser::ExtendViewOptionsMenu(const FContentBrowserConfig* Config)
+{
+	UToolMenu* Menu = UToolMenus::Get()->ExtendMenu("ContentBrowser.AssetViewOptions");
+
+	if (Config == nullptr || Config->bCanShowLockButton)
+	{
+		Menu->AddDynamicSection("ContentBrowserViewOptionsSection", FNewToolMenuDelegate::CreateLambda([](UToolMenu* InMenu)
+			{
+				if (UContentBrowserAssetViewContextMenuContext* Context = InMenu->FindContext<UContentBrowserAssetViewContextMenuContext>())
+				{
+					if (TSharedPtr<SContentBrowser> ContentBrowser = Context->OwningContentBrowser.Pin())
+					{
+						FToolMenuSection& Section = InMenu->AddSection("Locking", LOCTEXT("LockingMenuHeader", "Locking"), FToolMenuInsert("AssetViewType", EToolMenuInsertType::After));
+						Section.AddMenuEntry(
+							"ToggleLock",
+							TAttribute<FText>(ContentBrowser.ToSharedRef(), &SContentBrowser::GetLockMenuText),
+							LOCTEXT("LockToggleTooltip", "Toggle lock. If locked, this browser will ignore Find in Content Browser requests."),
+							TAttribute<FSlateIcon>(),
+							FUIAction(FExecuteAction::CreateLambda([ContentBrowser = Context->OwningContentBrowser]() {ContentBrowser.Pin()->ToggleLockClicked(); })));
+					}
+				}
+			}
+		));
+	}
+}
+
+SSplitter::ESizeRule SContentBrowser::GetFavoritesAreaSizeRule() const
+{
+	return FavoritesArea->IsExpanded() ? SSplitter::ESizeRule::FractionOfParent : SSplitter::ESizeRule::SizeToContent;
+}
+
+SSplitter::ESizeRule SContentBrowser::GetPathAreaSizeRule() const
+{
+	return PathArea->IsExpanded() ? SSplitter::ESizeRule::FractionOfParent : SSplitter::ESizeRule::SizeToContent;
+}
+
+SSplitter::ESizeRule SContentBrowser::GetCollectionsAreaSizeRule() const
+{
+	return CollectionArea->IsExpanded() ? SSplitter::ESizeRule::FractionOfParent : SSplitter::ESizeRule::SizeToContent;
+}
+
 FText SContentBrowser::GetHighlightedText() const
 {
 	return TextFilter->GetRawFilterText();
@@ -1191,6 +1250,10 @@ void SContentBrowser::SaveSettings() const
 	GConfig->SetBool(*SettingsIniSection, *(SettingsString + TEXT(".SourcesExpanded")), bSourcesViewExpanded, GEditorPerProjectIni);
 	GConfig->SetBool(*SettingsIniSection, *(SettingsString + TEXT(".Locked")), bIsLocked, GEditorPerProjectIni);
 
+	GConfig->SetBool(*SettingsIniSection, *(SettingsString + TEXT(".FavoritesAreaExpanded")), FavoritesArea->IsExpanded(), GEditorPerProjectIni);
+	GConfig->SetBool(*SettingsIniSection, *(SettingsString + TEXT(".PathAreaExpanded")), PathArea->IsExpanded(), GEditorPerProjectIni);
+	GConfig->SetBool(*SettingsIniSection, *(SettingsString + TEXT(".CollectionAreaExpanded")), CollectionArea->IsExpanded(), GEditorPerProjectIni);
+
 	for(int32 SlotIndex = 0; SlotIndex < PathAssetSplitterPtr->GetChildren()->Num(); SlotIndex++)
 	{
 		float SplitterSize = PathAssetSplitterPtr->SlotAt(SlotIndex).SizeValue.Get();
@@ -1203,6 +1266,7 @@ void SContentBrowser::SaveSettings() const
 		GConfig->SetFloat(*SettingsIniSection, *(SettingsString + FString::Printf(TEXT(".FavoriteSplitter.SlotSize%d"), SlotIndex)), SplitterSize, GEditorPerProjectIni);
 	}
 
+
 	// Save all our data using the settings string as a key in the user settings ini
 	FilterListPtr->SaveSettings(GEditorPerProjectIni, SettingsIniSection, SettingsString);
 	PathViewPtr->SaveSettings(GEditorPerProjectIni, SettingsIniSection, SettingsString);
@@ -1319,10 +1383,23 @@ void SContentBrowser::LoadSettings(const FName& InInstanceName)
 		}
 	}
 
+
 	// Now that we have determined the appropriate settings string, actually load the settings
 	GConfig->GetBool(*SettingsIniSection, *(SettingsString + TEXT(".SourcesExpanded")), bSourcesViewExpanded, GEditorPerProjectIni);
 	GConfig->GetBool(*SettingsIniSection, *(SettingsString + TEXT(".Locked")), bIsLocked, GEditorPerProjectIni);
 
+	bool bFavoritesAreaExpanded = false;
+	GConfig->GetBool(*SettingsIniSection, *(SettingsString + TEXT(".FavoritesAreaExpanded")), bFavoritesAreaExpanded, GEditorPerProjectIni);
+	FavoritesArea->SetExpanded(bFavoritesAreaExpanded);
+
+	bool bPathAreaExpanded = true;
+	GConfig->GetBool(*SettingsIniSection, *(SettingsString + TEXT(".PathAreaExpanded")), bPathAreaExpanded, GEditorPerProjectIni);
+	PathArea->SetExpanded(bPathAreaExpanded);
+
+	bool bCollectionAreaExpanded = false;
+	GConfig->GetBool(*SettingsIniSection, *(SettingsString + TEXT(".CollectionAreaExpanded")), bCollectionAreaExpanded, GEditorPerProjectIni);
+	CollectionArea->SetExpanded(bCollectionAreaExpanded);
+
 	for(int32 SlotIndex = 0; SlotIndex < PathAssetSplitterPtr->GetChildren()->Num(); SlotIndex++)
 	{
 		float SplitterSize = PathAssetSplitterPtr->SlotAt(SlotIndex).SizeValue.Get();
@@ -2053,6 +2130,19 @@ TSharedRef<SWidget> SContentBrowser::MakeAddNewContextMenu(const EContentBrowser
 		];
 }
 
+FReply SContentBrowser::OnImportClicked()
+{
+	FAssetToolsModule& AssetToolsModule = FModuleManager::LoadModuleChecked<FAssetToolsModule>("AssetTools");
+
+	FString CurrentPath = GetCurrentPath();
+	if (!CurrentPath.IsEmpty())
+	{
+		AssetToolsModule.Get().ImportAssetsWithDialog(GetCurrentPath());
+	}
+
+	return FReply::Handled();
+}
+
 void SContentBrowser::PopulateAddNewContextMenu(class UToolMenu* Menu)
 {
 	const UContentBrowserDataMenuContext_AddNewMenu* ContextObject = Menu->FindContext<UContentBrowserDataMenuContext_AddNewMenu>();
@@ -2261,18 +2351,20 @@ FReply SContentBrowser::ToggleLockClicked()
 	return FReply::Handled();
 }
 
-const FSlateBrush* SContentBrowser::GetToggleLockImage() const
+FText SContentBrowser::GetLockMenuText() const
 {
-	if ( bIsLocked )
-	{
-		return FEditorStyle::GetBrush("ContentBrowser.LockButton_Locked");
-	}
-	else
-	{
-		return FEditorStyle::GetBrush("ContentBrowser.LockButton_Unlocked");
-	}
+	return IsLocked() ? LOCTEXT("ContentBrowserLockMenu_Unlock", "Unlock Content Browser") : LOCTEXT("ContentBrowserLockMenu_Lock", "Lock Content Browser");
+}
+
+const FSlateBrush* SContentBrowser::GetLockIcon() const
+{
+	static const FName Unlock = "Icons.Unlock";
+	static const FName Lock = "Icons.Lock";
+
+	return FAppStyle::Get().GetBrush(IsLocked() ? Lock : Unlock);
 }
 
+
 EVisibility SContentBrowser::GetSourcesViewVisibility() const
 {
 	return bSourcesViewExpanded ? EVisibility::Visible : EVisibility::Collapsed;
@@ -2315,35 +2407,6 @@ EVisibility SContentBrowser::GetSourcesSwitcherVisibility() const
 	return GetDefault<UContentBrowserSettings>()->GetDockCollections() ? EVisibility::Collapsed : EVisibility::Visible;
 }
 
-const FSlateBrush* SContentBrowser::GetSourcesSwitcherIcon() const
-{
-	switch (ActiveSourcesWidgetIndex)
-	{
-	case ContentBrowserSourcesWidgetSwitcherIndex::PathView:
-		return FEditorStyle::GetBrush("ContentBrowser.Sources.Collections");
-	case ContentBrowserSourcesWidgetSwitcherIndex::CollectionsView:
-		return FEditorStyle::GetBrush("ContentBrowser.Sources.Paths");
-	default:
-		break;
-	}
-	check(false);
-	return nullptr;
-}
-
-FText SContentBrowser::GetSourcesSwitcherToolTipText() const
-{
-	switch (ActiveSourcesWidgetIndex)
-	{
-	case ContentBrowserSourcesWidgetSwitcherIndex::PathView:
-		return LOCTEXT("SwitchToCollectionView_ToolTip", "Switch to the Collections view");
-	case ContentBrowserSourcesWidgetSwitcherIndex::CollectionsView:
-		return LOCTEXT("SwitchToPathView_ToolTip", "Switch to the Paths view");
-	default:
-		break;
-	}
-	check(false);
-	return FText();
-}
 
 FReply SContentBrowser::OnSourcesSwitcherClicked()
 {
@@ -2354,29 +2417,12 @@ FReply SContentBrowser::OnSourcesSwitcherClicked()
 	return FReply::Handled();
 }
 
-FText SContentBrowser::GetSourcesSearchHintText() const
-{
-	switch (ActiveSourcesWidgetIndex)
-	{
-	case ContentBrowserSourcesWidgetSwitcherIndex::PathView:
-		return LOCTEXT("SearchPathsHint", "Search Paths");
-	case ContentBrowserSourcesWidgetSwitcherIndex::CollectionsView:
-		return LOCTEXT("SearchCollectionsHint", "Search Collections");
-	default:
-		break;
-	}
-	check(false);
-	return FText();
-}
+
 
 void SContentBrowser::OnContentBrowserSettingsChanged(FName PropertyName)
 {
-	const FName NAME_DockCollections = "DockCollections";//GET_MEMBER_NAME_CHECKED(UContentBrowserSettings, DockCollections); // Doesn't work as DockCollections is private :(
-	if (PropertyName.IsNone() || PropertyName == NAME_DockCollections)
+	if (PropertyName.IsNone())
 	{
-		// Ensure the omni-search is enabled correctly
-		CollectionViewPtr->SetAllowExternalSearch(!GetDefault<UContentBrowserSettings>()->GetDockCollections());
-
 		// Ensure the path is set to the correct view mode
 		UpdatePath();
 	}
diff --git a/Engine/Source/Editor/ContentBrowser/Private/SContentBrowser.h b/Engine/Source/Editor/ContentBrowser/Private/SContentBrowser.h
--- a/Engine/Source/Editor/ContentBrowser/Private/SContentBrowser.h
+++ b/Engine/Source/Editor/ContentBrowser/Private/SContentBrowser.h
@@ -15,6 +15,7 @@
 #include "CollectionManagerTypes.h"
 #include "IContentBrowserSingleton.h"
 #include "Editor/ContentBrowser/Private/HistoryManager.h"
+#include "Textures/SlateIcon.h"
 
 class FAssetContextMenu;
 class FFrontendFilter_Text;
@@ -285,6 +286,12 @@ private:
 	/** Handle creating a context menu for the "Add New" button */
 	TSharedRef<SWidget> MakeAddNewContextMenu(const EContentBrowserDataMenuContext_AddNewMenuDomain InDomain);
 
+	/* Whether or not the import button is enabled. */
+	bool IsImportButtonEnabled() const;
+
+	/** Called when the import button on the toolbar is clicked */
+	FReply OnImportClicked();
+
 	/** Handle populating a context menu for the "Add New" button */
 	void PopulateAddNewContextMenu(class UToolMenu* Menu);
 
@@ -318,8 +325,11 @@ private:
 	/** Handler for clicking the lock button */
 	FReply ToggleLockClicked();
 
-	/** Gets the brush used on the lock button */
-	const FSlateBrush* GetToggleLockImage() const;
+	/** Gets the menu text */
+	FText GetLockMenuText() const;
+
+	/** Gets icon for the lock button */
+	const FSlateBrush* GetLockIcon() const;
 
 	/** Gets the visibility state of the asset tree */
 	EVisibility GetSourcesViewVisibility() const;
@@ -345,9 +355,6 @@ private:
 	/** Handler for clicking the source switch button */
 	FReply OnSourcesSwitcherClicked();
 
-	/** Gets the source search hint text */
-	FText GetSourcesSearchHintText() const;
-
 	/** Called to handle the Content Browser settings changing */
 	void OnContentBrowserSettingsChanged(FName PropertyName);
 
@@ -486,12 +493,32 @@ private:
 	/** Get the visibility of the docked collections view */
 	EVisibility GetDockedCollectionsVisibility() const;
 
+	/** Get the visibility of the lock button */
+	EVisibility GetLockButtonVisibility() const;
+
+	/** Whether or not the collections view is docked or exists in its own panel in the same area as the sources view */
+	bool IsCollectionViewDocked() const;
+
 	/** Toggles the favorite status of an array of folders*/
 	void ToggleFolderFavorite(const TArray<FString>& FolderPaths);
 
 	/** Called when Asset View Options "Search" options change */
 	void HandleAssetViewSearchOptionsChanged();
 
+	/** Creates various widgets for the content browser main view */
+	TSharedRef<SWidget> CreateLockButton(const FContentBrowserConfig* Config);
+	TSharedRef<SWidget> CreateAssetView(const FContentBrowserConfig* Config);
+	TSharedRef<SWidget> CreateFavoritesView(const FContentBrowserConfig* Config);
+	TSharedRef<SWidget> CreatePathView(const FContentBrowserConfig* Config);
+	TSharedRef<SWidget> CreateDockedCollectionsView(const FContentBrowserConfig* Config);
+
+	/** Adds menu options to the view menu */
+	void ExtendViewOptionsMenu(const FContentBrowserConfig* Config);
+
+	/** Gets the size rule for various areas. When areas are a collapsed the splitter slot becomes auto sized, otherwise it is user sized */
+	SSplitter::ESizeRule GetFavoritesAreaSizeRule() const;
+	SSplitter::ESizeRule GetPathAreaSizeRule() const;
+	SSplitter::ESizeRule GetCollectionsAreaSizeRule() const;
 private:
 
 	/** The tab that contains this browser */
@@ -506,9 +533,12 @@ private:
 	/** The context menu manager for the path view */
 	TSharedPtr<class FPathContextMenu> PathContextMenu;
 
-	/** The sources search, shared between the paths and collections views */
+	/** The sources search for paths */
 	TSharedPtr<FSourcesSearch> SourcesSearch;
 
+	/** The sources search for collections*/
+	TSharedPtr<FSourcesSearch> CollectionSearch;
+
 	/** The asset tree widget */
 	TSharedPtr<SPathView> PathViewPtr;
 
@@ -533,6 +563,15 @@ private:
 	/** The path picker */
 	TSharedPtr<SComboButton> PathPickerButton;
 
+	/** Favorites area widget */
+	TSharedPtr<SExpandableArea> FavoritesArea;
+	
+	/** Path area widget */
+	TSharedPtr<SExpandableArea> PathArea;
+
+	/** Collection area widget */
+	TSharedPtr<SExpandableArea> CollectionArea;
+
 	/** Index of the active sources widget */
 	int32 ActiveSourcesWidgetIndex = 0;
 
diff --git a/Engine/Source/Editor/ContentBrowser/Private/SExpandableSearchArea.cpp b/Engine/Source/Editor/ContentBrowser/Private/SExpandableSearchArea.cpp
new file mode 100644
--- /dev/null
+++ b/Engine/Source/Editor/ContentBrowser/Private/SExpandableSearchArea.cpp
@@ -0,0 +1,60 @@
+// Copyright Epic Games, Inc. All Rights Reserved.
+
+#include "SExpandableSearchArea.h"
+#include "Widgets/SBoxPanel.h"
+#include "Widgets/Input/SButton.h"
+#include "Widgets/Images/SImage.h"
+#include "Widgets/Layout/SBox.h"
+#include "Widgets/Input/SSearchBox.h"
+
+void SExpandableSearchArea::Construct(const FArguments& InArgs, TSharedRef<SSearchBox> SearchBox)
+{
+	bIsExpanded = false;
+
+	SearchStyle = InArgs._Style;
+
+	ChildSlot
+	[
+		SNew(SHorizontalBox)
+		+ SHorizontalBox::Slot()
+		.HAlign(HAlign_Center)
+		.AutoWidth()
+		[
+			SNew(SButton)
+			.ButtonStyle(FAppStyle::Get(), "SimpleButton")
+			.ContentPadding(FMargin(0.0f, 2.0f))
+			.OnClicked(this, &SExpandableSearchArea::OnExpandSearchClicked)
+			[
+				SNew(SImage)
+				.Image(&SearchStyle->GlassImage)
+				.ColorAndOpacity(FSlateColor::UseForeground())
+			]
+		]
+		+ SHorizontalBox::Slot()
+		[
+			SNew(SBox)
+			.Visibility(this, &SExpandableSearchArea::GetSearchBoxVisibility)
+			.MinDesiredWidth(250.0f)
+			.MaxDesiredWidth(250.0f)
+			[
+				SearchBox
+			]
+		]
+	];
+}
+
+FReply SExpandableSearchArea::OnExpandSearchClicked()
+{
+	bIsExpanded = !bIsExpanded;
+	return FReply::Handled();
+}
+
+EVisibility SExpandableSearchArea::GetSearchBoxVisibility() const
+{
+	return bIsExpanded ? EVisibility::Visible : EVisibility::Collapsed;
+}
+
+const FSlateBrush* SExpandableSearchArea::GetExpandSearchImage() const
+{
+	return &SearchStyle->GlassImage;
+}
diff --git a/Engine/Source/Editor/ContentBrowser/Private/SExpandableSearchArea.h b/Engine/Source/Editor/ContentBrowser/Private/SExpandableSearchArea.h
new file mode 100644
--- /dev/null
+++ b/Engine/Source/Editor/ContentBrowser/Private/SExpandableSearchArea.h
@@ -0,0 +1,34 @@
+// Copyright Epic Games, Inc. All Rights Reserved.
+
+#pragma once 
+
+#include "Widgets/SCompoundWidget.h"
+#include "Widgets/DeclarativeSyntaxSupport.h"
+#include "Styling/SlateTypes.h"
+#include "Styling/AppStyle.h"
+#include "Styling/SlateWidgetStyleAsset.h"
+
+class SSearchBox;
+
+class SExpandableSearchArea : public SCompoundWidget
+{
+public:
+	SLATE_BEGIN_ARGS(SExpandableSearchArea)
+		: _Style(&FAppStyle::Get().GetWidgetStyle<FSearchBoxStyle>("SearchBox"))
+	{}
+		/** Style used to draw this search box */
+		SLATE_STYLE_ARGUMENT(FSearchBoxStyle, Style)
+
+	SLATE_END_ARGS()
+
+	void Construct(const FArguments& InArgs, TSharedRef<SSearchBox> SearchBox);
+
+private:
+
+	FReply OnExpandSearchClicked();
+	EVisibility GetSearchBoxVisibility() const;
+	const FSlateBrush* GetExpandSearchImage() const;
+private:
+	const FSearchBoxStyle* SearchStyle;
+	bool bIsExpanded;
+};
diff --git a/Engine/Source/Editor/ContentBrowser/Private/SFilterList.cpp b/Engine/Source/Editor/ContentBrowser/Private/SFilterList.cpp
--- a/Engine/Source/Editor/ContentBrowser/Private/SFilterList.cpp
+++ b/Engine/Source/Editor/ContentBrowser/Private/SFilterList.cpp
@@ -28,6 +28,7 @@
 #include "FrontendFilters.h"
 #include "ContentBrowserFrontEndFilterExtension.h"
 #include "Misc/BlacklistNames.h"
+#include "Widgets/Images/SImage.h"
 
 #define LOCTEXT_NAMESPACE "ContentBrowser"
 
@@ -183,8 +184,7 @@ public:
 		[
 			SNew(SBorder)
 			.Padding(0)
-			.BorderBackgroundColor( FLinearColor(0.2f, 0.2f, 0.2f, 0.2f) )
-			.BorderImage(FEditorStyle::GetBrush("ContentBrowser.FilterButtonBorder"))
+			.BorderImage(FAppStyle::Get().GetBrush("ContentBrowser.FilterBackground"))
 			[
 				SAssignNew( ToggleButtonPtr, SFilterCheckBox )
 				.Style(FEditorStyle::Get(), "ContentBrowser.FilterButton")
@@ -193,13 +193,24 @@ public:
 				.IsChecked(this, &SFilter::IsChecked)
 				.OnCheckStateChanged(this, &SFilter::FilterToggled)
 				.OnGetMenuContent(this, &SFilter::GetRightClickMenuContent)
-				.ForegroundColor(this, &SFilter::GetFilterForegroundColor)
 				[
-					SNew(STextBlock)
-					.ColorAndOpacity(this, &SFilter::GetFilterNameColorAndOpacity)
-					.Font(FEditorStyle::GetFontStyle("ContentBrowser.FilterNameFont"))
-					.ShadowOffset(FVector2D(1.f, 1.f))
-					.Text(this, &SFilter::GetFilterName)
+					SNew(SHorizontalBox)
+					+SHorizontalBox::Slot()
+					.VAlign(VAlign_Center)
+					.Padding(2.0f)
+					.AutoWidth()
+					[
+						SNew(SImage)
+						.Image(FAppStyle::Get().GetBrush("ContentBrowser.FilterImage"))
+						.ColorAndOpacity(this, &SFilter::GetFilterImageColorAndOpacity)
+					]
+					+SHorizontalBox::Slot()
+					.Padding(2.0f)
+					.VAlign(VAlign_Center)
+					[
+						SNew(STextBlock)
+						.Text(this, &SFilter::GetFilterName)
+					]
 				]
 			]
 		];
@@ -407,9 +418,9 @@ private:
 	}
 
 	/** Handler to determine the color of the checkbox when it is checked */
-	FSlateColor GetFilterForegroundColor() const
+	FSlateColor GetFilterImageColorAndOpacity() const
 	{
-		return IsChecked() == ECheckBoxState::Checked ? FilterColor : FLinearColor::White;
+		return IsChecked() == ECheckBoxState::Checked ? FilterColor : FLinearColor::Black;
 	}
 
 	/** Handler to determine the padding of the checkbox text when it is pressed */
@@ -418,12 +429,6 @@ private:
 		return ToggleButtonPtr->IsPressed() ? FMargin(3,2,4,0) : FMargin(3,1,4,1);
 	}
 
-	/** Handler to determine the color of the checkbox text when it is hovered */
-	FSlateColor GetFilterNameColorAndOpacity() const
-	{
-		const float DimFactor = 0.75f;
-		return IsHovered() ? FLinearColor(DimFactor, DimFactor, DimFactor, 1.0f) : FLinearColor::White;
-	}
 
 	/** Returns the display name for this filter */
 	FText GetFilterName() const
diff --git a/Engine/Source/Editor/ContentBrowser/Private/SourcesSearch.cpp b/Engine/Source/Editor/ContentBrowser/Private/SourcesSearch.cpp
--- a/Engine/Source/Editor/ContentBrowser/Private/SourcesSearch.cpp
+++ b/Engine/Source/Editor/ContentBrowser/Private/SourcesSearch.cpp
@@ -24,7 +24,7 @@ void FSourcesSearch::SetHintText(const TAttribute<FText>& InHintText)
 	SearchBox->SetHintText(InHintText);
 }
 
-TSharedRef<SWidget> FSourcesSearch::GetWidget() const
+TSharedRef<SSearchBox> FSourcesSearch::GetWidget() const
 {
 	return SearchBox.ToSharedRef();
 }
diff --git a/Engine/Source/Editor/ContentBrowser/Private/SourcesSearch.h b/Engine/Source/Editor/ContentBrowser/Private/SourcesSearch.h
--- a/Engine/Source/Editor/ContentBrowser/Private/SourcesSearch.h
+++ b/Engine/Source/Editor/ContentBrowser/Private/SourcesSearch.h
@@ -31,7 +31,7 @@ public:
 	void SetHintText(const TAttribute<FText>& InHintText);
 
 	/** Get the underlying search widget */
-	TSharedRef<SWidget> GetWidget() const;
+	TSharedRef<SSearchBox> GetWidget() const;
 
 private:
 	/** Handler for when search term changes */
diff --git a/Engine/Source/Editor/ContentBrowser/Private/SourcesViewWidgets.cpp b/Engine/Source/Editor/ContentBrowser/Private/SourcesViewWidgets.cpp
--- a/Engine/Source/Editor/ContentBrowser/Private/SourcesViewWidgets.cpp
+++ b/Engine/Source/Editor/ContentBrowser/Private/SourcesViewWidgets.cpp
@@ -29,6 +29,8 @@ void SAssetTreeItem::Construct( const FArguments& InArgs )
 	IsItemExpanded = InArgs._IsItemExpanded;
 	bDraggedOver = false;
 
+	IsSelected = InArgs._IsSelected;
+
 	FolderOpenBrush = FEditorStyle::GetBrush("ContentBrowser.AssetTreeFolderOpen");
 	FolderClosedBrush = FEditorStyle::GetBrush("ContentBrowser.AssetTreeFolderClosed");
 	FolderOpenCodeBrush = FEditorStyle::GetBrush("ContentBrowser.AssetTreeFolderOpenCode");
@@ -138,11 +140,6 @@ FReply SAssetTreeItem::OnDrop( const FGeometry& MyGeometry, const FDragDropEvent
 	return FReply::Unhandled();
 }
 
-void SAssetTreeItem::Tick( const FGeometry& AllottedGeometry, const double InCurrentTime, const float InDeltaTime )
-{
-	LastGeometry = AllottedGeometry;
-}
-
 bool SAssetTreeItem::VerifyNameChanged(const FText& InName, FText& OutError) const
 {
 	if ( TreeItem.IsValid() )
@@ -167,6 +164,7 @@ void SAssetTreeItem::HandleNameCommitted( const FText& NewText, ETextCommit::Typ
 		{
 			TreeItemPtr->SetNamingFolder(false);
 		
+			const FGeometry LastGeometry = GetTickSpaceGeometry();
 			FVector2D MessageLoc;
 			MessageLoc.X = LastGeometry.AbsolutePosition.X;
 			MessageLoc.Y = LastGeometry.AbsolutePosition.Y + LastGeometry.Size.Y * LastGeometry.Scale;
@@ -207,6 +205,7 @@ FSlateColor SAssetTreeItem::GetFolderColor() const
 {
 	if (TSharedPtr<FTreeItem> TreeItemPin = TreeItem.Pin())
 	{
+		FLinearColor FoundColor;
 		FContentBrowserItemDataAttributeValue ColorAttributeValue = TreeItemPin->GetItem().GetItemAttribute(ContentBrowserItemAttributes::ItemColor);
 		if (ColorAttributeValue.IsValid())
 		{
@@ -215,17 +214,27 @@ FSlateColor SAssetTreeItem::GetFolderColor() const
 			FLinearColor Color;
 			if (Color.InitFromString(ColorStr))
 			{
-				return Color;
+				FoundColor = Color;
 			}
 		}
 		else
 		{
 			if (TSharedPtr<FLinearColor> Color = ContentBrowserUtils::LoadColor(TreeItemPin->GetItem().GetVirtualPath().ToString()))
 			{
-				return *Color;
+				FoundColor = *Color;
 			}
 		}
+
+		if(FoundColor == ContentBrowserUtils::GetDefaultColor() && IsSelected.IsBound() && IsSelected.Execute())
+		{
+			static const FSlateColor InvertedForeground = FAppStyle::Get().GetSlateColor("Colors.ForegroundInverted");
+
+			return InvertedForeground;
+		}
+
+		return FoundColor;
 	}
+	
 	return ContentBrowserUtils::GetDefaultColor();
 }
 
@@ -276,7 +285,7 @@ void SCollectionTreeItem::Construct( const FArguments& InArgs )
 	ChildSlot
 	[
 		SAssignNew(AssetTagItem, SAssetTagItem)
-		.ViewMode(InArgs._ViewMode)
+		.ViewMode(EAssetTagItemViewMode::Compact)
 		.BaseColor(this, &SCollectionTreeItem::GetCollectionColor)
 		.DisplayName(this, &SCollectionTreeItem::GetNameText)
 		.CountText(this, &SCollectionTreeItem::GetCollectionObjectCountText)
diff --git a/Engine/Source/Editor/ContentBrowser/Private/SourcesViewWidgets.h b/Engine/Source/Editor/ContentBrowser/Private/SourcesViewWidgets.h
--- a/Engine/Source/Editor/ContentBrowser/Private/SourcesViewWidgets.h
+++ b/Engine/Source/Editor/ContentBrowser/Private/SourcesViewWidgets.h
@@ -61,7 +61,6 @@ public:
 	virtual void OnDragLeave( const FDragDropEvent& DragDropEvent ) override;
 	virtual FReply OnDragOver( const FGeometry& MyGeometry, const FDragDropEvent& DragDropEvent ) override;
 	virtual FReply OnDrop( const FGeometry& MyGeometry, const FDragDropEvent& DragDropEvent ) override;
-	virtual void Tick( const FGeometry& AllottedGeometry, const double InCurrentTime, const float InDeltaTime ) override;
 
 private:
 	/** Handles verifying name changes */
@@ -111,9 +110,6 @@ private:
 	/** True when this item has children and is expanded */
 	TAttribute<bool> IsItemExpanded;
 
-	/** The geometry last frame. Used when telling popup messages where to appear. */
-	FGeometry LastGeometry;
-
 	/** Brushes for the different folder states */
 	const FSlateBrush* FolderOpenBrush;
 	const FSlateBrush* FolderClosedBrush;
@@ -121,6 +117,9 @@ private:
 	const FSlateBrush* FolderClosedCodeBrush;
 	const FSlateBrush* FolderDeveloperBrush;
 
+	/** Delegate called to get the selection state of an asset path */
+	FIsSelected IsSelected;
+
 	/** True when a drag is over this item with a drag operation that we know how to handle. The operation itself may not be valid to drop. */
 	bool bDraggedOver;
 
@@ -128,7 +127,7 @@ private:
 	EFolderType FolderType;
 
 	/** Widget to display the name of the asset item and allows for renaming */
-	TSharedPtr< SInlineEditableTextBlock > InlineRenameWidget;
+	TSharedPtr<SInlineEditableTextBlock> InlineRenameWidget;
 
 	/** Handle to the registered EnterEditingMode delegate. */
 	FDelegateHandle EnterEditingModeDelegateHandle;
@@ -150,7 +149,6 @@ public:
 	SLATE_BEGIN_ARGS( SCollectionTreeItem )
 		: _CollectionItem( TSharedPtr<FCollectionItem>() )
 		, _ParentWidget()
-		, _ViewMode(EAssetTagItemViewMode::Standard)
 	{}
 
 		/** Data for the collection this item represents */
@@ -159,9 +157,6 @@ public:
 		/** The parent widget */
 		SLATE_ARGUMENT( TSharedPtr<SWidget>, ParentWidget )
 
-		/** Should this collection item use the standard or compact view? */
-		SLATE_ARGUMENT( EAssetTagItemViewMode, ViewMode )
-
 		/** Delegate for when the user begins to rename the item */
 		SLATE_EVENT( FOnBeginNameChange, OnBeginNameChange )
 
diff --git a/Engine/Source/Editor/ContentBrowser/Public/ContentBrowserMenuContexts.h b/Engine/Source/Editor/ContentBrowser/Public/ContentBrowserMenuContexts.h
--- a/Engine/Source/Editor/ContentBrowser/Public/ContentBrowserMenuContexts.h
+++ b/Engine/Source/Editor/ContentBrowser/Public/ContentBrowserMenuContexts.h
@@ -53,7 +53,7 @@ class CONTENTBROWSER_API UContentBrowserAssetViewContextMenuContext : public UOb
 	GENERATED_BODY()
 
 public:
-
+	TWeakPtr<SContentBrowser> OwningContentBrowser;
 	TWeakPtr<SAssetView> AssetView;
 };
 
diff --git a/Engine/Source/Editor/EditorStyle/Private/StarshipStyle.cpp b/Engine/Source/Editor/EditorStyle/Private/StarshipStyle.cpp
--- a/Engine/Source/Editor/EditorStyle/Private/StarshipStyle.cpp
+++ b/Engine/Source/Editor/EditorStyle/Private/StarshipStyle.cpp
@@ -1273,21 +1273,18 @@ void FStarshipEditorStyle::FStyle::SetupGeneralStyles()
 		
 	// Asset Thumbnail
 	{
-		Set( "AssetThumbnail.AssetBackground", new IMAGE_BRUSH( "Common/AssetBackground", FVector2D(64.f, 64.f), FLinearColor(1.f, 1.f, 1.f, 1.0f) ) );
+		Set( "AssetThumbnail.AssetBackground", new FSlateColorBrush(FStyleColors::Input));
 		Set( "AssetThumbnail.ClassBackground", new IMAGE_BRUSH( "Common/ClassBackground_64x", FVector2D(64.f, 64.f), FLinearColor(0.75f, 0.75f, 0.75f, 1.0f) ) );
-		Set( "AssetThumbnail.DataOnlyBPAssetBackground", new IMAGE_BRUSH( "Common/DataOnlyBPAssetBackground_64x", FVector2D(64.f, 64.f), FLinearColor(1, 1, 1, 1) ) );
 		Set( "AssetThumbnail.Font", DEFAULT_FONT( "Regular", 10 ) );
 		Set( "AssetThumbnail.FontSmall", DEFAULT_FONT( "Regular", 7 ) );
-		Set( "AssetThumbnail.ColorAndOpacity", FLinearColor(0.75f, 0.75f, 0.75f, 1) );
-		Set( "AssetThumbnail.ShadowOffset", FVector2D(1,1) );
-		Set( "AssetThumbnail.ShadowColorAndOpacity", FLinearColor(0, 0, 0, 0.5) );
+		Set( "AssetThumbnail.ColorAndOpacity", FLinearColor(1.75f, 1.75f, 1.75f, 1) );
 		Set( "AssetThumbnail.HintFont", DEFAULT_FONT( "Regular", 8 ) );
 		Set( "AssetThumbnail.HintFontSmall", DEFAULT_FONT( "Regular", 6 ) );
 		Set( "AssetThumbnail.HintColorAndOpacity", FLinearColor(0.75f, 0.75f, 0.75f, 1) );
 		Set( "AssetThumbnail.HintShadowOffset", FVector2D(1,1) );
 		Set( "AssetThumbnail.HintShadowColorAndOpacity", FLinearColor(0, 0, 0, 0.5) );
 		Set( "AssetThumbnail.HintBackground", new BOX_BRUSH( "Common/TableViewHeader", FMargin(8.0f/32.0f) ) );
-		Set( "AssetThumbnail.Border", new FSlateColorBrush( FColor::White ) );
+	
 	}
 
 	// Open any asset dialog
@@ -5907,27 +5904,27 @@ void FStarshipEditorStyle::FStyle::SetupContentBrowserStyle()
 	// Content Browser
 	{
 		// Tab and menu icon
-		Set( "ContentBrowser.TabIcon", new IMAGE_BRUSH_SVG("Starship/LevelEditor/Menus/ContentBrowser", Icon16x16) );
+		Set("ContentBrowser.TabIcon", new IMAGE_BRUSH_SVG("Starship/LevelEditor/Menus/ContentBrowser", Icon16x16));
 
 		// Sources View
-		Set( "ContentBrowser.SourceTitleFont", DEFAULT_FONT( "Regular", 12 ) );
+		Set("ContentBrowser.SourceTitleFont", DEFAULT_FONT( "Regular", 12 ) );
+
+		Set("ContentBrowser.SourceTreeItemFont", FStarshipCoreStyle::GetCoreStyle().GetFontStyle("NormalFont"));
+		Set("ContentBrowser.SourceTreeRootItemFont", FStarshipCoreStyle::GetCoreStyle().GetFontStyle("NormalFont"));
 
-		// @todo vreditor urgent: Increase the size of Content Browser fonts while in VR
-		Set("ContentBrowser.SourceListItemFont", DEFAULT_FONT("Regular", 10));
-		Set("ContentBrowser.SourceTreeItemFont", DEFAULT_FONT("Regular", 10));
+		Set("ContentBrowser.BreadcrumbPathPickerFolder", new CORE_IMAGE_BRUSH_SVG("Starship/Common/folder-closed", Icon16x16));
+
+		Set("ContentBrowser.AssetTreeFolderClosed", new CORE_IMAGE_BRUSH_SVG("Starship/Common/folder-closed", Icon16x16));
+		Set("ContentBrowser.AssetTreeFolderOpen", new CORE_IMAGE_BRUSH_SVG("Starship/Common/folder-open", Icon16x16));
 
-		Set( "ContentBrowser.SourceTreeRootItemFont", DEFAULT_FONT( "Regular", 12 ) );
-		Set( "ContentBrowser.AssetTreeFolderClosed", new IMAGE_BRUSH( "Icons/FolderClosed", FVector2D(18, 16) ) );
-		Set( "ContentBrowser.BreadcrumbPathPickerFolder", new IMAGE_BRUSH( "Icons/FolderClosed", FVector2D(18, 16), FLinearColor::Gray ) );
-		Set( "ContentBrowser.AssetTreeFolderOpen", new IMAGE_BRUSH( "Icons/FolderOpen", FVector2D(18, 16) ) );
 		Set( "ContentBrowser.AssetTreeFolderDeveloper", new IMAGE_BRUSH( "Icons/FolderDeveloper", FVector2D(18, 16) ) );
 		Set( "ContentBrowser.AssetTreeFolderOpenCode", new IMAGE_BRUSH( "Icons/FolderOpen_Code", FVector2D(18, 16) ) );
 		Set( "ContentBrowser.AssetTreeFolderClosedCode", new IMAGE_BRUSH( "Icons/FolderClosed_Code", FVector2D(18, 16) ) );
-		Set( "ContentBrowser.AddCollectionButtonIcon", new IMAGE_BRUSH( "Icons/PlusSymbol_12x", Icon12x12 ) );
 
-		Set( "ContentBrowser.Splitter", FSplitterStyle()
-			.SetHandleNormalBrush( FSlateColorBrush( FLinearColor(FColor( 32, 32, 32) ) ) )
-			.SetHandleHighlightBrush( FSlateColorBrush( FLinearColor(FColor( 96, 96, 96) ) ) ) 
+		Set("ContentBrowser.DefaultFolderColor", FStyleColors::AccentFolder);
+
+		Set( "ContentBrowser.Splitter", FSplitterStyle(FStarshipCoreStyle::GetCoreStyle().GetWidgetStyle<FSplitterStyle>("Splitter"))
+			.SetHandleNormalBrush(FSlateColorBrush(FStyleColors::Input))
 			);
 
 		// Asset list view
@@ -5947,11 +5944,12 @@ void FStarshipEditorStyle::FStyle::SetupContentBrowserStyle()
 		);
 
 		// Tile view
-		Set( "ContentBrowser.AssetTileViewNameFont", DEFAULT_FONT( "Regular", 9 ) );
+		Set( "ContentBrowser.AssetTileViewNameFont", DEFAULT_FONT("Regular", 9));
 		Set( "ContentBrowser.AssetTileViewNameFontSmall", DEFAULT_FONT( "VeryLight", 8 ) );
 		Set( "ContentBrowser.AssetTileViewNameFontVerySmall", DEFAULT_FONT( "VeryLight", 7 ) );
-		Set( "ContentBrowser.AssetTileViewNameFontDirty", DEFAULT_FONT( "Bold", 10 ) );
-		Set( "ContentBrowser.AssetListView.TableRow", FTableRowStyle()
+		Set( "ContentBrowser.AssetTileViewNameFontDirty", FStyleFonts::Get().SmallBold);
+
+		Set("ContentBrowser.AssetListView.ColumnListTableRow", FTableRowStyle()
 			.SetEvenRowBackgroundBrush( FSlateNoResource() )
 			.SetEvenRowBackgroundHoveredBrush( IMAGE_BRUSH( "Common/Selection", Icon8x8, FLinearColor(1.0f, 1.0f, 1.0f, 0.1f) ) )
 			.SetOddRowBackgroundBrush( FSlateNoResource() )
@@ -5965,6 +5963,20 @@ void FStarshipEditorStyle::FStyle::SetupContentBrowserStyle()
 			.SetSelectedTextColor( InvertedForeground )
 			);
 
+		Set("ContentBrowser.AssetListView.TileTableRow", FTableRowStyle()
+			.SetEvenRowBackgroundBrush(FSlateNoResource() )
+			.SetEvenRowBackgroundHoveredBrush(FSlateNoResource())
+			.SetOddRowBackgroundBrush(FSlateNoResource())
+			.SetOddRowBackgroundHoveredBrush(FSlateNoResource())
+			.SetSelectorFocusedBrush(FSlateNoResource())
+			.SetActiveBrush(FSlateNoResource())
+			.SetActiveHoveredBrush(FSlateNoResource())
+			.SetInactiveBrush(FSlateNoResource())
+			.SetInactiveHoveredBrush(FSlateNoResource())
+			.SetTextColor(DefaultForeground)
+			.SetSelectedTextColor(DefaultForeground)
+			);
+
 		Set( "ContentBrowser.TileViewTooltip.ToolTipBorder", new FSlateColorBrush( FLinearColor::Black ) );
 		Set( "ContentBrowser.TileViewTooltip.NonContentBorder", new BOX_BRUSH( "/Docking/TabContentArea", FMargin(4/16.0f) ) );
 		Set( "ContentBrowser.TileViewTooltip.ContentBorder", new FSlateColorBrush( FStyleColors::Background ));
@@ -5978,17 +5990,23 @@ void FStarshipEditorStyle::FStyle::SetupContentBrowserStyle()
 		// Filter list
 		/* Set images for various SCheckBox states associated with "ContentBrowser.FilterButton" ... */
 		const FCheckBoxStyle ContentBrowserFilterButtonCheckBoxStyle = FCheckBoxStyle()
-			.SetUncheckedImage( IMAGE_BRUSH( "ContentBrowser/FilterUnchecked", FVector2D( 10.0f,20.0f ) ) )
-			.SetUncheckedHoveredImage( IMAGE_BRUSH( "ContentBrowser/FilterUnchecked", FVector2D( 10.0f, 20.0f ), FLinearColor( 0.5f, 0.5f, 0.5f, 1.0f ) ) )
-			.SetUncheckedPressedImage( IMAGE_BRUSH( "ContentBrowser/FilterUnchecked",FVector2D( 10.0f, 20.0f ), FLinearColor( 0.5f, 0.5f, 0.5f, 1.0f ) ) )
-			.SetCheckedImage( IMAGE_BRUSH( "ContentBrowser/FilterChecked",  FVector2D( 10.0f, 20.0f ) ) )
-			.SetCheckedHoveredImage( IMAGE_BRUSH( "ContentBrowser/FilterChecked",  FVector2D( 10.0f, 20.0f ), FLinearColor( 0.5f, 0.5f, 0.5f, 1.0f ) ) )
-			.SetCheckedPressedImage( IMAGE_BRUSH( "ContentBrowser/FilterChecked", FVector2D( 10.0f, 20.0f ), FLinearColor( 0.5f, 0.5f, 0.5f, 1.0f ) ) );
-		/* ... and add the new style */
-		Set( "ContentBrowser.FilterButton", ContentBrowserFilterButtonCheckBoxStyle );
+			.SetUncheckedImage(FSlateNoResource())
+			.SetUncheckedHoveredImage(FSlateNoResource())
+			.SetUncheckedPressedImage(FSlateNoResource())
+			.SetCheckedImage(FSlateNoResource())
+			.SetCheckedHoveredImage(FSlateNoResource())
+			.SetCheckedPressedImage(FSlateNoResource())
+			.SetForegroundColor(FStyleColors::Foreground)
+			.SetHoveredForegroundColor(FStyleColors::ForegroundHover)
+			.SetCheckedForegroundColor(FStyleColors::Foreground)
+			.SetCheckedHoveredForegroundColor(FStyleColors::ForegroundHover)
+			.SetCheckedPressedForegroundColor(FStyleColors::ForegroundHover)
+			.SetPressedForegroundColor(FStyleColors::ForegroundHover);
 
-		Set( "ContentBrowser.FilterNameFont", DEFAULT_FONT( "Regular", 10 ) );
-		Set( "ContentBrowser.FilterButtonBorder", new BOX_BRUSH( "Common/RoundedSelection_16x", FMargin(4.0f/16.0f) ) );
+		Set("ContentBrowser.FilterImage", new CORE_IMAGE_BRUSH_SVG("Starship/Common/filled-circle", Icon8x8));
+		Set("ContentBrowser.FilterBackground", new FSlateRoundedBoxBrush(FStyleColors::Header, 2.0f));
+		/* ... and add the new style */
+		Set("ContentBrowser.FilterButton", ContentBrowserFilterButtonCheckBoxStyle );
 
 		// Sources view
 		Set("ContentBrowser.Sources.Paths", new IMAGE_BRUSH("ContentBrowser/Sources_Paths_16x", Icon16x16));
@@ -5996,21 +6014,7 @@ void FStarshipEditorStyle::FStyle::SetupContentBrowserStyle()
 		Set("ContentBrowser.Sources.Collections.Compact", new IMAGE_BRUSH("ContentBrowser/Sources_Collections_Compact_16x", Icon16x16));
 
 		// Asset tags (common)
-		Set("ContentBrowser.AssetTagBackground", new BOX_BRUSH("Common/RoundedSelection_16x", FMargin(4.0f/16.0f)));
-		Set("ContentBrowser.AssetTagTableRow", FTableRowStyle(NormalTableRowStyle)
-			.SetEvenRowBackgroundBrush(FSlateNoResource())
-			.SetEvenRowBackgroundHoveredBrush(BORDER_BRUSH("Common/RoundedSelection_16x", FMargin(2.0f/16.0f), FLinearColor(1.0f, 1.0f, 1.0f, 0.1f)))
-			.SetOddRowBackgroundBrush(FSlateNoResource())
-			.SetOddRowBackgroundHoveredBrush(BORDER_BRUSH("Common/RoundedSelection_16x", FMargin(2.0f/16.0f), FLinearColor(1.0f, 1.0f, 1.0f, 0.1f)))
-			.SetActiveBrush(BORDER_BRUSH("Common/RoundedSelection_16x", FMargin(2.0f/16.0f), SelectionColor))
-			.SetActiveHoveredBrush(BORDER_BRUSH("Common/RoundedSelection_16x", FMargin(2.0f/16.0f), SelectionColor))
-			.SetInactiveBrush(BORDER_BRUSH("Common/RoundedSelection_16x", FMargin(2.0f/16.0f), SelectionColor_Inactive))
-			.SetInactiveHoveredBrush(BORDER_BRUSH("Common/RoundedSelection_16x", FMargin(2.0f/16.0f), SelectionColor_Inactive))
-			.SetActiveHighlightedBrush(BORDER_BRUSH("Common/RoundedSelection_16x", FMargin(2.0f/16.0f), HighlightColor))
-			.SetInactiveHighlightedBrush(BORDER_BRUSH("Common/RoundedSelection_16x", FMargin(2.0f/16.0f), HighlightColor))
-			.SetTextColor(DefaultForeground)
-			.SetSelectedTextColor(DefaultForeground)
-			);
+		Set("ContentBrowser.AssetTagBackground", new FSlateRoundedBoxBrush(FStyleColors::White, 2.0));
 
 		// Asset tags (standard)
 		Set("ContentBrowser.AssetTagButton", FCheckBoxStyle()
@@ -6025,9 +6029,8 @@ void FStarshipEditorStyle::FStyle::SetupContentBrowserStyle()
 			.SetCheckedPressedImage(IMAGE_BRUSH("ContentBrowser/AssetTagCheckbox_Flat", FVector2D(14.0f, 28.0f), FLinearColor(0.5f, 0.5f, 0.5f, 1.0f)))
 			.SetPadding(0.0f)
 			);
-		Set("ContentBrowser.AssetTagNameFont", DEFAULT_FONT("Bold", 11));
+
 		Set("ContentBrowser.AssetTagNamePadding", FMargin(4.0f));
-		Set("ContentBrowser.AssetTagCountFont", DEFAULT_FONT("Bold", 8));
 		Set("ContentBrowser.AssetTagCountPadding", FMargin(4.0f));
 
 		// Asset tags (compact)
@@ -6043,9 +6046,8 @@ void FStarshipEditorStyle::FStyle::SetupContentBrowserStyle()
 			.SetCheckedPressedImage(IMAGE_BRUSH("ContentBrowser/AssetTagCheckbox_Flat_Compact", FVector2D(10.0f, 20.0f), FLinearColor(0.5f, 0.5f, 0.5f, 1.0f)))
 			.SetPadding(0.0f)
 			);
-		Set("ContentBrowser.AssetTagNameFont.Compact", DEFAULT_FONT("Regular", 10));
+
 		Set("ContentBrowser.AssetTagNamePadding.Compact", FMargin(2.0f));
-		Set("ContentBrowser.AssetTagCountFont.Compact", DEFAULT_FONT("Light", 8));
 		Set("ContentBrowser.AssetTagCountPadding.Compact", FMargin(2.0f));
 
 		// Thumbnail editing mode
@@ -6055,7 +6057,7 @@ void FStarshipEditorStyle::FStyle::SetupContentBrowserStyle()
 			.SetShadowOffset( FVector2D::ZeroVector )
 		);
 
-		Set( "ContentBrowser.EditModeLabelBorder", new FSlateColorBrush(SelectorColor.GetSpecifiedColor()) );
+
 		Set( "ContentBrowser.PrimitiveCustom", new IMAGE_BRUSH( "ContentBrowser/ThumbnailCustom", Icon32x32 ) );
 		Set( "ContentBrowser.PrimitiveSphere", new IMAGE_BRUSH( "ContentBrowser/ThumbnailSphere", Icon32x32 ) );
 		Set( "ContentBrowser.PrimitiveCube", new IMAGE_BRUSH( "ContentBrowser/ThumbnailCube", Icon32x32 ) );
@@ -6063,8 +6065,6 @@ void FStarshipEditorStyle::FStyle::SetupContentBrowserStyle()
 		Set( "ContentBrowser.PrimitiveCylinder", new IMAGE_BRUSH( "ContentBrowser/ThumbnailCylinder", Icon32x32 ) );
 		Set( "ContentBrowser.ResetPrimitiveToDefault", new IMAGE_BRUSH("ContentBrowser/ThumbnailReset", Icon32x32) );
 
-		Set( "ContentBrowser.TopBar.GroupBorder", new FSlateColorBrush(FStyleColors::Background));
-
 		Set( "ContentBrowser.TopBar.Font", FTextBlockStyle( NormalText )
 			.SetFont( DEFAULT_FONT( "Bold", 11 ) )
 			.SetColorAndOpacity( FLinearColor( 1.0f, 1.0f, 1.0f ) )
@@ -6115,36 +6115,37 @@ void FStarshipEditorStyle::FStyle::SetupContentBrowserStyle()
 		Set("MediaAsset.AssetActions.Mute.Large", new IMAGE_BRUSH("Icons/icon_SCueEd_Mute_40x", Icon40x40));
 				
 		// Misc
-		Set( "ContentBrowser.ThumbnailShadow", new BOX_BRUSH( "ContentBrowser/ThumbnailShadow" , FMargin( 4.0f / 64.0f ) ) );
+		/** Should be moved, shared */ Set( "ContentBrowser.ThumbnailShadow", new BOX_BRUSH( "ContentBrowser/ThumbnailShadow" , FMargin( 4.0f / 64.0f ) ) );
+
+
+
 		Set( "ContentBrowser.ColumnViewAssetIcon", new IMAGE_BRUSH( "Icons/doc_16x", Icon16x16 ) );
-		Set( "ContentBrowser.ColumnViewFolderIcon", new IMAGE_BRUSH( "Icons/FolderClosed", FVector2D(18, 16) ) );
+
+		Set( "ContentBrowser.ColumnViewFolderIcon", new CORE_IMAGE_BRUSH_SVG( "Starship/Common/folder-closed", Icon16x16 ) );
 		Set( "ContentBrowser.ColumnViewDeveloperFolderIcon", new IMAGE_BRUSH( "Icons/FolderDeveloper", FVector2D(18, 16) ) );
-		Set( "ContentBrowser.ListViewFolderIcon.Base", new IMAGE_BRUSH( "Icons/Folders/Folder_Base_256x", FVector2D(256, 256) ) );
-		Set( "ContentBrowser.ListViewFolderIcon.Mask", new IMAGE_BRUSH( "Icons/Folders/Folder_BaseHi_256x", FVector2D(256, 256) ) );
-		Set( "ContentBrowser.ListViewDeveloperFolderIcon.Base", new IMAGE_BRUSH( "Icons/Folders/FolderDev_Base_256x", FVector2D(256, 256) ) );
-		Set( "ContentBrowser.ListViewDeveloperFolderIcon.Mask", new IMAGE_BRUSH( "Icons/Folders/FolderDev_BaseHi_256x", FVector2D(256, 256) ) );
-		Set( "ContentBrowser.TileViewFolderIcon.Base", new IMAGE_BRUSH( "Icons/Folders/Folder_Base_512x", FVector2D(512, 512) ) );
-		Set( "ContentBrowser.TileViewFolderIcon.Mask", new IMAGE_BRUSH( "Icons/Folders/Folder_BaseHi_512x", FVector2D(512, 512) ) );
-		Set( "ContentBrowser.TileViewDeveloperFolderIcon.Base", new IMAGE_BRUSH( "Icons/Folders/FolderDev_Base_512x", FVector2D(512, 512) ) );
-		Set( "ContentBrowser.TileViewDeveloperFolderIcon.Mask", new IMAGE_BRUSH( "Icons/Folders/FolderDev_BaseHi_512x", FVector2D(512, 512) ) );
-		Set( "ContentBrowser.PathText", FTextBlockStyle(NormalText)
-			.SetFont( DEFAULT_FONT( "Bold", 11 ) )
-			.SetColorAndOpacity( FLinearColor( 1.0f, 1.0f, 1.0f ) )
-			.SetHighlightColor( FLinearColor( 1.0f, 1.0f, 1.0f ) )
-			.SetShadowOffset( FVector2D( 1,1 ) )
-			.SetShadowColorAndOpacity( FLinearColor(0,0,0,0.9f) ) );
+
+		Set( "ContentBrowser.ListViewFolderIcon", new CORE_IMAGE_BRUSH_SVG( "Starship/Common/folder-closed", FVector2D(256, 256) ) );
+		Set( "ContentBrowser.ListViewDeveloperFolderIcon", new IMAGE_BRUSH( "Icons/Folders/FolderDev_Base_256x", FVector2D(256, 256) ) );
+
+		Set( "ContentBrowser.TileViewFolderIcon", new CORE_IMAGE_BRUSH_SVG("Starship/Common/folder-closed", FVector2D(512, 512) ) );
+		Set( "ContentBrowser.TileViewDeveloperFolderIcon", new IMAGE_BRUSH( "Icons/Folders/FolderDev_Base_512x", FVector2D(512, 512) ) );
+
+
+		Set("ContentBrowser.AssetTileItem.ThumbnailAreaBackground", new FSlateRoundedBoxBrush(FStyleColors::Input, 4.0f));
+		Set("ContentBrowser.AssetTileItem.NameAreaBackground", new FSlateRoundedBoxBrush(FStyleColors::Dropdown, 4.0f));
+		Set("ContentBrowser.AssetTileItem.SelectedBorder", new FSlateRoundedBoxBrush(FStyleColors::Transparent, 4.0f, FStyleColors::Primary, 2.0f));
+		Set("ContentBrowser.AssetTileItem.SelectedHoverBorder", new FSlateRoundedBoxBrush(FStyleColors::Transparent, 4.0f, FStyleColors::PrimaryHover, 2.0f));
+		Set("ContentBrowser.AssetTileItem.HoverBorder", new FSlateRoundedBoxBrush(FStyleColors::Transparent, 4.0f, FStyleColors::Hover, 2.0f));
+
+		Set("ContentBrowser.AssetTileItem.DropShadow", new BOX_BRUSH("Starship/ContentBrowser/drop-shadow", FMargin(4.0f / 64.0f)));
+
 
 		Set("ReferenceViewer.PathText", FEditableTextBoxStyle(NormalEditableTextBoxStyle)
 			.SetFont(DEFAULT_FONT("Bold", 11)));
 
-		Set( "ContentBrowser.Sources", new IMAGE_BRUSH( "ContentBrowser/sources_16x", Icon16x16 ) );
-		Set( "ContentBrowser.PathDelimiter", new IMAGE_BRUSH( "Common/SmallArrowRight", Icon10x10 ) );
-		Set( "ContentBrowser.LockButton_Locked", new IMAGE_BRUSH( "Icons/padlock_locked_16x", Icon16x16 ) );
-		Set( "ContentBrowser.LockButton_Unlocked", new IMAGE_BRUSH( "Icons/padlock_unlocked_16x", Icon16x16 ) );
 		Set( "ContentBrowser.ShowSourcesView", new IMAGE_BRUSH( "ContentBrowser/sourcestoggle_16x_collapsed", Icon16x16 ) );
 		Set( "ContentBrowser.HideSourcesView", new IMAGE_BRUSH( "ContentBrowser/sourcestoggle_16x_expanded", Icon16x16 ) );
-		Set( "ContentBrowser.HistoryBack", new IMAGE_BRUSH( "Icons/assign_left_16x", Icon16x16) );
-		Set( "ContentBrowser.HistoryForward", new IMAGE_BRUSH("Icons/assign_right_16x", Icon16x16) );
+
 		Set( "ContentBrowser.DirectoryUp", new IMAGE_BRUSH("Icons/icon_folder_up_16x", Icon16x16) );
 		Set( "ContentBrowser.PathPickerButton", new IMAGE_BRUSH("Icons/ellipsis_12x", Icon12x12, FLinearColor::Black) );
 		Set( "ContentBrowser.SCC_CheckedOut", new IMAGE_BRUSH( "ContentBrowser/SCC_CheckedOut", Icon32x32) );
@@ -6166,6 +6167,7 @@ void FStarshipEditorStyle::FStyle::SetupContentBrowserStyle()
 		Set( "ContentBrowser.CollectionTreeDragDropBorder", new BOX_BRUSH( "Old/Window/ViewportDebugBorder", 0.8f ) );
 		Set( "ContentBrowser.PopupMessageIcon", new IMAGE_BRUSH( "Icons/alert", Icon32x32) );
 		Set( "ContentBrowser.NewFolderIcon", new IMAGE_BRUSH("Icons/icon_AddFolder_16x", Icon16x16 ) );
+
 		Set( "ContentBrowser.Local", new IMAGE_BRUSH( "ContentBrowser/Content_Local_12x", Icon12x12 ) );
 		Set( "ContentBrowser.Local.Small", new IMAGE_BRUSH( "ContentBrowser/Content_Local_16x", Icon16x16 ) );
 		Set( "ContentBrowser.Local.Large", new IMAGE_BRUSH( "ContentBrowser/Content_Local_64x", Icon64x64 ) );
diff --git a/Engine/Source/Editor/PropertyEditor/Private/UserInterface/PropertyTable/SPropertyTable.h b/Engine/Source/Editor/PropertyEditor/Private/UserInterface/PropertyTable/SPropertyTable.h
--- a/Engine/Source/Editor/PropertyEditor/Private/UserInterface/PropertyTable/SPropertyTable.h
+++ b/Engine/Source/Editor/PropertyEditor/Private/UserInterface/PropertyTable/SPropertyTable.h
@@ -89,7 +89,6 @@ public:
 					.BorderImage( FEditorStyle::GetBrush( "ToolPanel.GroupBorder" ) )
 					[
 						SAssignNew( BreadcrumbTrail, SBreadcrumbTrail< int32 > )
-						.DelimiterImage(FEditorStyle::GetBrush("ContentBrowser.PathDelimiter"))
 						.PersistentBreadcrumbs( true )
 						.OnCrumbClicked( this, &SPropertyTable::OnCrumbClicked )
 						.GetCrumbMenuContent( this, &SPropertyTable::GetCrumbMenuContent )
diff --git a/Engine/Source/Editor/UnrealEd/Private/AssetThumbnail.cpp b/Engine/Source/Editor/UnrealEd/Private/AssetThumbnail.cpp
--- a/Engine/Source/Editor/UnrealEd/Private/AssetThumbnail.cpp
+++ b/Engine/Source/Editor/UnrealEd/Private/AssetThumbnail.cpp
@@ -52,20 +52,22 @@ public:
 		, _HintColorAndOpacity(FLinearColor(0.0f, 0.0f, 0.0f, 0.0f))
 		, _ClassThumbnailBrushOverride(NAME_None)
 		, _AssetTypeColorOverride()
+		, _Padding(0)
 		{}
 
-		SLATE_ARGUMENT( FName, Style )
-		SLATE_ARGUMENT( TSharedPtr<FAssetThumbnail>, AssetThumbnail )
-		SLATE_ARGUMENT( TSharedPtr<FAssetThumbnailPool>, ThumbnailPool )
-		SLATE_ARGUMENT( bool, AllowFadeIn )
-		SLATE_ARGUMENT( bool, ForceGenericThumbnail )
-		SLATE_ARGUMENT( bool, AllowHintText )
-		SLATE_ARGUMENT( bool, AllowAssetSpecificThumbnailOverlay )
-		SLATE_ARGUMENT( EThumbnailLabel::Type, Label )
-		SLATE_ATTRIBUTE( FText, HighlightedText )
-		SLATE_ATTRIBUTE( FLinearColor, HintColorAndOpacity )
-		SLATE_ARGUMENT( FName, ClassThumbnailBrushOverride )
-		SLATE_ARGUMENT( TOptional<FLinearColor>, AssetTypeColorOverride )
+		SLATE_ARGUMENT(FName, Style)
+		SLATE_ARGUMENT(TSharedPtr<FAssetThumbnail>, AssetThumbnail)
+		SLATE_ARGUMENT(TSharedPtr<FAssetThumbnailPool>, ThumbnailPool)
+		SLATE_ARGUMENT(bool, AllowFadeIn)
+		SLATE_ARGUMENT(bool, ForceGenericThumbnail)
+		SLATE_ARGUMENT(bool, AllowHintText)
+		SLATE_ARGUMENT(bool, AllowAssetSpecificThumbnailOverlay)
+		SLATE_ARGUMENT(EThumbnailLabel::Type, Label)
+		SLATE_ATTRIBUTE(FText, HighlightedText)
+		SLATE_ATTRIBUTE(FLinearColor, HintColorAndOpacity)
+		SLATE_ARGUMENT(FName, ClassThumbnailBrushOverride)
+		SLATE_ARGUMENT(TOptional<FLinearColor>, AssetTypeColorOverride)
+		SLATE_ARGUMENT(FMargin, Padding)
 
 	SLATE_END_ARGS()
 
@@ -118,17 +120,16 @@ public:
 
 		// The generic representation of the thumbnail, for use before the rendered version, if it exists
 		OverlayWidget->AddSlot()
+		.Padding(InArgs._Padding)
 		[
 			SAssignNew(AssetBackgroundWidget, SBorder)
 			.BorderImage(GetAssetBackgroundBrush())
-			.BorderBackgroundColor(AssetColor.CopyWithNewOpacity(0.3f))
-			.Padding(GenericThumbnailBorderPadding)
+			.Padding(GenericThumbnailBorderPadding)		
 			.VAlign(VAlign_Center)
 			.HAlign(HAlign_Center)
 			.Visibility(this, &SAssetThumbnail::GetGenericThumbnailVisibility)
 			[
 				SNew(SOverlay)
-
 				+SOverlay::Slot()
 				[
 					SAssignNew(GenericLabelTextBlock, STextBlock)
@@ -178,9 +179,11 @@ public:
 			OverlayWidget->AddSlot()
 			[
 				SAssignNew(RenderedThumbnailWidget, SBorder)
-				.Padding(0)
-				.BorderImage(FEditorStyle::GetBrush("NoBrush"))
+				.Padding(InArgs._Padding)
+				.BorderImage(FStyleDefaults::GetNoBrush())
 				.ColorAndOpacity(this, &SAssetThumbnail::GetViewportColorAndOpacity)
+				.HAlign(HAlign_Center)
+				.VAlign(VAlign_Center)
 				[
 					Viewport.ToSharedRef()
 				]
@@ -192,7 +195,7 @@ public:
 			OverlayWidget->AddSlot()
 			.VAlign(VAlign_Bottom)
 			.HAlign(HAlign_Right)
-			.Padding(TAttribute<FMargin>(this, &SAssetThumbnail::GetClassIconPadding))
+			.Padding(GetClassIconPadding())
 			[
 				SAssignNew(ClassIconWidget, SBorder)
 				.BorderImage(FEditorStyle::GetNoBrush())
@@ -414,8 +417,7 @@ private:
 
 	float GetAssetColorStripHeight() const
 	{
-		// The strip is 2.5% the height of the thumbnail, but at least 3 units tall
-		return FMath::Max(FMath::CeilToFloat(WidthLastFrame*0.025f), 3.0f);
+		return 2.0f;
 	}
 
 	FMargin GetAssetColorStripPadding() const
@@ -812,17 +814,18 @@ TSharedRef<SWidget> FAssetThumbnail::MakeThumbnailWidget( const FAssetThumbnailC
 {
 	return
 		SNew(SAssetThumbnail)
-		.AssetThumbnail( SharedThis(this) )
-		.ThumbnailPool( ThumbnailPool.Pin() )
-		.AllowFadeIn( InConfig.bAllowFadeIn )
-		.ForceGenericThumbnail( InConfig.bForceGenericThumbnail )
-		.Label( InConfig.ThumbnailLabel )
-		.HighlightedText( InConfig.HighlightedText )
-		.HintColorAndOpacity( InConfig.HintColorAndOpacity )
-		.AllowHintText( InConfig.bAllowHintText )
-		.ClassThumbnailBrushOverride( InConfig.ClassThumbnailBrushOverride )
-		.AllowAssetSpecificThumbnailOverlay( InConfig.bAllowAssetSpecificThumbnailOverlay )
-		.AssetTypeColorOverride( InConfig.AssetTypeColorOverride );
+		.AssetThumbnail(SharedThis(this))
+		.ThumbnailPool(ThumbnailPool.Pin())
+		.AllowFadeIn(InConfig.bAllowFadeIn)
+		.ForceGenericThumbnail(InConfig.bForceGenericThumbnail)
+		.Label(InConfig.ThumbnailLabel)
+		.HighlightedText(InConfig.HighlightedText)
+		.HintColorAndOpacity(InConfig.HintColorAndOpacity)
+		.AllowHintText(InConfig.bAllowHintText)
+		.ClassThumbnailBrushOverride(InConfig.ClassThumbnailBrushOverride)
+		.AllowAssetSpecificThumbnailOverlay(InConfig.bAllowAssetSpecificThumbnailOverlay)
+		.AssetTypeColorOverride(InConfig.AssetTypeColorOverride)
+		.Padding(InConfig.Padding);
 }
 
 void FAssetThumbnail::RefreshThumbnail()
diff --git a/Engine/Source/Editor/UnrealEd/Private/ObjectTools.cpp b/Engine/Source/Editor/UnrealEd/Private/ObjectTools.cpp
--- a/Engine/Source/Editor/UnrealEd/Private/ObjectTools.cpp
+++ b/Engine/Source/Editor/UnrealEd/Private/ObjectTools.cpp
@@ -4481,7 +4481,6 @@ namespace ThumbnailTools
 		FPackageFileSummary FileSummary;
 		(*FileReader) << FileSummary;
 
-
 		// Make sure this is indeed a package
 		if( FileSummary.Tag != PACKAGE_FILE_TAG || FileReader->IsError() )
 		{
@@ -4497,6 +4496,11 @@ namespace ThumbnailTools
 			return false;
 		}
 
+		// Hacky way to reset thumbnails for UE5
+		if(FileSummary.SavedByEngineVersion.GetMajor() < 5)
+		{
+			return false;
+		}
 
 		// Seek the the part of the file where the thumbnail table lives
 		FileReader->Seek( FileSummary.ThumbnailTableOffset );
diff --git a/Engine/Source/Editor/UnrealEd/Private/Settings/SettingsClasses.cpp b/Engine/Source/Editor/UnrealEd/Private/Settings/SettingsClasses.cpp
--- a/Engine/Source/Editor/UnrealEd/Private/Settings/SettingsClasses.cpp
+++ b/Engine/Source/Editor/UnrealEd/Private/Settings/SettingsClasses.cpp
@@ -50,6 +50,7 @@ UContentBrowserSettings::FSettingChangedEvent UContentBrowserSettings::SettingCh
 UContentBrowserSettings::UContentBrowserSettings( const FObjectInitializer& ObjectInitializer )
 	: Super(ObjectInitializer)
 	, bShowFullCollectionNameInToolTip(true)
+	, DockCollections(true)
 {
 }
 
diff --git a/Engine/Source/Editor/UnrealEd/Public/AssetThumbnail.h b/Engine/Source/Editor/UnrealEd/Public/AssetThumbnail.h
--- a/Engine/Source/Editor/UnrealEd/Public/AssetThumbnail.h
+++ b/Engine/Source/Editor/UnrealEd/Public/AssetThumbnail.h
@@ -41,6 +41,7 @@ struct FAssetThumbnailConfig
 		, HighlightedText( FText::GetEmpty() )
 		, HintColorAndOpacity( FLinearColor( 0.0f, 0.0f, 0.0f, 0.0f ) )
 		, AssetTypeColorOverride()
+		, Padding(0)
 	{
 	}
 
@@ -53,6 +54,7 @@ struct FAssetThumbnailConfig
 	TAttribute< FText > HighlightedText;
 	TAttribute< FLinearColor > HintColorAndOpacity;
 	TOptional< FLinearColor > AssetTypeColorOverride;
+	FMargin Padding;
 };
 
 
diff --git a/Engine/Source/Editor/WorldBrowser/Private/Tiles/SWorldLayers.cpp b/Engine/Source/Editor/WorldBrowser/Private/Tiles/SWorldLayers.cpp
--- a/Engine/Source/Editor/WorldBrowser/Private/Tiles/SWorldLayers.cpp
+++ b/Engine/Source/Editor/WorldBrowser/Private/Tiles/SWorldLayers.cpp
@@ -247,7 +247,6 @@ void SWorldLayerButton::Construct(const FArguments& InArgs)
 						.Padding(3)
 						[
 							SNew(STextBlock)
-								.Font(FEditorStyle::GetFontStyle("ContentBrowser.FilterNameFont"))
 								.ShadowOffset(FVector2D(1.f, 1.f))
 								.Text(FText::FromString(WorldLayer.Name))
 						]
diff --git a/Engine/Source/Runtime/Slate/Private/Widgets/Input/SSearchBox.cpp b/Engine/Source/Runtime/Slate/Private/Widgets/Input/SSearchBox.cpp
--- a/Engine/Source/Runtime/Slate/Private/Widgets/Input/SSearchBox.cpp
+++ b/Engine/Source/Runtime/Slate/Private/Widgets/Input/SSearchBox.cpp
@@ -118,17 +118,16 @@ void SSearchBox::Construct( const FArguments& InArgs )
 
 	// Add a search glass image so that the user knows this text box is for searching
 	Box->InsertSlot(SlotIndex++)
-	.AutoWidth()
-	.Padding(InArgs._Style->ImagePadding)
-	.HAlign(HAlign_Center)
-	.VAlign(VAlign_Center)
-	[
-		SNew(SImage)
-		.Visibility(this, &SSearchBox::GetSearchGlassVisibility)
-		.Image( &InArgs._Style->GlassImage )
-		.ColorAndOpacity( FSlateColor::UseForeground() )
-	];
-
+		.AutoWidth()
+		.Padding(InArgs._Style->ImagePadding)
+		.HAlign(HAlign_Center)
+		.VAlign(VAlign_Center)
+		[
+			SNew(SImage)
+			.Visibility(this, &SSearchBox::GetSearchGlassVisibility)
+			.Image(&InArgs._Style->GlassImage)
+			.ColorAndOpacity(FAppStyle::Get().GetSlateColor("Colors.Foreground"))
+		];
 	// Add an X to clear the search whenever there is some text typed into it
 	Box->InsertSlot(SlotIndex++)
 	.AutoWidth()
@@ -138,12 +137,11 @@ void SSearchBox::Construct( const FArguments& InArgs )
 	[
 		SNew(SButton)
 		.Visibility(this, &SSearchBox::GetXVisibility)
-		.ButtonStyle( FCoreStyle::Get(), "NoBorder" )
+		.ButtonStyle(FAppStyle::Get(), "SimpleButton")
 		.ContentPadding(0)
 		.HAlign(HAlign_Center)
 		.VAlign(VAlign_Center)
-		.OnClicked( this, &SSearchBox::OnClearSearch )
-		.ForegroundColor( FSlateColor::UseForeground() )
+		.OnClicked(this, &SSearchBox::OnClearSearch)
 		// Allow the button to steal focus so that the search text will be automatically committed. Afterwards focus will be returned to the text box.
 		// If the user is keyboard-centric, they'll "ctrl+a, delete" to clear the search
 		.IsFocusable(true)
diff --git a/Engine/Source/Runtime/Slate/Private/Widgets/Layout/SExpandableArea.cpp b/Engine/Source/Runtime/Slate/Private/Widgets/Layout/SExpandableArea.cpp
--- a/Engine/Source/Runtime/Slate/Private/Widgets/Layout/SExpandableArea.cpp
+++ b/Engine/Source/Runtime/Slate/Private/Widgets/Layout/SExpandableArea.cpp
@@ -50,6 +50,7 @@ void SExpandableArea::Construct( const FArguments& InArgs )
 	const TAttribute<const FSlateBrush*> TitleBorderImage = !bBodyDiffers ? FStyleDefaults::GetNoBrush() : InArgs._BorderImage;
 	const TAttribute<FSlateColor> TitleBorderBackgroundColor = !bBodyDiffers ? FLinearColor::Transparent : InArgs._BorderBackgroundColor;
 
+	bAllowAnimatedTransition = InArgs._AllowAnimatedTransition;
 	ChildSlot
 	[
 		SNew( SBorder )
@@ -165,18 +166,6 @@ EVisibility SExpandableArea::OnGetContentVisibility() const
 	return Scale > 0 ? EVisibility::Visible : EVisibility::Collapsed;
 }
 
-/*
-FReply SExpandableArea::OnMouseDown( const FGeometry& MyGeometry, const FPointerEvent& MouseEvent )
-{
-	if ( MouseEvent.GetEffectingButton() == EKeys::LeftMouseButton )
-	{
-		//we need to capture the mouse for MouseUp events
-		return FReply::Handled().CaptureMouse( HeaderBorder.ToSharedRef() ).SetUserFocus( AsShared(), EFocusCause::Mouse );
-	}
-
-	return FReply::Unhandled();
-}*/
-
 
 FReply SExpandableArea::OnHeaderClicked()
 {
@@ -188,7 +177,7 @@ FReply SExpandableArea::OnHeaderClicked()
 
 void SExpandableArea::OnToggleContentVisibility()
 {
-	SetExpanded_Animated( !!bAreaCollapsed );
+	bAllowAnimatedTransition ? SetExpanded_Animated(!!bAreaCollapsed) : SetExpanded(!!bAreaCollapsed);
 }
 
 const FSlateBrush* SExpandableArea::OnGetCollapseImage() const
diff --git a/Engine/Source/Runtime/Slate/Public/Widgets/Input/SSearchBox.h b/Engine/Source/Runtime/Slate/Public/Widgets/Input/SSearchBox.h
--- a/Engine/Source/Runtime/Slate/Public/Widgets/Input/SSearchBox.h
+++ b/Engine/Source/Runtime/Slate/Public/Widgets/Input/SSearchBox.h
@@ -89,6 +89,8 @@ public:
 		/** Callback delegate to have first chance handling of the OnKeyDown event */
 		SLATE_EVENT(FOnKeyDown, OnKeyDownHandler)
 
+		/** Optional callback to allow for special functionality to be invoked when the search glass button is clicked */
+		SLATE_EVENT(FOnClicked, OnSearchGlassClicked)
 	SLATE_END_ARGS()
 
 	void Construct( const FArguments& InArgs );
diff --git a/Engine/Source/Runtime/Slate/Public/Widgets/Layout/SExpandableArea.h b/Engine/Source/Runtime/Slate/Public/Widgets/Layout/SExpandableArea.h
--- a/Engine/Source/Runtime/Slate/Public/Widgets/Layout/SExpandableArea.h
+++ b/Engine/Source/Runtime/Slate/Public/Widgets/Layout/SExpandableArea.h
@@ -30,7 +30,7 @@ public:
 		: _Style( &FCoreStyle::Get().GetWidgetStyle<FExpandableAreaStyle>("ExpandableArea") )
 		, _BorderBackgroundColor( FLinearColor::White )
 		, _BorderImage( FCoreStyle::Get().GetBrush("ExpandableArea.Border") )
-		, _BodyBorderBackgroundColor()
+		, _BodyBorderBackgroundColor(FLinearColor::White)
 		, _BodyBorderImage()
 		, _AreaTitle( )
 		, _InitiallyCollapsed( false )
@@ -41,6 +41,7 @@ public:
 		, _Padding( 1.0f )
 		, _AreaTitleFont( FCoreStyle::Get().GetFontStyle("ExpandableArea.TitleFont") )
 		, _HeaderCursor( )
+		, _AllowAnimatedTransition(true)
 		{}
 		
 		/** Style used to draw this area */
@@ -77,6 +78,8 @@ public:
 		SLATE_ATTRIBUTE( FSlateFontInfo, AreaTitleFont )
 		/** Override for Cursor, so you can specify a different cursor for the header */
 		SLATE_ATTRIBUTE( TOptional<EMouseCursor::Type>, HeaderCursor )
+		/** Whether or not the animated effect is used when opening or closing the area */
+		SLATE_ARGUMENT(bool, AllowAnimatedTransition)
 			
 	SLATE_END_ARGS()
 
@@ -167,4 +170,7 @@ protected:
 
 	/** Image to use when the area is expanded */
 	const FSlateBrush* ExpandedImage;
+
+	/** If true expanding and collapsing the area will animate, otherwise it will expand/collapse instantly */
+	bool bAllowAnimatedTransition;
 };
diff --git a/Engine/Source/Runtime/Slate/Public/Widgets/Navigation/SBreadcrumbTrail.h b/Engine/Source/Runtime/Slate/Public/Widgets/Navigation/SBreadcrumbTrail.h
--- a/Engine/Source/Runtime/Slate/Public/Widgets/Navigation/SBreadcrumbTrail.h
+++ b/Engine/Source/Runtime/Slate/Public/Widgets/Navigation/SBreadcrumbTrail.h
@@ -183,10 +183,11 @@ public:
 				.ButtonStyle(ButtonStyle)
 				.Visibility(this, &SBreadcrumbTrail::GetCrumbDelimiterVisibility, NextValidCrumbID)
 				.OnClicked(this, &SBreadcrumbTrail::OnCrumbDelimiterClicked, NextValidCrumbID)
-				.ContentPadding(FMargin(5, 0))
+				.ContentPadding(FMargin(3, 0))
 				[
 					SNew(SImage)
 					.Image(DelimiterImage)
+					.ColorAndOpacity(FSlateColor::UseForeground())
 				];
 		}
 		else
@@ -195,10 +196,11 @@ public:
 				.VAlign(EVerticalAlignment::VAlign_Center)
 				.ButtonStyle(ButtonStyle)
 				.Visibility( this, &SBreadcrumbTrail::GetDelimiterVisibility, NextValidCrumbID )
-				.ContentPadding( FMargin(5, 0) )
+				.ContentPadding( FMargin(3, 0) )
 				[
 					SNew(SImage)
 					.Image(DelimiterImage)
+					.ColorAndOpacity(FSlateColor::UseForeground())
 				];
 		}
 
@@ -400,7 +402,7 @@ private:
 				}
 			}
 
-			static const FName DefaultForegroundName;
+			static const FName DefaultForegroundName("DefaultForeground");
 			return FCoreStyle::Get().GetSlateColor( DefaultForegroundName );
 		}
 		else
diff --git a/Engine/Source/Runtime/SlateCore/Private/Styling/StarshipCoreStyle.cpp b/Engine/Source/Runtime/SlateCore/Private/Styling/StarshipCoreStyle.cpp
--- a/Engine/Source/Runtime/SlateCore/Private/Styling/StarshipCoreStyle.cpp
+++ b/Engine/Source/Runtime/SlateCore/Private/Styling/StarshipCoreStyle.cpp
@@ -318,14 +318,17 @@ TSharedRef<ISlateStyle> FStarshipCoreStyle::Create()
 
 		Style->Set("TrashCan", new IMAGE_BRUSH("Icons/TrashCan", FVector2D(64, 64)));
 		Style->Set("TrashCan_Small", new IMAGE_BRUSH("Icons/TrashCan_Small", FVector2D(18, 18)));
-		Style->Set("Icons.Cross", new IMAGE_BRUSH("Icons/Cross_12x", Icon12x12));
+
 		Style->Set("Icons.Denied", new IMAGE_BRUSH("Icons/denied_16x", Icon16x16));
-		Style->Set("Icons.Error", new IMAGE_BRUSH("Icons/icon_error_16x", Icon16x16));
+	
 		Style->Set("Icons.Help", new IMAGE_BRUSH("Icons/icon_help_16x", Icon16x16));
 		Style->Set("Icons.Info", new IMAGE_BRUSH("Icons/icon_info_16x", Icon16x16));
-		Style->Set("Icons.Warning", new IMAGE_BRUSH("Icons/icon_warning_16x", Icon16x16));
+	
 		Style->Set("Icons.Download", new IMAGE_BRUSH("Icons/icon_Downloads_16x", Icon16x16));
 
+		Style->Set("Icons.Error", new IMAGE_BRUSH_SVG("Starship/Common/alert-circle", Icon16x16));
+		Style->Set("Icons.Warning", new IMAGE_BRUSH_SVG("Starship/Common/alert-triangle", Icon16x16));
+
 		Style->Set("Icons.box-perspective", new IMAGE_BRUSH_SVG("Starship/Common/box-perspective", Icon16x16));
 		Style->Set("Icons.cylinder", new IMAGE_BRUSH_SVG("Starship/Common/cylinder", Icon16x16));
 		Style->Set("Icons.pyramid", new IMAGE_BRUSH_SVG("Starship/Common/pyriamid", Icon16x16));
@@ -333,7 +336,9 @@ TSharedRef<ISlateStyle> FStarshipCoreStyle::Create()
 
 		Style->Set("Icons.Settings", new IMAGE_BRUSH_SVG("Starship/Common/settings", Icon16x16));
 		Style->Set("Icons.Blueprints", new IMAGE_BRUSH_SVG("Starship/Common/blueprint", Icon16x16));
+		Style->Set("Icons.Cross", new IMAGE_BRUSH_SVG("Starship/Common/plus", Icon16x16));
 		Style->Set("Icons.Plus", new IMAGE_BRUSH_SVG("Starship/Common/plus", Icon16x16));
+		Style->Set("Icons.PlusCircle", new IMAGE_BRUSH_SVG("Starship/Common/plus-circle", Icon16x16));
 
 		Style->Set("Icons.Save", new IMAGE_BRUSH_SVG("Starship/Common/save", Icon16x16));
 
@@ -350,6 +355,20 @@ TSharedRef<ISlateStyle> FStarshipCoreStyle::Create()
 		Style->Set("Icons.CircleArrowDown", new IMAGE_BRUSH_SVG("Starship/Common/circle-arrow-down", Icon16x16));
 
 		Style->Set("Icons.Check", new IMAGE_BRUSH_SVG("Starship/Common/check", Icon16x16));
+
+		Style->Set("Icons.FolderOpen", new IMAGE_BRUSH_SVG("Starship/Common/folder-open", Icon16x16));
+		Style->Set("Icons.FolderClosed", new IMAGE_BRUSH_SVG("Starship/Common/folder-closed", Icon16x16));
+
+		Style->Set("Icons.ChevronLeft", new IMAGE_BRUSH_SVG("Starship/Common/chevron-left", Icon16x16));
+		Style->Set("Icons.ChevronRight", new IMAGE_BRUSH_SVG("Starship/Common/chevron-right", Icon16x16));
+
+		Style->Set("Icons.ChevronUp", new IMAGE_BRUSH_SVG("Starship/Common/chevron-up", Icon16x16));
+		Style->Set("Icons.ChevronDown", new IMAGE_BRUSH_SVG("Starship/Common/chevron-down", Icon16x16));
+
+		Style->Set("Icons.Search", new IMAGE_BRUSH_SVG("Starship/Common/search", Icon16x16));
+
+		Style->Set("Icons.FilledCircle", new IMAGE_BRUSH_SVG("Starship/Common/filled-circle", Icon16x16));
+
 	}
 
 	// Tool panels
@@ -457,8 +476,8 @@ TSharedRef<ISlateStyle> FStarshipCoreStyle::Create()
 	// SExpandableArea defaults...
 	{
 		Style->Set("ExpandableArea", FExpandableAreaStyle()
-			.SetCollapsedImage(IMAGE_BRUSH("Common/TreeArrow_Collapsed", Icon10x10, DefaultForeground))
-			.SetExpandedImage(IMAGE_BRUSH("Common/TreeArrow_Expanded", Icon10x10, DefaultForeground))
+			.SetCollapsedImage(IMAGE_BRUSH_SVG("Starship/Common/chevron-right", Icon16x16, DefaultForeground))
+			.SetExpandedImage(IMAGE_BRUSH_SVG("Starship/Common/chevron-down", Icon16x16, DefaultForeground))
 		);
 		Style->Set("ExpandableArea.TitleFont", StyleFonts.SmallBold);
 		Style->Set("ExpandableArea.Border", new FSlateRoundedBoxBrush(FStyleColors::Background, 4) );
@@ -982,9 +1001,9 @@ void FStarshipCoreStyle::SetupTextStyles(TSharedRef<FStyle>& Style)
 			.SetTextBoxStyle(NormalEditableTextBoxStyle)
 			.SetUpArrowImage(IMAGE_BRUSH_SVG("Starship/Common/arrow-north", Icon8x8, FStyleColors::Foreground))
 			.SetDownArrowImage(IMAGE_BRUSH_SVG("Starship/Common/arrow-south", Icon8x8, FStyleColors::Foreground))
-			.SetGlassImage(IMAGE_BRUSH_SVG("Starship/Common/search", Icon16x16, FStyleColors::Foreground))
-			.SetClearImage(IMAGE_BRUSH_SVG("Starship/Common/close", Icon16x16, FStyleColors::Foreground))
-			.SetImagePadding(FMargin(3.f, 2.f, -2.f, 2.0))
+			.SetGlassImage(IMAGE_BRUSH_SVG("Starship/Common/search", Icon16x16))
+			.SetClearImage(IMAGE_BRUSH_SVG("Starship/Common/close", Icon16x16))
+			.SetImagePadding(FMargin(3.f, 0.f, -2.f, 0.0))
 			.SetLeftAlignButtons(true)
 		);
 	}
diff --git a/Engine/Source/Runtime/SlateCore/Private/Styling/StyleColors.cpp b/Engine/Source/Runtime/SlateCore/Private/Styling/StyleColors.cpp
--- a/Engine/Source/Runtime/SlateCore/Private/Styling/StyleColors.cpp
+++ b/Engine/Source/Runtime/SlateCore/Private/Styling/StyleColors.cpp
@@ -2,6 +2,7 @@
 
 #include "Styling/StyleColors.h"
 
+const FSlateColor FStyleColors::Transparent = FSlateColor(FLinearColor::Transparent);
 
 const FSlateColor FStyleColors::Black = EStyleColor::Black;
 const FSlateColor FStyleColors::Title = EStyleColor::Title;
diff --git a/Engine/Source/Runtime/SlateCore/Public/Styling/StyleColors.h b/Engine/Source/Runtime/SlateCore/Public/Styling/StyleColors.h
--- a/Engine/Source/Runtime/SlateCore/Public/Styling/StyleColors.h
+++ b/Engine/Source/Runtime/SlateCore/Public/Styling/StyleColors.h
@@ -143,6 +143,7 @@ private:
  */
 struct SLATECORE_API FStyleColors
 {
+	static const FSlateColor Transparent;
 	static const FSlateColor Black;
 	static const FSlateColor Title;
 	static const FSlateColor Foldout;
