Editgrate CL 13760226

Fixed broken text filter in Content Browser path view

#rnx

diff --git a/Engine/Source/Editor/ContentBrowser/Private/SPathView.cpp b/Engine/Source/Editor/ContentBrowser/Private/SPathView.cpp
--- a/Engine/Source/Editor/ContentBrowser/Private/SPathView.cpp
+++ b/Engine/Source/Editor/ContentBrowser/Private/SPathView.cpp
@@ -1103,7 +1103,7 @@ void SPathView::Populate()
 			{
 				// Use the whole path so we deliberately include any children of matched parents in the filtered list
 				const FString PathStr = InItemData.GetVirtualPath().ToString();
-				bPassesFilter |= SearchBoxFolderFilter->PassesFilter(PathStr);
+				bPassesFilter &= SearchBoxFolderFilter->PassesFilter(PathStr);
 			}
 
 			if (bPassesFilter)
