Updated styling and workflow for expandable search areas in the content browser paths and collections views.

diff --git a/Engine/Source/Editor/ContentBrowser/Private/SContentBrowser.cpp b/Engine/Source/Editor/ContentBrowser/Private/SContentBrowser.cpp
--- a/Engine/Source/Editor/ContentBrowser/Private/SContentBrowser.cpp
+++ b/Engine/Source/Editor/ContentBrowser/Private/SContentBrowser.cpp
@@ -272,6 +272,7 @@ void SContentBrowser::Construct( const FArguments& InArgs, const FName& InInstan
 						.ToolTipText( LOCTEXT( "ImportTooltip", "Import assets from files to the currently selected folder" ) )
 						.ContentPadding(2)
 						.OnClicked(this, &SContentBrowser::OnImportClicked)
+						.IsEnabled(this, &SContentBrowser::IsAddNewEnabled)
 						[
 							SNew( SHorizontalBox )
 							+ SHorizontalBox::Slot()
@@ -886,7 +887,7 @@ TSharedRef<SWidget> SContentBrowser::CreatePathView(const FContentBrowserConfig*
 			.HAlign(HAlign_Right)
 			.Padding(5.0f, 0.0f)
 			[
-				SNew(SExpandableSearchArea, SourcesSearch->GetWidget())
+				SAssignNew(PathSearchArea, SExpandableSearchArea, SourcesSearch->GetWidget())
 			]
 		]
 		.BodyContent()
@@ -947,7 +948,7 @@ TSharedRef<SWidget> SContentBrowser::CreateDockedCollectionsView(const FContentB
 			.HAlign(HAlign_Right)
 			.Padding(5.0f, 0.0f)
 			[
-				SNew(SExpandableSearchArea, CollectionSearch->GetWidget())
+				SAssignNew(CollectionSearchArea, SExpandableSearchArea, CollectionSearch->GetWidget())
 			]
 		]
 		.BodyContent()
@@ -1254,6 +1255,9 @@ void SContentBrowser::SaveSettings() const
 	GConfig->SetBool(*SettingsIniSection, *(SettingsString + TEXT(".PathAreaExpanded")), PathArea->IsExpanded(), GEditorPerProjectIni);
 	GConfig->SetBool(*SettingsIniSection, *(SettingsString + TEXT(".CollectionAreaExpanded")), CollectionArea->IsExpanded(), GEditorPerProjectIni);
 
+	GConfig->SetBool(*SettingsIniSection, *(SettingsString + TEXT(".PathSearchAreaExpanded")), PathSearchArea->IsExpanded(), GEditorPerProjectIni);
+	GConfig->SetBool(*SettingsIniSection, *(SettingsString + TEXT(".CollectionSearchAreaExpanded")), CollectionSearchArea->IsExpanded(), GEditorPerProjectIni);
+
 	for(int32 SlotIndex = 0; SlotIndex < PathAssetSplitterPtr->GetChildren()->Num(); SlotIndex++)
 	{
 		float SplitterSize = PathAssetSplitterPtr->SlotAt(SlotIndex).SizeValue.Get();
@@ -1400,6 +1404,14 @@ void SContentBrowser::LoadSettings(const FName& InInstanceName)
 	GConfig->GetBool(*SettingsIniSection, *(SettingsString + TEXT(".CollectionAreaExpanded")), bCollectionAreaExpanded, GEditorPerProjectIni);
 	CollectionArea->SetExpanded(bCollectionAreaExpanded);
 
+	bool bPathSearchAreaExpanded = false;
+	GConfig->GetBool(*SettingsIniSection, *(SettingsString + TEXT(".PathSearchAreaExpanded")), bPathSearchAreaExpanded, GEditorPerProjectIni);
+	PathSearchArea->SetExpanded(bPathSearchAreaExpanded);
+
+	bool bCollectionSearchAreaExpanded = false;
+	GConfig->GetBool(*SettingsIniSection, *(SettingsString + TEXT(".CollectionSearchAreaExpanded")), bCollectionSearchAreaExpanded, GEditorPerProjectIni);
+	CollectionSearchArea->SetExpanded(bCollectionSearchAreaExpanded);
+
 	for(int32 SlotIndex = 0; SlotIndex < PathAssetSplitterPtr->GetChildren()->Num(); SlotIndex++)
 	{
 		float SplitterSize = PathAssetSplitterPtr->SlotAt(SlotIndex).SizeValue.Get();
diff --git a/Engine/Source/Editor/ContentBrowser/Private/SContentBrowser.h b/Engine/Source/Editor/ContentBrowser/Private/SContentBrowser.h
--- a/Engine/Source/Editor/ContentBrowser/Private/SContentBrowser.h
+++ b/Engine/Source/Editor/ContentBrowser/Private/SContentBrowser.h
@@ -30,6 +30,7 @@ class SFilterList;
 class SPathView;
 class SSplitter;
 class UFactory;
+class SExpandableSearchArea;
 
 struct FToolMenuContext;
 
@@ -565,13 +566,19 @@ private:
 
 	/** Favorites area widget */
 	TSharedPtr<SExpandableArea> FavoritesArea;
-	
+
 	/** Path area widget */
 	TSharedPtr<SExpandableArea> PathArea;
 
+	/** Search area for paths */
+	TSharedPtr<SExpandableSearchArea> PathSearchArea;
+
 	/** Collection area widget */
 	TSharedPtr<SExpandableArea> CollectionArea;
 
+	/** Search area for collections */
+	TSharedPtr<SExpandableSearchArea> CollectionSearchArea;
+
 	/** Index of the active sources widget */
 	int32 ActiveSourcesWidgetIndex = 0;
 
diff --git a/Engine/Source/Editor/ContentBrowser/Private/SExpandableSearchArea.cpp b/Engine/Source/Editor/ContentBrowser/Private/SExpandableSearchArea.cpp
--- a/Engine/Source/Editor/ContentBrowser/Private/SExpandableSearchArea.cpp
+++ b/Engine/Source/Editor/ContentBrowser/Private/SExpandableSearchArea.cpp
@@ -13,24 +13,14 @@ void SExpandableSearchArea::Construct(const FArguments& InArgs, TSharedRef<SSear
 
 	SearchStyle = InArgs._Style;
 
+	SearchBoxPtr = SearchBox;
+
 	ChildSlot
 	[
 		SNew(SHorizontalBox)
 		+ SHorizontalBox::Slot()
 		.HAlign(HAlign_Center)
 		.AutoWidth()
-		[
-			SNew(SButton)
-			.ButtonStyle(FAppStyle::Get(), "SimpleButton")
-			.ContentPadding(FMargin(0.0f, 2.0f))
-			.OnClicked(this, &SExpandableSearchArea::OnExpandSearchClicked)
-			[
-				SNew(SImage)
-				.Image(&SearchStyle->GlassImage)
-				.ColorAndOpacity(FSlateColor::UseForeground())
-			]
-		]
-		+ SHorizontalBox::Slot()
 		[
 			SNew(SBox)
 			.Visibility(this, &SExpandableSearchArea::GetSearchBoxVisibility)
@@ -40,12 +30,51 @@ void SExpandableSearchArea::Construct(const FArguments& InArgs, TSharedRef<SSear
 				SearchBox
 			]
 		]
+		+ SHorizontalBox::Slot()
+		[
+			SNew(SButton)
+			.ToolTipText(NSLOCTEXT("ExpandableSearchArea", "ExpandCollapseSearchButton", "Expands or collapses the search text box"))
+			.ButtonStyle(FAppStyle::Get(), "SimpleButton")
+			.ContentPadding(FMargin(0.0f, 2.0f))
+			.OnClicked(this, &SExpandableSearchArea::OnExpandSearchClicked)
+			[
+				SNew(SHorizontalBox)
+				+ SHorizontalBox::Slot()
+				.VAlign(VAlign_Center)
+				.AutoWidth()
+				[
+					SNew(SImage)
+					.Image(this, &SExpandableSearchArea::GetExpandSearchImage)
+					.ColorAndOpacity(FSlateColor::UseForeground())
+				]
+				+ SHorizontalBox::Slot()
+				.VAlign(VAlign_Center)
+				.AutoWidth()
+				[
+					SNew(SImage)
+					.Image(&SearchStyle->GlassImage)
+					.ColorAndOpacity(FSlateColor::UseForeground())
+					.Visibility(this, &SExpandableSearchArea::GetSearchGlassVisibility)
+				]
+			]
+		]
 	];
 }
 
+void SExpandableSearchArea::SetExpanded(bool bInExpanded)
+{
+	bIsExpanded = bInExpanded;
+}
+
 FReply SExpandableSearchArea::OnExpandSearchClicked()
 {
-	bIsExpanded = !bIsExpanded;
+	if(TSharedPtr<SSearchBox> SearchBox = SearchBoxPtr.Pin())
+	{
+		bIsExpanded = !bIsExpanded;
+
+		return FReply::Handled().SetUserFocus(SearchBox.ToSharedRef(), EFocusCause::SetDirectly);
+	}
+
 	return FReply::Handled();
 }
 
@@ -54,7 +83,16 @@ EVisibility SExpandableSearchArea::GetSearchBoxVisibility() const
 	return bIsExpanded ? EVisibility::Visible : EVisibility::Collapsed;
 }
 
+EVisibility SExpandableSearchArea::GetSearchGlassVisibility() const
+{
+	return bIsExpanded ? EVisibility::Collapsed : EVisibility::Visible;
+}
+
+
 const FSlateBrush* SExpandableSearchArea::GetExpandSearchImage() const
 {
-	return &SearchStyle->GlassImage;
+	static const FName RightIcon("Icons.ChevronRight");
+	static const FName LeftIcon("Icons.ChevronLeft");
+
+	return bIsExpanded ? FAppStyle::Get().GetBrush(RightIcon) : FAppStyle::Get().GetBrush(LeftIcon);
 }
diff --git a/Engine/Source/Editor/ContentBrowser/Private/SExpandableSearchArea.h b/Engine/Source/Editor/ContentBrowser/Private/SExpandableSearchArea.h
--- a/Engine/Source/Editor/ContentBrowser/Private/SExpandableSearchArea.h
+++ b/Engine/Source/Editor/ContentBrowser/Private/SExpandableSearchArea.h
@@ -23,12 +23,19 @@ public:
 
 	void Construct(const FArguments& InArgs, TSharedRef<SSearchBox> SearchBox);
 
+	/** @return true if the search area is expanded and the search box exposed */
+	bool IsExpanded() const { return bIsExpanded; }
+
+	/** Sets whether or not the search area is expanded to expose the search box */
+	void SetExpanded(bool bInExpanded);
 private:
 
 	FReply OnExpandSearchClicked();
 	EVisibility GetSearchBoxVisibility() const;
+	EVisibility GetSearchGlassVisibility() const;
 	const FSlateBrush* GetExpandSearchImage() const;
 private:
 	const FSearchBoxStyle* SearchStyle;
 	bool bIsExpanded;
+	TWeakPtr<SSearchBox> SearchBoxPtr;
 };
