Content browser reskin reworks

Changed the look of folders
Changed the look of filters
Changed the look of the text for asset tiles

diff --git a/Engine/Source/Editor/ContentBrowser/Private/AssetViewWidgets.cpp b/Engine/Source/Editor/ContentBrowser/Private/AssetViewWidgets.cpp
--- a/Engine/Source/Editor/ContentBrowser/Private/AssetViewWidgets.cpp
+++ b/Engine/Source/Editor/ContentBrowser/Private/AssetViewWidgets.cpp
@@ -115,9 +115,14 @@ TSharedRef<SWidget> FAssetViewItemHelper::CreateListTileItemContents(T* const In
 		// Folder base
 		ItemContentsOverlay->AddSlot()
 		[
-			SNew(SImage)
-			.Image(FolderBaseImage)
-			.ColorAndOpacity(InTileOrListItem, &T::GetAssetColor)
+			SNew(SBorder)
+			.BorderImage(FAppStyle::Get().GetBrush("ContentBrowser.FolderItem.DropShadow"))
+			.Padding(FMargin(0,0,4.0f,4.0f))
+			[
+				SNew(SImage)
+				.Image(FolderBaseImage)
+				.ColorAndOpacity(InTileOrListItem, &T::GetAssetColor)
+			]
 		];
 
 		if (bCollectionFolder)
@@ -1632,7 +1637,7 @@ void SAssetTileItem::Construct( const FArguments& InArgs )
 		// Drop shadow border
 		SNew(SBorder)
 		.Padding(FMargin(0.0f, 0.0f, 5.0f, 5.0f))
-		.BorderImage(FAppStyle::Get().GetBrush(ItemShadowBorderName))
+		.BorderImage(IsFolder() ? FStyleDefaults::GetNoBrush() : FAppStyle::Get().GetBrush(ItemShadowBorderName))
 		[
 			SNew(SOverlay)
 			.AddMetaData<FTagMetaData>(FTagMetaData(AssetItem->GetItem().GetVirtualPath()))
@@ -1640,7 +1645,7 @@ void SAssetTileItem::Construct( const FArguments& InArgs )
 			[
 				SNew(SBorder)
 				.Padding(0)
-				.BorderImage(FAppStyle::Get().GetBrush("ContentBrowser.AssetTileItem.ThumbnailAreaBackground"))
+				.BorderImage(IsFolder() ? FStyleDefaults::GetNoBrush() : FAppStyle::Get().GetBrush("ContentBrowser.AssetTileItem.ThumbnailAreaBackground"))
 				[
 					SNew(SVerticalBox)
 					// Thumbnail
@@ -1663,24 +1668,13 @@ void SAssetTileItem::Construct( const FArguments& InArgs )
 					[
 						SNew(SBorder)
 						.Padding(FMargin(2.0f, 3.0f))
-						.BorderImage(FAppStyle::Get().GetBrush("ContentBrowser.AssetTileItem.NameAreaBackground"))
+						.BorderImage(IsFolder() ? FStyleDefaults::GetNoBrush() : FAppStyle::Get().GetBrush("ContentBrowser.AssetTileItem.NameAreaBackground"))
 						[
 							SNew(SVerticalBox)
 							+ SVerticalBox::Slot()
+							.Padding(2.0f)
 							.HAlign(HAlign_Left)
 							.VAlign(VAlign_Top)
-							.Padding(2.0f)
-							.AutoHeight()
-							[
-								SNew(STextBlock)
-								.Visibility(this, &SAssetTileItem::GetAssetClassLabelVisibility)
-								.TextStyle(FAppStyle::Get(), "ButtonText")
-								.TransformPolicy(ETextTransformPolicy::ToUpper)
-								.Text(this, &SAssetTileItem::GetAssetClassText)
-							]
-							+ SVerticalBox::Slot()
-							.Padding(2.0f)
-							.VAlign(VAlign_Center)
 							[
 								SAssignNew(InlineRenameWidget, SInlineEditableTextBlock)
 								.Font(this, &SAssetTileItem::GetThumbnailFont)
@@ -1693,9 +1687,22 @@ void SAssetTileItem::Construct( const FArguments& InArgs )
 								.IsReadOnly(this, &SAssetTileItem::IsNameReadOnly)
 								.Justification(ETextJustify::Left)
 								.LineBreakPolicy(FBreakIterator::CreateCamelCaseBreakIterator())
-
-
 							]
+							+ SVerticalBox::Slot()
+							.HAlign(HAlign_Left)
+							.VAlign(VAlign_Bottom)
+							.Padding(2.0f)
+							.AutoHeight()
+							[
+								SNew(STextBlock)
+								.Visibility(this, &SAssetTileItem::GetAssetClassLabelVisibility)
+								.TextStyle(FAppStyle::Get(), "SmallText")
+								//.Font(FAppStyle::Get().GetFontStyle("ContentBrowser.AssetTileViewNameFontSmall"))
+								.TransformPolicy(ETextTransformPolicy::ToUpper)
+								.Text(this, &SAssetTileItem::GetAssetClassText)
+								.ColorAndOpacity(FSlateColor::UseSubduedForeground())
+							]
+							
 						]
 					]
 				]
@@ -1783,11 +1790,14 @@ FOptionalSize SAssetTileItem::GetThumbnailBoxSize() const
 
 EVisibility SAssetTileItem::GetAssetClassLabelVisibility() const
 {
-	FOptionalSize ThumbSize = GetThumbnailBoxSize();
-	if (ThumbSize.IsSet())
+	if(!IsFolder())
 	{
-		float Size = ThumbSize.Get() *GetTickSpaceGeometry().Scale;
-		return Size < 100 ? EVisibility::Collapsed : EVisibility::Visible;
+		FOptionalSize ThumbSize = GetThumbnailBoxSize();
+		if (ThumbSize.IsSet())
+		{
+			float Size = ThumbSize.Get() * GetTickSpaceGeometry().Scale;
+			return Size < 100 ? EVisibility::Collapsed : EVisibility::Visible;
+		}
 	}
 	return EVisibility::Collapsed;
 }
diff --git a/Engine/Source/Editor/ContentBrowser/Private/SAssetView.cpp b/Engine/Source/Editor/ContentBrowser/Private/SAssetView.cpp
--- a/Engine/Source/Editor/ContentBrowser/Private/SAssetView.cpp
+++ b/Engine/Source/Editor/ContentBrowser/Private/SAssetView.cpp
@@ -171,7 +171,7 @@ void SAssetView::Construct( const FArguments& InArgs )
 	TileViewThumbnailSize = 150;
 	TileViewThumbnailPadding = 9;
 
-	TileViewNameHeight = 115;
+	TileViewNameHeight = 100;
 
 	MinThumbnailScale = 0.2f * ThumbnailScaleRangeScalar;
 	MaxThumbnailScale = 1.9f * ThumbnailScaleRangeScalar;
@@ -310,7 +310,7 @@ void SAssetView::Construct( const FArguments& InArgs )
 			[
 				// Container for the view types
 				SAssignNew(ViewContainer, SBox)
-				.Padding(0)
+				.Padding(FMargin(5.0f,0,0,0))
 			]
 
 			+ SOverlay::Slot()
diff --git a/Engine/Source/Editor/ContentBrowser/Private/SFilterList.cpp b/Engine/Source/Editor/ContentBrowser/Private/SFilterList.cpp
--- a/Engine/Source/Editor/ContentBrowser/Private/SFilterList.cpp
+++ b/Engine/Source/Editor/ContentBrowser/Private/SFilterList.cpp
@@ -183,9 +183,9 @@ public:
 		ChildSlot
 		[
 			SNew(SBorder)
-			.Padding(0)
-			.BorderImage(FAppStyle::Get().GetBrush("ContentBrowser.FilterBackground"))
-			[
+ 			.Padding(0)
+			.BorderImage(this, &SFilter::GetFilterBorderBackground)
+ 			[
 				SAssignNew( ToggleButtonPtr, SFilterCheckBox )
 				.Style(FEditorStyle::Get(), "ContentBrowser.FilterButton")
 				.ToolTipText(FilterToolTip)
@@ -210,6 +210,7 @@ public:
 					[
 						SNew(STextBlock)
 						.Text(this, &SFilter::GetFilterName)
+						//.ColorAndOpacity(this, &SFilter::GetFilterTextColorAndOpacity);
 					]
 				]
 			]
@@ -420,7 +421,15 @@ private:
 	/** Handler to determine the color of the checkbox when it is checked */
 	FSlateColor GetFilterImageColorAndOpacity() const
 	{
-		return IsChecked() == ECheckBoxState::Checked ? FilterColor : FLinearColor::Black;
+		return bEnabled ? FilterColor : FilterColor.Desaturate(.4f);
+	}
+
+	const FSlateBrush* GetFilterBorderBackground() const
+	{
+		static const FName Unchecked("ContentBrowser.FilterBackgroundUnchecked");
+		static const FName Checked("ContentBrowser.FilterBackgroundChecked");
+
+		return FAppStyle::Get().GetBrush(bEnabled ? Checked : Unchecked);
 	}
 
 	/** Handler to determine the padding of the checkbox text when it is pressed */
diff --git a/Engine/Source/Editor/EditorStyle/Private/StarshipStyle.cpp b/Engine/Source/Editor/EditorStyle/Private/StarshipStyle.cpp
--- a/Engine/Source/Editor/EditorStyle/Private/StarshipStyle.cpp
+++ b/Engine/Source/Editor/EditorStyle/Private/StarshipStyle.cpp
@@ -5998,13 +5998,14 @@ void FStarshipEditorStyle::FStyle::SetupContentBrowserStyle()
 			.SetCheckedPressedImage(FSlateNoResource())
 			.SetForegroundColor(FStyleColors::Foreground)
 			.SetHoveredForegroundColor(FStyleColors::ForegroundHover)
-			.SetCheckedForegroundColor(FStyleColors::Foreground)
-			.SetCheckedHoveredForegroundColor(FStyleColors::ForegroundHover)
-			.SetCheckedPressedForegroundColor(FStyleColors::ForegroundHover)
+			.SetCheckedForegroundColor(FStyleColors::Primary)
+			.SetCheckedHoveredForegroundColor(FStyleColors::PrimaryHover)
+			.SetCheckedPressedForegroundColor(FStyleColors::PrimaryHover)
 			.SetPressedForegroundColor(FStyleColors::ForegroundHover);
 
 		Set("ContentBrowser.FilterImage", new CORE_IMAGE_BRUSH_SVG("Starship/Common/filled-circle", Icon8x8));
-		Set("ContentBrowser.FilterBackground", new FSlateRoundedBoxBrush(FStyleColors::Header, 2.0f));
+		Set("ContentBrowser.FilterBackgroundUnchecked", new FSlateRoundedBoxBrush(FStyleColors::Dropdown, 4.0f));
+		Set("ContentBrowser.FilterBackgroundChecked", new FSlateRoundedBoxBrush(FStyleColors::Input, 4.0f));
 		/* ... and add the new style */
 		Set("ContentBrowser.FilterButton", ContentBrowserFilterButtonCheckBoxStyle );
 
@@ -6124,20 +6125,17 @@ void FStarshipEditorStyle::FStyle::SetupContentBrowserStyle()
 		Set( "ContentBrowser.ColumnViewFolderIcon", new CORE_IMAGE_BRUSH_SVG( "Starship/Common/folder-closed", Icon16x16 ) );
 		Set( "ContentBrowser.ColumnViewDeveloperFolderIcon", new IMAGE_BRUSH( "Icons/FolderDeveloper", FVector2D(18, 16) ) );
 
-		Set( "ContentBrowser.ListViewFolderIcon", new CORE_IMAGE_BRUSH_SVG( "Starship/Common/folder-closed", FVector2D(256, 256) ) );
+		Set( "ContentBrowser.ListViewFolderIcon", new IMAGE_BRUSH_SVG( "Starship/ContentBrowser/folder", FVector2D(64, 64) ) );
 		Set( "ContentBrowser.ListViewDeveloperFolderIcon", new IMAGE_BRUSH( "Icons/Folders/FolderDev_Base_256x", FVector2D(256, 256) ) );
 
-		Set( "ContentBrowser.TileViewFolderIcon", new CORE_IMAGE_BRUSH_SVG("Starship/Common/folder-closed", FVector2D(512, 512) ) );
-		Set( "ContentBrowser.TileViewDeveloperFolderIcon", new IMAGE_BRUSH( "Icons/Folders/FolderDev_Base_512x", FVector2D(512, 512) ) );
-
-
 		Set("ContentBrowser.AssetTileItem.ThumbnailAreaBackground", new FSlateRoundedBoxBrush(FStyleColors::Input, 4.0f));
 		Set("ContentBrowser.AssetTileItem.NameAreaBackground", new FSlateRoundedBoxBrush(FStyleColors::Dropdown, 4.0f));
 		Set("ContentBrowser.AssetTileItem.SelectedBorder", new FSlateRoundedBoxBrush(FStyleColors::Transparent, 4.0f, FStyleColors::Primary, 2.0f));
 		Set("ContentBrowser.AssetTileItem.SelectedHoverBorder", new FSlateRoundedBoxBrush(FStyleColors::Transparent, 4.0f, FStyleColors::PrimaryHover, 2.0f));
 		Set("ContentBrowser.AssetTileItem.HoverBorder", new FSlateRoundedBoxBrush(FStyleColors::Transparent, 4.0f, FStyleColors::Hover, 2.0f));
-
 		Set("ContentBrowser.AssetTileItem.DropShadow", new BOX_BRUSH("Starship/ContentBrowser/drop-shadow", FMargin(4.0f / 64.0f)));
+		Set("ContentBrowser.FolderItem.DropShadow", new BOX_BRUSH("Starship/ContentBrowser/folder-drop-shadow", FMargin(2.0f/64.0f)));
+
 
 
 		Set("ReferenceViewer.PathText", FEditableTextBoxStyle(NormalEditableTextBoxStyle)
diff --git a/Engine/Source/Editor/UnrealEd/Private/AssetThumbnail.cpp b/Engine/Source/Editor/UnrealEd/Private/AssetThumbnail.cpp
--- a/Engine/Source/Editor/UnrealEd/Private/AssetThumbnail.cpp
+++ b/Engine/Source/Editor/UnrealEd/Private/AssetThumbnail.cpp
@@ -368,7 +368,7 @@ private:
 			AssetColor = AssetTypeActions.Pin()->GetTypeColor();
 		}
 
-		AssetBackgroundWidget->SetBorderBackgroundColor(AssetColor.CopyWithNewOpacity(0.3f));
+		//AssetBackgroundWidget->SetBorderBackgroundColor(AssetColor.CopyWithNewOpacity(0.3f));
 		AssetColorStripWidget->SetBorderBackgroundColor(AssetColor);
 
 		UpdateThumbnailVisibilities();
@@ -400,11 +400,6 @@ private:
 		return FEditorStyle::GetBrush(ClassBackgroundBrushName);
 	}
 
-	FSlateColor GetViewportBorderColorAndOpacity() const
-	{
-		return FLinearColor(AssetColor.R, AssetColor.G, AssetColor.B, ViewportFadeCurve.GetLerp());
-	}
-
 	FLinearColor GetViewportColorAndOpacity() const
 	{
 		return FLinearColor(1, 1, 1, ViewportFadeCurve.GetLerp());
