Fixed asset class name in CB thumbnails not wrapping.

diff --git a/Engine/Source/Editor/ContentBrowser/Private/AssetViewWidgets.cpp b/Engine/Source/Editor/ContentBrowser/Private/AssetViewWidgets.cpp
--- a/Engine/Source/Editor/ContentBrowser/Private/AssetViewWidgets.cpp
+++ b/Engine/Source/Editor/ContentBrowser/Private/AssetViewWidgets.cpp
@@ -432,7 +432,9 @@ void SAssetViewItem::Tick( const FGeometry& AllottedGeometry, const double InCur
 	// We set this only when changed because binding a delegate to text wrapping attributes is expensive
 	if( PrevSizeX != AllottedGeometry.Size.X && InlineRenameWidget.IsValid() )
 	{
-		InlineRenameWidget->SetWrapTextAt( GetNameTextWrapWidth() );
+		const float WrapWidth = GetNameTextWrapWidth();
+		InlineRenameWidget->SetWrapTextAt(WrapWidth);
+		ClassTextWidget->SetWrapTextAt(WrapWidth);
 	}
 
 	UpdateDirtyState();
@@ -531,6 +533,12 @@ void SAssetViewItem::OnAssetDataChanged()
 		InlineRenameWidget->SetText( GetNameText() );
 	}
 
+	if (ClassTextWidget.IsValid())
+	{
+		ClassTextWidget->SetText(GetAssetClassText());
+	}
+
+
 	CacheDisplayTags();
 }
 
@@ -1503,7 +1511,7 @@ void SAssetListItem::Construct( const FArguments& InArgs )
 				.Padding(0, 1)
 				[
 					// Class
-					SAssignNew(ClassText, STextBlock)
+					SAssignNew(ClassTextWidget, STextBlock)
 					.Font(FEditorStyle::GetFontStyle("ContentBrowser.AssetListViewClassFont"))
 					.Text(GetAssetClassText())
 					.HighlightText(InArgs._HighlightText)
@@ -1524,9 +1532,9 @@ void SAssetListItem::OnAssetDataChanged()
 {
 	SAssetViewItem::OnAssetDataChanged();
 
-	if (ClassText.IsValid())
+	if (ClassTextWidget.IsValid())
 	{
-		ClassText->SetText(GetAssetClassText());
+		ClassTextWidget->SetText(GetAssetClassText());
 	}
 
 	if (AssetThumbnail)
@@ -1673,7 +1681,6 @@ void SAssetTileItem::Construct( const FArguments& InArgs )
 							SNew(SVerticalBox)
 							+ SVerticalBox::Slot()
 							.Padding(2.0f)
-							.HAlign(HAlign_Left)
 							.VAlign(VAlign_Top)
 							[
 								SAssignNew(InlineRenameWidget, SInlineEditableTextBlock)
@@ -1689,18 +1696,17 @@ void SAssetTileItem::Construct( const FArguments& InArgs )
 								.LineBreakPolicy(FBreakIterator::CreateCamelCaseBreakIterator())
 							]
 							+ SVerticalBox::Slot()
-							.HAlign(HAlign_Left)
 							.VAlign(VAlign_Bottom)
 							.Padding(2.0f)
-							.AutoHeight()
 							[
-								SNew(STextBlock)
+								SAssignNew(ClassTextWidget, STextBlock)
 								.Visibility(this, &SAssetTileItem::GetAssetClassLabelVisibility)
 								.TextStyle(FAppStyle::Get(), "SmallText")
-								//.Font(FAppStyle::Get().GetFontStyle("ContentBrowser.AssetTileViewNameFontSmall"))
+								//.Font(FAppStyle::Get().GetFontStyle("SmallFontBold"))
 								.TransformPolicy(ETextTransformPolicy::ToUpper)
 								.Text(this, &SAssetTileItem::GetAssetClassText)
 								.ColorAndOpacity(FSlateColor::UseSubduedForeground())
+								.LineBreakPolicy(FBreakIterator::CreateWordBreakIterator())
 							]
 							
 						]
diff --git a/Engine/Source/Editor/ContentBrowser/Private/AssetViewWidgets.h b/Engine/Source/Editor/ContentBrowser/Private/AssetViewWidgets.h
--- a/Engine/Source/Editor/ContentBrowser/Private/AssetViewWidgets.h
+++ b/Engine/Source/Editor/ContentBrowser/Private/AssetViewWidgets.h
@@ -257,6 +257,8 @@ protected:
 
 	TSharedPtr< SInlineEditableTextBlock > InlineRenameWidget;
 
+	TSharedPtr<STextBlock> ClassTextWidget;
+
 	/** The data for this item */
 	TSharedPtr<FAssetViewItem> AssetItem;
 
@@ -433,9 +435,6 @@ private:
 
 	/** The height allowed for this item */
 	TAttribute<float> ItemHeight;
-
-	/** The text block containing the class name */
-	TSharedPtr<STextBlock> ClassText;
 };
 
 /** An item in the asset tile view */
