Content browser reskin changes
- Made content browser toolbar extensible so that the import button functions properly
- Various incremental folder and asset style changes

diff --git a/Engine/Plugins/Editor/ContentBrowser/ContentBrowserAssetDataSource/Source/ContentBrowserAssetDataSource/ContentBrowserAssetDataSource.Build.cs b/Engine/Plugins/Editor/ContentBrowser/ContentBrowserAssetDataSource/Source/ContentBrowserAssetDataSource/ContentBrowserAssetDataSource.Build.cs
--- a/Engine/Plugins/Editor/ContentBrowser/ContentBrowserAssetDataSource/Source/ContentBrowserAssetDataSource/ContentBrowserAssetDataSource.Build.cs
+++ b/Engine/Plugins/Editor/ContentBrowser/ContentBrowserAssetDataSource/Source/ContentBrowserAssetDataSource/ContentBrowserAssetDataSource.Build.cs
@@ -28,6 +28,7 @@ namespace UnrealBuildTool.Rules
 					"AssetRegistry",
 					"AssetTools",
 					"CollectionManager",
+					"ContentBrowser",
 					"Engine",
 					"UnrealEd",
 					"EditorStyle",
diff --git a/Engine/Plugins/Editor/ContentBrowser/ContentBrowserAssetDataSource/Source/ContentBrowserAssetDataSource/Private/ContentBrowserAssetDataSource.cpp b/Engine/Plugins/Editor/ContentBrowser/ContentBrowserAssetDataSource/Source/ContentBrowserAssetDataSource/Private/ContentBrowserAssetDataSource.cpp
--- a/Engine/Plugins/Editor/ContentBrowser/ContentBrowserAssetDataSource/Source/ContentBrowserAssetDataSource/Private/ContentBrowserAssetDataSource.cpp
+++ b/Engine/Plugins/Editor/ContentBrowser/ContentBrowserAssetDataSource/Source/ContentBrowserAssetDataSource/Private/ContentBrowserAssetDataSource.cpp
@@ -24,6 +24,9 @@
 #include "NewAssetContextMenu.h"
 #include "AssetFolderContextMenu.h"
 #include "AssetFileContextMenu.h"
+#include "ContentBrowserMenuContexts.h"
+#include "Widgets/Input/SButton.h"
+#include "Widgets/Images/SImage.h"
 
 #define LOCTEXT_NAMESPACE "ContentBrowserAssetDataSource"
 
@@ -100,6 +103,19 @@ void UContentBrowserAssetDataSource::Initialize(const FName InMountRoot, const b
 			}));
 		}
 
+
+		if (UToolMenu* Menu = UToolMenus::Get()->ExtendMenu("ContentBrowser.ToolBar"))
+		{
+			Menu->AddDynamicSection(*FString::Printf(TEXT("DynamicSection_DataSource_%s"), *GetName()), FNewToolMenuDelegate::CreateLambda([WeakThis = TWeakObjectPtr<UContentBrowserAssetDataSource>(this)](UToolMenu* InMenu)
+			{
+				if (UContentBrowserAssetDataSource* This = WeakThis.Get())
+				{
+					This->PopulateContentBrowserToolBar(InMenu);
+				}
+			}));
+		}
+
+
 		if (UToolMenu* Menu = UToolMenus::Get()->ExtendMenu("ContentBrowser.FolderContextMenu"))
 		{
 			Menu->AddDynamicSection(*FString::Printf(TEXT("DynamicSection_DataSource_%s"), *GetName()), FNewToolMenuDelegate::CreateLambda([WeakThis = TWeakObjectPtr<UContentBrowserAssetDataSource>(this)](UToolMenu* InMenu)
@@ -1616,6 +1632,55 @@ void UContentBrowserAssetDataSource::PopulateAddNewContextMenu(UToolMenu* InMenu
 		);
 }
 
+void UContentBrowserAssetDataSource::PopulateContentBrowserToolBar(UToolMenu* InMenu)
+{
+	const UContentBrowserToolbarMenuContext* ContextObject = InMenu->FindContext<UContentBrowserToolbarMenuContext>();
+	checkf(ContextObject, TEXT("Required context UContentBrowserToolbarMenuContext was missing!"));
+
+	TSharedRef<SWidget> ImportButton = 
+		SNew(SButton)
+		.ButtonStyle(FAppStyle::Get(), "SimpleButton")
+		.ToolTipText(LOCTEXT("ImportTooltip", "Import assets from files to the currently selected folder"))
+		.ContentPadding(2)
+		.OnClicked_UObject(this, &UContentBrowserAssetDataSource::OnImportClicked, ContextObject)
+		.IsEnabled_UObject(this, &UContentBrowserAssetDataSource::IsImportEnabled, ContextObject)
+		[
+			SNew(SHorizontalBox)
+			+ SHorizontalBox::Slot()
+			.AutoWidth()
+			.HAlign(HAlign_Center)
+			.VAlign(VAlign_Center)
+			[
+				SNew(SImage)
+				.Image(FAppStyle::Get().GetBrush("Icons.Import"))
+				.ColorAndOpacity(FSlateColor::UseForeground())
+			]
+			+ SHorizontalBox::Slot()
+			.Padding(FMargin(3, 0, 0, 0))
+			.VAlign(VAlign_Center)
+			.AutoWidth()
+			[
+				SNew(STextBlock)
+				.TextStyle(FAppStyle::Get(), "NormalText")
+				.Text(LOCTEXT("Import", "Import"))
+			]
+		];
+
+	FToolMenuSection& Section = InMenu->FindOrAddSection("New");
+
+	Section.AddSeparator(NAME_None);
+
+	Section.AddEntry(
+		FToolMenuEntry::InitWidget(
+			"Import",
+			ImportButton,
+			FText::GetEmpty(),
+			true,
+			false
+		));
+	
+}
+
 void UContentBrowserAssetDataSource::PopulateAssetFolderContextMenu(UToolMenu* InMenu)
 {
 	return ContentBrowserAssetData::PopulateAssetFolderContextMenu(this, InMenu, *AssetFolderContextMenu);
@@ -1764,6 +1829,23 @@ bool UContentBrowserAssetDataSource::OnValidateItemName(const FContentBrowserIte
 	return CanRenameItem(InItem, &InProposedName, OutErrorMsg);
 }
 
+FReply UContentBrowserAssetDataSource::OnImportClicked(const UContentBrowserToolbarMenuContext* ContextObject)
+{
+	// Extract the internal asset paths that belong to this data source from the full list of selected paths given in the context
+	FName InternalPath;
+	if (TryConvertVirtualPathToInternal(ContextObject->GetCurrentPath(), InternalPath) && IsKnownContentPath(InternalPath))
+	{
+		OnImportAsset(InternalPath);
+	}
+
+	return FReply::Handled();
+}
+
+bool UContentBrowserAssetDataSource::IsImportEnabled(const UContentBrowserToolbarMenuContext* ContextObject) const
+{
+	return ContextObject->GetCurrentPath() != NAME_None;
+}
+
 FContentBrowserItemData UContentBrowserAssetDataSource::OnFinalizeCreateFolder(const FContentBrowserItemData& InItemData, const FString& InProposedName, FText* OutErrorMsg)
 {
 	checkf(InItemData.GetOwnerDataSource() == this, TEXT("OnFinalizeCreateFolder was bound to an instance from the wrong data source!"));
diff --git a/Engine/Plugins/Editor/ContentBrowser/ContentBrowserAssetDataSource/Source/ContentBrowserAssetDataSource/Public/ContentBrowserAssetDataSource.h b/Engine/Plugins/Editor/ContentBrowser/ContentBrowserAssetDataSource/Source/ContentBrowserAssetDataSource/Public/ContentBrowserAssetDataSource.h
--- a/Engine/Plugins/Editor/ContentBrowser/ContentBrowserAssetDataSource/Source/ContentBrowserAssetDataSource/Public/ContentBrowserAssetDataSource.h
+++ b/Engine/Plugins/Editor/ContentBrowser/ContentBrowserAssetDataSource/Source/ContentBrowserAssetDataSource/Public/ContentBrowserAssetDataSource.h
@@ -9,6 +9,7 @@
 #include "Misc/BlacklistNames.h"
 #include "ContentBrowserDataMenuContexts.h"
 #include "ContentBrowserAssetDataPayload.h"
+#include "Input/Reply.h"
 #include "ContentBrowserAssetDataSource.generated.h"
 
 class IAssetTools;
@@ -19,6 +20,7 @@ class UToolMenu;
 class FAssetFolderContextMenu;
 class FAssetFileContextMenu;
 struct FCollectionNameType;
+class UContentBrowserToolbarMenuContext;
 
 USTRUCT()
 struct CONTENTBROWSERASSETDATASOURCE_API FContentBrowserCompiledAssetDataFilter
@@ -201,6 +203,8 @@ private:
 
 	void PopulateAddNewContextMenu(UToolMenu* InMenu);
 
+	void PopulateContentBrowserToolBar(UToolMenu* InMenu);
+
 	void PopulateAssetFolderContextMenu(UToolMenu* InMenu);
 
 	void PopulateAssetFileContextMenu(UToolMenu* InMenu);
@@ -217,6 +221,10 @@ private:
 
 	bool OnValidateItemName(const FContentBrowserItemData& InItem, const FString& InProposedName, FText* OutErrorMsg);
 
+	FReply OnImportClicked(const UContentBrowserToolbarMenuContext* ContextObject);
+
+	bool IsImportEnabled(const UContentBrowserToolbarMenuContext* ContextObject) const;
+
 	FContentBrowserItemData OnFinalizeCreateFolder(const FContentBrowserItemData& InItemData, const FString& InProposedName, FText* OutErrorMsg);
 
 	FContentBrowserItemData OnFinalizeCreateAsset(const FContentBrowserItemData& InItemData, const FString& InProposedName, FText* OutErrorMsg);
diff --git a/Engine/Source/Editor/ContentBrowser/Private/AssetViewWidgets.cpp b/Engine/Source/Editor/ContentBrowser/Private/AssetViewWidgets.cpp
--- a/Engine/Source/Editor/ContentBrowser/Private/AssetViewWidgets.cpp
+++ b/Engine/Source/Editor/ContentBrowser/Private/AssetViewWidgets.cpp
@@ -114,10 +114,11 @@ TSharedRef<SWidget> FAssetViewItemHelper::CreateListTileItemContents(T* const In
 
 		// Folder base
 		ItemContentsOverlay->AddSlot()
+		.Padding(FMargin(5))
 		[
 			SNew(SBorder)
 			.BorderImage(FAppStyle::Get().GetBrush("ContentBrowser.FolderItem.DropShadow"))
-			.Padding(FMargin(0,0,4.0f,4.0f))
+			.Padding(FMargin(0,0,2.0f,2.0f))
 			[
 				SNew(SImage)
 				.Image(FolderBaseImage)
@@ -434,7 +435,11 @@ void SAssetViewItem::Tick( const FGeometry& AllottedGeometry, const double InCur
 	{
 		const float WrapWidth = GetNameTextWrapWidth();
 		InlineRenameWidget->SetWrapTextAt(WrapWidth);
-		ClassTextWidget->SetWrapTextAt(WrapWidth);
+
+		if (ClassTextWidget.IsValid())
+		{
+			ClassTextWidget->SetWrapTextAt(WrapWidth);
+		}
 	}
 
 	UpdateDirtyState();
@@ -1645,7 +1650,7 @@ void SAssetTileItem::Construct( const FArguments& InArgs )
 		// Drop shadow border
 		SNew(SBorder)
 		.Padding(FMargin(0.0f, 0.0f, 5.0f, 5.0f))
-		.BorderImage(IsFolder() ? FStyleDefaults::GetNoBrush() : FAppStyle::Get().GetBrush(ItemShadowBorderName))
+		.BorderImage(IsFolder() ? TAttribute<const FSlateBrush*>(this, &SAssetTileItem::GetFolderBackgroundShadowImage) : FAppStyle::Get().GetBrush(ItemShadowBorderName))
 		[
 			SNew(SOverlay)
 			.AddMetaData<FTagMetaData>(FTagMetaData(AssetItem->GetItem().GetVirtualPath()))
@@ -1653,7 +1658,7 @@ void SAssetTileItem::Construct( const FArguments& InArgs )
 			[
 				SNew(SBorder)
 				.Padding(0)
-				.BorderImage(IsFolder() ? FStyleDefaults::GetNoBrush() : FAppStyle::Get().GetBrush("ContentBrowser.AssetTileItem.ThumbnailAreaBackground"))
+				.BorderImage(IsFolder() ? TAttribute<const FSlateBrush*>(this, &SAssetTileItem::GetFolderBackgroundImage) : FAppStyle::Get().GetBrush("ContentBrowser.AssetTileItem.ThumbnailAreaBackground"))
 				[
 					SNew(SVerticalBox)
 					// Thumbnail
@@ -1692,7 +1697,7 @@ void SAssetTileItem::Construct( const FArguments& InArgs )
 								.HighlightText(InArgs._HighlightText)
 								.IsSelected(InArgs._IsSelectedExclusively)
 								.IsReadOnly(this, &SAssetTileItem::IsNameReadOnly)
-								.Justification(ETextJustify::Left)
+								.Justification(IsFolder() ? ETextJustify::Center : ETextJustify::Left)
 								.LineBreakPolicy(FBreakIterator::CreateCamelCaseBreakIterator())
 							]
 							+ SVerticalBox::Slot()
@@ -1701,8 +1706,7 @@ void SAssetTileItem::Construct( const FArguments& InArgs )
 							[
 								SAssignNew(ClassTextWidget, STextBlock)
 								.Visibility(this, &SAssetTileItem::GetAssetClassLabelVisibility)
-								.TextStyle(FAppStyle::Get(), "SmallText")
-								//.Font(FAppStyle::Get().GetFontStyle("SmallFontBold"))
+								.TextStyle(FAppStyle::Get(), "ContentBrowser.ClassFont")
 								.TransformPolicy(ETextTransformPolicy::ToUpper)
 								.Text(this, &SAssetTileItem::GetAssetClassText)
 								.ColorAndOpacity(FSlateColor::UseSubduedForeground())
@@ -1751,7 +1755,6 @@ const FSlateBrush* SAssetTileItem::GetBorderImage() const
 {
 	const bool bIsSelected = IsSelected.IsBound() ? IsSelected.Execute() : false;
 	const bool bIsHoveredOrDraggedOver = IsHovered() || bDraggedOver;
-
 	if (bIsSelected && bIsHoveredOrDraggedOver)
 	{
 		static const FName SelectedHover("ContentBrowser.AssetTileItem.SelectedHoverBorder");
@@ -1762,15 +1765,13 @@ const FSlateBrush* SAssetTileItem::GetBorderImage() const
 		static const FName Selected("ContentBrowser.AssetTileItem.SelectedBorder");
 		return FAppStyle::Get().GetBrush(Selected);
 	}
-	else if (bIsHoveredOrDraggedOver)
+	else if (bIsHoveredOrDraggedOver && !IsFolder())
 	{
 		static const FName Hovered("ContentBrowser.AssetTileItem.HoverBorder");
 		return FAppStyle::Get().GetBrush(Hovered);
 	}
-	else
-	{
-		return FStyleDefaults::GetNoBrush();
-	}
+
+	return FStyleDefaults::GetNoBrush();
 }
 
 float SAssetTileItem::GetExtraStateIconWidth() const
@@ -1832,6 +1833,34 @@ FSlateFontInfo SAssetTileItem::GetThumbnailFont() const
 
 
 
+const FSlateBrush* SAssetTileItem::GetFolderBackgroundImage() const
+{
+	const bool bIsSelected = IsSelected.IsBound() ? IsSelected.Execute() : false;
+	const bool bIsHoveredOrDraggedOver = IsHovered() || bDraggedOver;
+
+	if (bIsSelected || bIsHoveredOrDraggedOver)
+	{
+		static const FName HoveredBackground("ContentBrowser.AssetTileItem.FolderAreaHoveredBackground");
+		return FAppStyle::Get().GetBrush(HoveredBackground);
+	}
+
+	return FStyleDefaults::GetNoBrush();
+}
+
+const FSlateBrush* SAssetTileItem::GetFolderBackgroundShadowImage() const
+{
+	const bool bIsSelected = IsSelected.IsBound() ? IsSelected.Execute() : false;
+	const bool bIsHoveredOrDraggedOver = IsHovered() || bDraggedOver;
+
+	if (bIsSelected || bIsHoveredOrDraggedOver)
+	{
+		static const FName DropShadowName("ContentBrowser.AssetTileItem.DropShadow");
+		return FAppStyle::Get().GetBrush(DropShadowName);
+	}
+
+	return FStyleDefaults::GetNoBrush();
+}
+
 ///////////////////////////////
 // SAssetColumnItem
 ///////////////////////////////
diff --git a/Engine/Source/Editor/ContentBrowser/Private/AssetViewWidgets.h b/Engine/Source/Editor/ContentBrowser/Private/AssetViewWidgets.h
--- a/Engine/Source/Editor/ContentBrowser/Private/AssetViewWidgets.h
+++ b/Engine/Source/Editor/ContentBrowser/Private/AssetViewWidgets.h
@@ -551,6 +551,8 @@ protected:
 	/** Returns the font to use for the thumbnail label */
 	FSlateFontInfo GetThumbnailFont() const;
 
+	const FSlateBrush* GetFolderBackgroundImage() const;
+	const FSlateBrush* GetFolderBackgroundShadowImage() const;
 private:
 	/** The handle to the thumbnail that this item is rendering */
 	TSharedPtr<FAssetThumbnail> AssetThumbnail;
diff --git a/Engine/Source/Editor/ContentBrowser/Private/ContentBrowserMenuContexts.cpp b/Engine/Source/Editor/ContentBrowser/Private/ContentBrowserMenuContexts.cpp
new file mode 100644
--- /dev/null
+++ b/Engine/Source/Editor/ContentBrowser/Private/ContentBrowserMenuContexts.cpp
@@ -0,0 +1,16 @@
+// Copyright Epic Games, Inc. All Rights Reserved.
+
+#pragma once
+
+#include "ContentBrowserMenuContexts.h"
+#include "SContentBrowser.h"
+
+FName UContentBrowserToolbarMenuContext::GetCurrentPath() const
+{
+	if (TSharedPtr<SContentBrowser> Browser = ContentBrowser.Pin())
+	{
+		return *Browser->GetCurrentPath();
+	}
+
+	return NAME_None;
+}
diff --git a/Engine/Source/Editor/ContentBrowser/Private/SContentBrowser.cpp b/Engine/Source/Editor/ContentBrowser/Private/SContentBrowser.cpp
--- a/Engine/Source/Editor/ContentBrowser/Private/SContentBrowser.cpp
+++ b/Engine/Source/Editor/ContentBrowser/Private/SContentBrowser.cpp
@@ -221,119 +221,7 @@ void SContentBrowser::Construct( const FArguments& InArgs, const FName& InInstan
 				.FillEmptySpace( true )
 				.HAlign(HAlign_Left)
 				[
-					SNew( SHorizontalBox )
-					// Add
-					+ SHorizontalBox::Slot()
-					.Padding(5.0f, 0.0f)
-					.AutoWidth()
-					.VAlign( VAlign_Center )
-					.HAlign( HAlign_Left )
-					[
-						SNew( SComboButton )
-						.ComboButtonStyle(&FAppStyle::Get().GetWidgetStyle<FComboButtonStyle>("SimpleComboButton"))
-						.ForegroundColor(FSlateColor::UseStyle())
-						.ContentPadding(2)
-						.OnGetMenuContent_Lambda( [this]{ return MakeAddNewContextMenu(EContentBrowserDataMenuContext_AddNewMenuDomain::Toolbar); } )
-						.ToolTipText( this, &SContentBrowser::GetAddNewToolTipText )
-						.IsEnabled( this, &SContentBrowser::IsAddNewEnabled )
-						.AddMetaData<FTagMetaData>(FTagMetaData(TEXT("ContentBrowserNewAsset")))
-						.HasDownArrow(false)
-						.ButtonContent()
-						[
-							SNew(SHorizontalBox)
-							+ SHorizontalBox::Slot()
-							.AutoWidth()
-							.HAlign(HAlign_Center)
-							.VAlign(VAlign_Center)
-							[
-								SNew(SImage)
-								.Image(FAppStyle::Get().GetBrush("Icons.Plus"))
-								.ColorAndOpacity(FSlateColor::UseForeground())
-							]
-							+ SHorizontalBox::Slot()
-							.Padding(FMargin(3, 0, 0, 0))
-							.VAlign(VAlign_Center)
-							.AutoWidth()
-							[
-								SNew(STextBlock)
-								.TextStyle(FAppStyle::Get(), "NormalText")
-								.Text(LOCTEXT("NewAssetButton", "New"))
-							]
-						]
-					]
-					+ SHorizontalBox::Slot()
-					.Padding(5.0f, 0.0f)
-					.AutoWidth()
-					.VAlign(VAlign_Center)
-					.HAlign(HAlign_Left)
-					[
-						SNew(SButton)
-						.ButtonStyle(FAppStyle::Get(), "SimpleButton")
-						.ToolTipText( LOCTEXT( "ImportTooltip", "Import assets from files to the currently selected folder" ) )
-						.ContentPadding(2)
-						.OnClicked(this, &SContentBrowser::OnImportClicked)
-						.IsEnabled(this, &SContentBrowser::IsAddNewEnabled)
-						[
-							SNew( SHorizontalBox )
-							+ SHorizontalBox::Slot()
-							.AutoWidth()
-							.HAlign(HAlign_Center)
-							.VAlign(VAlign_Center)
-							[
-								SNew(SImage)
-								.Image(FAppStyle::Get().GetBrush("Icons.Import"))
-								.ColorAndOpacity(FSlateColor::UseForeground())
-							]
-							+ SHorizontalBox::Slot()
-							.Padding(FMargin(3, 0, 0, 0))
-							.VAlign(VAlign_Center)
-							.AutoWidth()
-							[
-								SNew(STextBlock)
-								.TextStyle(FAppStyle::Get(), "NormalText")
-								.Text( LOCTEXT( "Import", "Import" ) )
-							]
-						]
-					]
-					// Save
-					+ SHorizontalBox::Slot()
-					.Padding(5.0f, 0.0f)
-					.AutoWidth()
-					.VAlign(VAlign_Center)
-					.HAlign(HAlign_Left)
-					[
-						SNew( SButton )
-						.ButtonStyle(FAppStyle::Get(), "SimpleButton")
-						.ToolTipText( LOCTEXT( "SaveDirtyPackagesTooltip", "Save all modified assets." ) )
-						.ContentPadding(2)
-						.OnClicked(this, &SContentBrowser::OnSaveClicked)
-						.AddMetaData<FTagMetaData>(FTagMetaData(TEXT("ContentBrowserSaveDirtyPackages")))
-						[
-							SNew( SHorizontalBox )
-
-							// Save All Icon
-							+ SHorizontalBox::Slot()
-							.AutoWidth()
-							.HAlign(HAlign_Center)
-							.VAlign(VAlign_Center)
-							[
-								SNew(SImage)
-								.Image(FAppStyle::Get().GetBrush("Icons.Save"))
-								.ColorAndOpacity(FSlateColor::UseForeground())
-							]
-
-							// Save All Text
-							+ SHorizontalBox::Slot()
-							.Padding(FMargin(3, 0, 0, 0))
-							.VAlign(VAlign_Center)
-							.AutoWidth()
-							[
-								SNew(STextBlock)
-								.TextStyle(FAppStyle::Get(), "NormalText")
-								.Text( LOCTEXT( "SaveAll", "Save All" ) )
-							]
-						]
-					]
+					CreateToolBar(Config)
 				]
 				+ SWrapBox::Slot()
 				.FillEmptySpace(true)
@@ -443,12 +331,12 @@ void SContentBrowser::Construct( const FArguments& InArgs, const FName& InInstan
 			.Value(0.25f)
 			[
 				SNew(SBorder)
-				.Padding(FMargin(3))
+				.Padding(FMargin(0,0))
 				.BorderImage(FEditorStyle::GetBrush("Brushes.Background"))
 				.Visibility(this, &SContentBrowser::GetSourcesViewVisibility)
 				[
 					SNew(SBorder)
-					.Padding(FMargin(3))
+					.Padding(FMargin(0.0f, 2.0f, 0.0f, 3.0f))
 					.BorderImage(FEditorStyle::GetBrush("Brushes.Input"))
 					[
 						// Note: If adding more widgets here, fix ContentBrowserSourcesWidgetSwitcherIndex and the code that uses it!
@@ -459,8 +347,8 @@ void SContentBrowser::Construct( const FArguments& InArgs, const FName& InInstan
 						[
 							SAssignNew(PathFavoriteSplitterPtr, SSplitter)
 							.Style(FEditorStyle::Get(), "ContentBrowser.Splitter")
-							.PhysicalSplitterHandleSize(5.0f)
-							.HitDetectionSplitterHandleSize(5.0f)
+							.PhysicalSplitterHandleSize(3.0f)
+							.HitDetectionSplitterHandleSize(3.0f)
 							.Orientation(EOrientation::Orient_Vertical)
 							.MinimumSlotHeight(26.0f)
 							.Visibility( this, &SContentBrowser::GetSourcesViewVisibility )
@@ -687,6 +575,20 @@ void SContentBrowser::HandleAssetViewSearchOptionsChanged()
 	TextFilter->SetIncludeCollectionNames(AssetViewPtr->IsIncludingCollectionNames());
 }
 
+TSharedRef<SWidget> SContentBrowser::CreateToolBar(const FContentBrowserConfig* Config)
+{
+	RegisterContentBrowserToolBar();
+
+	FToolMenuContext MenuContext;
+
+	UContentBrowserToolbarMenuContext* CommonContextObject = NewObject<UContentBrowserToolbarMenuContext>();
+	CommonContextObject->ContentBrowser = SharedThis(this);
+
+	MenuContext.AddObject(CommonContextObject);
+
+	return UToolMenus::Get()->GenerateWidget("ContentBrowser.ToolBar", MenuContext);
+}
+
 TSharedRef<SWidget> SContentBrowser::CreateLockButton(const FContentBrowserConfig* Config)
 {
 	if(Config == nullptr || Config->bCanShowLockButton)
@@ -712,9 +614,9 @@ TSharedRef<SWidget> SContentBrowser::CreateLockButton(const FContentBrowserConfi
 TSharedRef<SWidget> SContentBrowser::CreateAssetView(const FContentBrowserConfig* Config)
 {
 	return
-
 		SNew(SVerticalBox)
 		+ SVerticalBox::Slot()
+		.Padding(0.0f, 2.0f, 0.0f, 0.0f)
 		.AutoHeight()
 		[
 			SNew(SBorder)
@@ -836,7 +738,7 @@ TSharedRef<SWidget> SContentBrowser::CreateFavoritesView(const FContentBrowserCo
 		SAssignNew(FavoritesArea, SExpandableArea)
 		.BorderImage(FAppStyle::Get().GetBrush("Brushes.Background"))
 		.BodyBorderImage(FAppStyle::Get().GetBrush("Brushes.Input"))
-		.HeaderPadding(FMargin(5.0f, 5.0f))
+		.HeaderPadding(FMargin(5.0f, 3.0f))
 		.Visibility(this, &SContentBrowser::GetFavoriteFolderVisibility)
 		.AllowAnimatedTransition(false)
 		.HeaderContent()
@@ -867,7 +769,7 @@ TSharedRef<SWidget> SContentBrowser::CreatePathView(const FContentBrowserConfig*
 		SAssignNew(PathArea, SExpandableArea)
 		.BorderImage(FAppStyle::Get().GetBrush("Brushes.Background"))
 		.BodyBorderImage(FAppStyle::Get().GetBrush("Brushes.Input"))
-		.HeaderPadding(FMargin(5.0f, 1.0f))
+		.HeaderPadding(FMargin(5.0f, 3.0f))
 		.AllowAnimatedTransition(false)
 		.HeaderContent()
 		[
@@ -910,7 +812,7 @@ TSharedRef<SWidget> SContentBrowser::CreateDockedCollectionsView(const FContentB
 		SAssignNew(CollectionArea, SExpandableArea)
 		.BorderImage(FAppStyle::Get().GetBrush("Brushes.Background"))
 		.BodyBorderImage(FAppStyle::Get().GetBrush("Brushes.Input"))
-		.HeaderPadding(FMargin(5.0f, 1.0f))
+		.HeaderPadding(FMargin(5.0f, 3.0f))
 		.Visibility(this, &SContentBrowser::GetDockedCollectionsVisibility)
 		.AllowAnimatedTransition(false)
 		.HeaderContent()
@@ -929,6 +831,7 @@ TSharedRef<SWidget> SContentBrowser::CreateDockedCollectionsView(const FContentB
 			+ SHorizontalBox::Slot()
 			.VAlign(VAlign_Center)
 			.HAlign(HAlign_Right)
+			.AutoWidth()
 			.Padding(5.0f, 0.0f)
 			[
 				SNew(SButton)
@@ -943,7 +846,6 @@ TSharedRef<SWidget> SContentBrowser::CreateDockedCollectionsView(const FContentB
 				]
 			]
 			+ SHorizontalBox::Slot()
-			.AutoWidth()
 			.VAlign(VAlign_Center)
 			.HAlign(HAlign_Right)
 			.Padding(5.0f, 0.0f)
@@ -983,6 +885,122 @@ void SContentBrowser::ExtendViewOptionsMenu(const FContentBrowserConfig* Config)
 	}
 }
 
+void SContentBrowser::RegisterContentBrowserToolBar()
+{
+	static const FName ToolBarName("ContentBrowser.ToolBar");
+	UToolMenus* ToolMenus = UToolMenus::Get();
+	if (ToolMenus->IsMenuRegistered(ToolBarName))
+	{
+		return;
+	}
+
+	UToolMenu* ToolBar = UToolMenus::Get()->RegisterMenu(ToolBarName, NAME_None, EMultiBoxType::SlimHorizontalToolBar);
+	ToolBar->StyleName = "ContentBrowser.ToolBar";
+
+	{
+		FToolMenuSection& Section = ToolBar->AddSection("New");
+		
+		Section.AddDynamicEntry("New", FNewToolMenuSectionDelegate::CreateLambda([](FToolMenuSection& InSection)
+			{
+				UContentBrowserToolbarMenuContext* Context = InSection.FindContext<UContentBrowserToolbarMenuContext>();
+				TSharedRef<SContentBrowser> ContentBrowser = Context->ContentBrowser.Pin().ToSharedRef();
+		
+				TSharedRef<SComboButton> NewButton =
+					SNew(SComboButton)
+					.ComboButtonStyle(&FAppStyle::Get().GetWidgetStyle<FComboButtonStyle>("SimpleComboButton"))
+					.ForegroundColor(FSlateColor::UseStyle())
+					.ContentPadding(2)
+					.OnGetMenuContent_Lambda([Context] { return Context->ContentBrowser.Pin()->MakeAddNewContextMenu(EContentBrowserDataMenuContext_AddNewMenuDomain::Toolbar, Context); })
+					.ToolTipText(ContentBrowser, &SContentBrowser::GetAddNewToolTipText)
+					.IsEnabled(ContentBrowser, &SContentBrowser::IsAddNewEnabled)
+					.AddMetaData<FTagMetaData>(FTagMetaData(TEXT("ContentBrowserNewAsset")))
+					.HasDownArrow(false)
+					.ButtonContent()
+					[
+						SNew(SHorizontalBox)
+						+ SHorizontalBox::Slot()
+						.AutoWidth()
+						.HAlign(HAlign_Center)
+						.VAlign(VAlign_Center)
+						[
+							SNew(SImage)
+							.Image(FAppStyle::Get().GetBrush("Icons.Plus"))
+							.ColorAndOpacity(FSlateColor::UseForeground())
+						]
+						+ SHorizontalBox::Slot()
+						.Padding(FMargin(3, 0, 0, 0))
+						.VAlign(VAlign_Center)
+						.AutoWidth()
+						[
+							SNew(STextBlock)
+							.TextStyle(FAppStyle::Get(), "NormalText")
+							.Text(LOCTEXT("NewAssetButton", "New"))
+						]
+					];
+		
+				InSection.AddEntry(
+					FToolMenuEntry::InitWidget(
+						"NewButton",
+						NewButton,
+						FText::GetEmpty(),
+						true,
+						false
+					));
+			}));
+		}
+
+	{
+		FToolMenuSection& Section = ToolBar->AddSection("Save");
+		Section.AddDynamicEntry("Save", FNewToolMenuSectionDelegate::CreateLambda([](FToolMenuSection& InSection)
+		{
+			UContentBrowserToolbarMenuContext* Context = InSection.FindContext<UContentBrowserToolbarMenuContext>();
+			TSharedRef<SContentBrowser> ContentBrowser = Context->ContentBrowser.Pin().ToSharedRef();
+
+
+			TSharedRef<SButton> SaveButton =
+				SNew(SButton)
+				.ButtonStyle(FAppStyle::Get(), "SimpleButton")
+				.ToolTipText(LOCTEXT("SaveDirtyPackagesTooltip", "Save all modified assets."))
+				.ContentPadding(2)
+				.OnClicked(ContentBrowser, &SContentBrowser::OnSaveClicked)
+				.AddMetaData<FTagMetaData>(FTagMetaData(TEXT("ContentBrowserSaveDirtyPackages")))
+				[
+					SNew(SHorizontalBox)
+					// Save All Icon
+					+ SHorizontalBox::Slot()
+					.AutoWidth()
+					.HAlign(HAlign_Center)
+					.VAlign(VAlign_Center)
+					[
+						SNew(SImage)
+						.Image(FAppStyle::Get().GetBrush("Icons.Save"))
+						.ColorAndOpacity(FSlateColor::UseForeground())
+					]
+
+					// Save All Text
+					+ SHorizontalBox::Slot()
+					.Padding(FMargin(3, 0, 0, 0))
+					.VAlign(VAlign_Center)
+					.AutoWidth()
+					[
+						SNew(STextBlock)
+						.TextStyle(FAppStyle::Get(), "NormalText")
+						.Text(LOCTEXT("SaveAll", "Save All"))
+					]
+				];
+
+			InSection.AddEntry(
+				FToolMenuEntry::InitWidget(
+					"SaveButton",
+					SaveButton,
+					FText::GetEmpty(),
+					true,
+					false
+				));
+		}));
+	}
+}
+
 SSplitter::ESizeRule SContentBrowser::GetFavoritesAreaSizeRule() const
 {
 	return FavoritesArea->IsExpanded() ? SSplitter::ESizeRule::FractionOfParent : SSplitter::ESizeRule::SizeToContent;
@@ -2045,14 +2063,14 @@ FString SContentBrowser::GetCurrentPath() const
 	return CurrentPath;
 }
 
-void SContentBrowser::AppendNewMenuContextObjects(const EContentBrowserDataMenuContext_AddNewMenuDomain InDomain, const TArray<FName>& InSelectedPaths, FToolMenuContext& InOutMenuContext)
+void SContentBrowser::AppendNewMenuContextObjects(const EContentBrowserDataMenuContext_AddNewMenuDomain InDomain, const TArray<FName>& InSelectedPaths, FToolMenuContext& InOutMenuContext, UContentBrowserToolbarMenuContext* CommonContext)
 {
 	if (!UToolMenus::Get()->IsMenuRegistered("ContentBrowser.AddNewContextMenu"))
 	{
 		UToolMenu* Menu = UToolMenus::Get()->RegisterMenu("ContentBrowser.AddNewContextMenu");
 		Menu->AddDynamicSection("DynamicSection_Common", FNewToolMenuDelegate::CreateLambda([](UToolMenu* InMenu)
 		{
-			if (const UContentBrowserAddNewContextMenuContext* ContextObject = InMenu->FindContext<UContentBrowserAddNewContextMenuContext>())
+			if (const UContentBrowserToolbarMenuContext* ContextObject = InMenu->FindContext<UContentBrowserToolbarMenuContext>())
 			{
 				if (TSharedPtr<SContentBrowser> ContentBrowser = ContextObject->ContentBrowser.Pin())
 				{
@@ -2062,11 +2080,16 @@ void SContentBrowser::AppendNewMenuContextObjects(const EContentBrowserDataMenuC
 		}));
 	}
 
+	if(!CommonContext)
 	{
-		UContentBrowserAddNewContextMenuContext* CommonContextObject = NewObject<UContentBrowserAddNewContextMenuContext>();
+		UContentBrowserToolbarMenuContext* CommonContextObject = NewObject<UContentBrowserToolbarMenuContext>();
 		CommonContextObject->ContentBrowser = SharedThis(this);
 		InOutMenuContext.AddObject(CommonContextObject);
 	}
+	else
+	{
+		InOutMenuContext.AddObject(CommonContext);
+	}
 
 	{
 		UContentBrowserDataMenuContext_AddNewMenu* DataContextObject = NewObject<UContentBrowserDataMenuContext_AddNewMenu>();
@@ -2077,7 +2100,7 @@ void SContentBrowser::AppendNewMenuContextObjects(const EContentBrowserDataMenuC
 	}
 }
 
-TSharedRef<SWidget> SContentBrowser::MakeAddNewContextMenu(const EContentBrowserDataMenuContext_AddNewMenuDomain InDomain)
+TSharedRef<SWidget> SContentBrowser::MakeAddNewContextMenu(const EContentBrowserDataMenuContext_AddNewMenuDomain InDomain, UContentBrowserToolbarMenuContext* CommonContext)
 {
 	const FSourcesData& SourcesData = AssetViewPtr->GetSourcesData();
 
@@ -2123,7 +2146,7 @@ TSharedRef<SWidget> SContentBrowser::MakeAddNewContextMenu(const EContentBrowser
 	}
 
 	FToolMenuContext ToolMenuContext(nullptr, MenuExtender, nullptr);
-	AppendNewMenuContextObjects(InDomain, SourcesData.VirtualPaths, ToolMenuContext);
+	AppendNewMenuContextObjects(InDomain, SourcesData.VirtualPaths, ToolMenuContext, CommonContext);
 
 	FDisplayMetrics DisplayMetrics;
 	FSlateApplication::Get().GetCachedDisplayMetrics( DisplayMetrics );
@@ -2142,19 +2165,6 @@ TSharedRef<SWidget> SContentBrowser::MakeAddNewContextMenu(const EContentBrowser
 		];
 }
 
-FReply SContentBrowser::OnImportClicked()
-{
-	FAssetToolsModule& AssetToolsModule = FModuleManager::LoadModuleChecked<FAssetToolsModule>("AssetTools");
-
-	FString CurrentPath = GetCurrentPath();
-	if (!CurrentPath.IsEmpty())
-	{
-		AssetToolsModule.Get().ImportAssetsWithDialog(GetCurrentPath());
-	}
-
-	return FReply::Handled();
-}
-
 void SContentBrowser::PopulateAddNewContextMenu(class UToolMenu* Menu)
 {
 	const UContentBrowserDataMenuContext_AddNewMenu* ContextObject = Menu->FindContext<UContentBrowserDataMenuContext_AddNewMenu>();
@@ -3367,7 +3377,7 @@ TSharedPtr<SWidget> SContentBrowser::GetItemContextMenu(TArrayView<const FConten
 			{
 				SelectedVirtualPaths.Add(SelectedFolder.GetVirtualPath());
 			}
-			AppendNewMenuContextObjects(EContentBrowserDataMenuContext_AddNewMenuDomain::PathView, SelectedVirtualPaths, MenuContext);
+			AppendNewMenuContextObjects(EContentBrowserDataMenuContext_AddNewMenuDomain::PathView, SelectedVirtualPaths, MenuContext, nullptr);
 		}
 
 		return UToolMenus::Get()->GenerateWidget("ContentBrowser.FolderContextMenu", MenuContext);
@@ -3381,7 +3391,7 @@ TSharedPtr<SWidget> SContentBrowser::GetItemContextMenu(TArrayView<const FConten
 	else if (ViewContext == EContentBrowserViewContext::AssetView)
 	{
 		// Nothing selected - show the new asset menu
-		return MakeAddNewContextMenu(EContentBrowserDataMenuContext_AddNewMenuDomain::AssetView);
+		return MakeAddNewContextMenu(EContentBrowserDataMenuContext_AddNewMenuDomain::AssetView, nullptr);
 	}
 
 	return nullptr;
diff --git a/Engine/Source/Editor/ContentBrowser/Private/SContentBrowser.h b/Engine/Source/Editor/ContentBrowser/Private/SContentBrowser.h
--- a/Engine/Source/Editor/ContentBrowser/Private/SContentBrowser.h
+++ b/Engine/Source/Editor/ContentBrowser/Private/SContentBrowser.h
@@ -31,6 +31,7 @@ class SPathView;
 class SSplitter;
 class UFactory;
 class SExpandableSearchArea;
+class UContentBrowserToolbarMenuContext;
 
 struct FToolMenuContext;
 
@@ -282,16 +283,10 @@ private:
 	TSharedRef<SWidget> GetPathPickerContent();
 
 	/** Register the context objects needed for the "Add New" menu */
-	void AppendNewMenuContextObjects(const EContentBrowserDataMenuContext_AddNewMenuDomain InDomain, const TArray<FName>& InSelectedPaths, FToolMenuContext& InOutMenuContext);
+	void AppendNewMenuContextObjects(const EContentBrowserDataMenuContext_AddNewMenuDomain InDomain, const TArray<FName>& InSelectedPaths, FToolMenuContext& InOutMenuContext, UContentBrowserToolbarMenuContext* CommonContext);
 
 	/** Handle creating a context menu for the "Add New" button */
-	TSharedRef<SWidget> MakeAddNewContextMenu(const EContentBrowserDataMenuContext_AddNewMenuDomain InDomain);
-
-	/* Whether or not the import button is enabled. */
-	bool IsImportButtonEnabled() const;
-
-	/** Called when the import button on the toolbar is clicked */
-	FReply OnImportClicked();
+	TSharedRef<SWidget> MakeAddNewContextMenu(const EContentBrowserDataMenuContext_AddNewMenuDomain InDomain, UContentBrowserToolbarMenuContext* CommonContext);
 
 	/** Handle populating a context menu for the "Add New" button */
 	void PopulateAddNewContextMenu(class UToolMenu* Menu);
@@ -507,6 +502,7 @@ private:
 	void HandleAssetViewSearchOptionsChanged();
 
 	/** Creates various widgets for the content browser main view */
+	TSharedRef<SWidget> CreateToolBar(const FContentBrowserConfig* Config);
 	TSharedRef<SWidget> CreateLockButton(const FContentBrowserConfig* Config);
 	TSharedRef<SWidget> CreateAssetView(const FContentBrowserConfig* Config);
 	TSharedRef<SWidget> CreateFavoritesView(const FContentBrowserConfig* Config);
@@ -516,6 +512,9 @@ private:
 	/** Adds menu options to the view menu */
 	void ExtendViewOptionsMenu(const FContentBrowserConfig* Config);
 
+	/** Creates the content browser toolbar menu */
+	void RegisterContentBrowserToolBar();
+
 	/** Gets the size rule for various areas. When areas are a collapsed the splitter slot becomes auto sized, otherwise it is user sized */
 	SSplitter::ESizeRule GetFavoritesAreaSizeRule() const;
 	SSplitter::ESizeRule GetPathAreaSizeRule() const;
diff --git a/Engine/Source/Editor/ContentBrowser/Private/SExpandableSearchArea.cpp b/Engine/Source/Editor/ContentBrowser/Private/SExpandableSearchArea.cpp
--- a/Engine/Source/Editor/ContentBrowser/Private/SExpandableSearchArea.cpp
+++ b/Engine/Source/Editor/ContentBrowser/Private/SExpandableSearchArea.cpp
@@ -19,8 +19,6 @@ void SExpandableSearchArea::Construct(const FArguments& InArgs, TSharedRef<SSear
 	[
 		SNew(SHorizontalBox)
 		+ SHorizontalBox::Slot()
-		.HAlign(HAlign_Center)
-		.AutoWidth()
 		[
 			SNew(SBox)
 			.Visibility(this, &SExpandableSearchArea::GetSearchBoxVisibility)
@@ -31,6 +29,7 @@ void SExpandableSearchArea::Construct(const FArguments& InArgs, TSharedRef<SSear
 			]
 		]
 		+ SHorizontalBox::Slot()
+		.AutoWidth()
 		[
 			SNew(SButton)
 			.ToolTipText(NSLOCTEXT("ExpandableSearchArea", "ExpandCollapseSearchButton", "Expands or collapses the search text box"))
@@ -92,7 +91,6 @@ EVisibility SExpandableSearchArea::GetSearchGlassVisibility() const
 const FSlateBrush* SExpandableSearchArea::GetExpandSearchImage() const
 {
 	static const FName RightIcon("Icons.ChevronRight");
-	static const FName LeftIcon("Icons.ChevronLeft");
 
-	return bIsExpanded ? FAppStyle::Get().GetBrush(RightIcon) : FAppStyle::Get().GetBrush(LeftIcon);
+	return bIsExpanded ? FAppStyle::Get().GetBrush(RightIcon) : FStyleDefaults::GetNoBrush();
 }
diff --git a/Engine/Source/Editor/ContentBrowser/Public/ContentBrowserMenuContexts.h b/Engine/Source/Editor/ContentBrowser/Public/ContentBrowserMenuContexts.h
--- a/Engine/Source/Editor/ContentBrowser/Public/ContentBrowserMenuContexts.h
+++ b/Engine/Source/Editor/ContentBrowser/Public/ContentBrowserMenuContexts.h
@@ -78,11 +78,14 @@ public:
 };
 
 UCLASS()
-class CONTENTBROWSER_API UContentBrowserAddNewContextMenuContext : public UObject
+class CONTENTBROWSER_API UContentBrowserToolbarMenuContext : public UObject
 {
 	GENERATED_BODY()
 
+public:
+	FName GetCurrentPath() const;
 public:
 
 	TWeakPtr<SContentBrowser> ContentBrowser;
 };
+
diff --git a/Engine/Source/Editor/EditorStyle/Private/StarshipStyle.cpp b/Engine/Source/Editor/EditorStyle/Private/StarshipStyle.cpp
--- a/Engine/Source/Editor/EditorStyle/Private/StarshipStyle.cpp
+++ b/Engine/Source/Editor/EditorStyle/Private/StarshipStyle.cpp
@@ -6051,14 +6051,7 @@ void FStarshipEditorStyle::FStyle::SetupContentBrowserStyle()
 		Set("ContentBrowser.AssetTagNamePadding.Compact", FMargin(2.0f));
 		Set("ContentBrowser.AssetTagCountPadding.Compact", FMargin(2.0f));
 
-		// Thumbnail editing mode
-		Set( "ContentBrowser.EditModeLabelFont", FTextBlockStyle(NormalText)
-			.SetFont( DEFAULT_FONT( "Regular", 10 ) )
-			.SetColorAndOpacity( FLinearColor::Black )
-			.SetShadowOffset( FVector2D::ZeroVector )
-		);
-
-
+			
 		Set( "ContentBrowser.PrimitiveCustom", new IMAGE_BRUSH( "ContentBrowser/ThumbnailCustom", Icon32x32 ) );
 		Set( "ContentBrowser.PrimitiveSphere", new IMAGE_BRUSH( "ContentBrowser/ThumbnailSphere", Icon32x32 ) );
 		Set( "ContentBrowser.PrimitiveCube", new IMAGE_BRUSH( "ContentBrowser/ThumbnailCube", Icon32x32 ) );
@@ -6073,6 +6066,10 @@ void FStarshipEditorStyle::FStyle::SetupContentBrowserStyle()
 			.SetShadowOffset( FVector2D( 1, 1 ) )
 			.SetShadowColorAndOpacity( FLinearColor( 0, 0, 0, 0.9f ) ) );
 
+		Set("ContentBrowser.ClassFont", FTextBlockStyle(NormalText)
+			.SetFont(DEFAULT_FONT("Regular", 7)));
+
+
 		// New Asset
 		Set( "ContentBrowser.NewAsset", new IMAGE_BRUSH( "Icons/icon_file_new_40x", Icon25x25 ) );
 
@@ -6128,13 +6125,14 @@ void FStarshipEditorStyle::FStyle::SetupContentBrowserStyle()
 		Set( "ContentBrowser.ListViewFolderIcon", new IMAGE_BRUSH_SVG( "Starship/ContentBrowser/folder", FVector2D(64, 64) ) );
 		Set( "ContentBrowser.ListViewDeveloperFolderIcon", new IMAGE_BRUSH( "Icons/Folders/FolderDev_Base_256x", FVector2D(256, 256) ) );
 
+		Set("ContentBrowser.AssetTileItem.FolderAreaHoveredBackground", new FSlateRoundedBoxBrush(FStyleColors::Dropdown, 4.0f));
 		Set("ContentBrowser.AssetTileItem.ThumbnailAreaBackground", new FSlateRoundedBoxBrush(FStyleColors::Input, 4.0f));
 		Set("ContentBrowser.AssetTileItem.NameAreaBackground", new FSlateRoundedBoxBrush(FStyleColors::Dropdown, 4.0f));
-		Set("ContentBrowser.AssetTileItem.SelectedBorder", new FSlateRoundedBoxBrush(FStyleColors::Transparent, 4.0f, FStyleColors::Primary, 2.0f));
-		Set("ContentBrowser.AssetTileItem.SelectedHoverBorder", new FSlateRoundedBoxBrush(FStyleColors::Transparent, 4.0f, FStyleColors::PrimaryHover, 2.0f));
-		Set("ContentBrowser.AssetTileItem.HoverBorder", new FSlateRoundedBoxBrush(FStyleColors::Transparent, 4.0f, FStyleColors::Hover, 2.0f));
+		Set("ContentBrowser.AssetTileItem.SelectedBorder", new FSlateRoundedBoxBrush(FStyleColors::Transparent, 4.0f, FStyleColors::Primary, 1.0f));
+		Set("ContentBrowser.AssetTileItem.SelectedHoverBorder", new FSlateRoundedBoxBrush(FStyleColors::Transparent, 4.0f, FStyleColors::PrimaryHover, 1.0f));
+		Set("ContentBrowser.AssetTileItem.HoverBorder", new FSlateRoundedBoxBrush(FStyleColors::Transparent, 4.0f, FStyleColors::Hover, 1.0f));
 		Set("ContentBrowser.AssetTileItem.DropShadow", new BOX_BRUSH("Starship/ContentBrowser/drop-shadow", FMargin(4.0f / 64.0f)));
-		Set("ContentBrowser.FolderItem.DropShadow", new BOX_BRUSH("Starship/ContentBrowser/folder-drop-shadow", FMargin(2.0f/64.0f)));
+		Set("ContentBrowser.FolderItem.DropShadow", new IMAGE_BRUSH("Starship/ContentBrowser/folder-drop-shadow", FVector2D(256, 256)));
 
 
 
@@ -6182,6 +6180,14 @@ void FStarshipEditorStyle::FStyle::SetupContentBrowserStyle()
 
 		Set( "ContentBrowser.SaveAllCurrentFolder", new IMAGE_BRUSH("Icons/icon_file_saveall_16px", Icon16x16) );
 		Set( "ContentBrowser.ResaveAllCurrentFolder", new IMAGE_BRUSH("Icons/icon_file_saveall_16px", Icon16x16) );
+
+
+		FToolBarStyle ContentBrowserToolBarStyle = FStarshipCoreStyle::GetCoreStyle().GetWidgetStyle<FToolBarStyle>("SlimToolBar");
+
+		ContentBrowserToolBarStyle.SetSeparatorBrush(FSlateNoResource());
+		ContentBrowserToolBarStyle.SetSeparatorPadding(FMargin(4.0f, 0.0f));
+	
+		Set("ContentBrowser.ToolBar", ContentBrowserToolBarStyle);
 	}
 #endif // #if WITH_EDITOR
 }
