Updated AssetPicker and Save Dialogs

diff --git a/Engine/Source/Editor/ContentBrowser/Private/SAssetDialog.cpp b/Engine/Source/Editor/ContentBrowser/Private/SAssetDialog.cpp
--- a/Engine/Source/Editor/ContentBrowser/Private/SAssetDialog.cpp
+++ b/Engine/Source/Editor/ContentBrowser/Private/SAssetDialog.cpp
@@ -12,6 +12,7 @@
 #include "Widgets/Input/SEditableTextBox.h"
 #include "Widgets/Input/SButton.h"
 #include "Widgets/Layout/SSplitter.h"
+#include "Styling/AppStyle.h"
 #include "EditorStyleSet.h"
 #include "AssetRegistryModule.h"
 #include "ContentBrowserSingleton.h"
@@ -190,7 +191,7 @@ void SAssetDialog::Construct(const FArguments& InArgs, const FSharedAssetDialogC
 		+ SVerticalBox::Slot()
 		.FillHeight(1)
 		.VAlign(VAlign_Center)
-		.Padding(0, 2, 0, 2)
+		.Padding(0, 4, 0, 4)
 		[
 			SNew(STextBlock).Text(LOCTEXT("PathBoxLabel", "Path:"))
 		];
@@ -199,7 +200,7 @@ void SAssetDialog::Construct(const FArguments& InArgs, const FSharedAssetDialogC
 		+ SVerticalBox::Slot()
 		.FillHeight(1)
 		.VAlign(VAlign_Center)
-		.Padding(0, 2, 0, 2)
+		.Padding(0, 4, 0, 4)
 		[
 			SAssignNew(PathText, STextBlock)
 			.Text(this, &SAssetDialog::GetPathNameText)
@@ -210,7 +211,7 @@ void SAssetDialog::Construct(const FArguments& InArgs, const FSharedAssetDialogC
 		LabelsBox->AddSlot()
 			.FillHeight(1)
 			.VAlign(VAlign_Center)
-			.Padding(0, 2, 0, 2)
+			.Padding(0, 4, 0, 4)
 			[
 				SNew(STextBlock).Text(LOCTEXT("NameBoxLabel", "Name:"))
 			];
@@ -234,35 +235,36 @@ void SAssetDialog::Construct(const FArguments& InArgs, const FSharedAssetDialogC
 		.AutoWidth()
 		.HAlign(HAlign_Right)
 		.VAlign(VAlign_Bottom)
-		.Padding(bIncludeNameBox ? 80 : 4, 20, 4, 3)
+		.Padding(0.f, 0.f, 4.f, 0.f)
 		[
 			LabelsBox
 		]
 		+ SHorizontalBox::Slot()
 		.FillWidth(1)
 		.VAlign(VAlign_Bottom)
-		.Padding(4, 3)
+		.Padding(4.f, 0.f) 
 		[
 			ContentBox
 		]
 		+ SHorizontalBox::Slot()
 		.AutoWidth()
 		.VAlign(VAlign_Bottom)
-		.Padding(4, 3)
+		.Padding(4.f, 0.f) 
 		[
 			SNew(SButton)
+			.ButtonStyle(FAppStyle::Get(), "PrimaryButton")
+			.TextStyle(FAppStyle::Get(), "DialogButtonText")
 			.Text(ConfirmButtonText)
-			.ContentPadding(FMargin(8, 2, 8, 2))
 			.IsEnabled(this, &SAssetDialog::IsConfirmButtonEnabled)
 			.OnClicked(this, &SAssetDialog::OnConfirmClicked)
 		]
 		+ SHorizontalBox::Slot()
 		.AutoWidth()
 		.VAlign(VAlign_Bottom)
-		.Padding(4, 3)
+		.Padding(4.f, 0.f, 0.0f, 0.0f) 
 		[
 			SNew(SButton)
-			.ContentPadding(FMargin(8, 2, 8, 2))
+			.TextStyle(FAppStyle::Get(), "DialogButtonText")
 			.Text(LOCTEXT("AssetDialogCancelButton", "Cancel"))
 			.OnClicked(this, &SAssetDialog::OnCancelClicked)
 		];
@@ -270,7 +272,7 @@ void SAssetDialog::Construct(const FArguments& InArgs, const FSharedAssetDialogC
 	MainVerticalBox->AddSlot()
 		.AutoHeight()
 		.HAlign(HAlign_Fill)
-		.Padding(0)
+		.Padding(16.f, 4.f, 16.f, 16.f)
 		[
 			ButtonsAndNameBox
 		];
diff --git a/Engine/Source/Editor/ContentBrowser/Private/SAssetPicker.cpp b/Engine/Source/Editor/ContentBrowser/Private/SAssetPicker.cpp
--- a/Engine/Source/Editor/ContentBrowser/Private/SAssetPicker.cpp
+++ b/Engine/Source/Editor/ContentBrowser/Private/SAssetPicker.cpp
@@ -101,12 +101,13 @@ void SAssetPicker::Construct( const FArguments& InArgs )
 
 	FrontendFilters = MakeShareable(new FAssetFilterCollectionType());
 
+	TSharedRef<SHorizontalBox> HorizontalBox = SNew(SHorizontalBox);
+
 	// Search box
 	if (!InArgs._AssetPickerConfig.bAutohideSearchBar)
 	{
 		HighlightText = TAttribute< FText >( this, &SAssetPicker::GetHighlightedText );
 
-		TSharedRef<SHorizontalBox> HorizontalBox = SNew(SHorizontalBox);
 
 		if (InArgs._AssetPickerConfig.bAddFilterUI)
 		{
@@ -147,12 +148,13 @@ void SAssetPicker::Construct( const FArguments& InArgs )
 		.AutoWidth()
 		[
 			SNew(SCheckBox)
-			.Style(FEditorStyle::Get(), "ToggleButtonCheckbox")
+			.Style(FEditorStyle::Get(), "ToggleButtonCheckBox")
 			.ToolTipText(this, &SAssetPicker::GetShowOtherDevelopersToolTip)
 			.OnCheckStateChanged(this, &SAssetPicker::HandleShowOtherDevelopersCheckStateChanged)
 			.IsChecked(this, &SAssetPicker::GetShowOtherDevelopersCheckState)
 			[
 				SNew(SImage)
+				.ColorAndOpacity(FSlateColor::UseForeground())
 				.Image(FEditorStyle::GetBrush("ContentBrowser.ColumnViewDeveloperFolderIcon"))
 			]
 		];
@@ -303,6 +305,7 @@ void SAssetPicker::Construct( const FArguments& InArgs )
 		.CanShowFolders( InArgs._AssetPickerConfig.bCanShowFolders )
 		.ShowPathInColumnView( InArgs._AssetPickerConfig.bShowPathInColumnView)
 		.ShowTypeInColumnView( InArgs._AssetPickerConfig.bShowTypeInColumnView)
+		.ShowViewOptions(false)  // We control this in the asset picker
 		.SortByPathInColumnView( InArgs._AssetPickerConfig.bSortByPathInColumnView)
 		.FilterRecursivelyWithBackendFilter( false )
 		.CanShowRealTimeThumbnails( InArgs._AssetPickerConfig.bCanShowRealTimeThumbnails )
@@ -318,6 +321,24 @@ void SAssetPicker::Construct( const FArguments& InArgs )
 		.OnSearchOptionsChanged(this, &SAssetPicker::HandleSearchSettingsChanged)
 	];
 
+
+	HorizontalBox->AddSlot()
+	.AutoWidth()
+	[
+		SNew(SComboButton)
+		.ContentPadding(0)
+		.ComboButtonStyle(&FAppStyle::Get().GetWidgetStyle<FComboButtonStyle>("SimpleComboButton"))
+		.OnGetMenuContent(AssetViewPtr.ToSharedRef(), &SAssetView::GetViewButtonContent)
+		.HasDownArrow(false)
+		.ButtonContent()
+		[
+			SNew(SImage)
+			.ColorAndOpacity(FSlateColor::UseForeground())
+			.Image(FAppStyle::Get().GetBrush("Icons.Settings"))
+		]	
+	];
+
+
 	LoadSettings();
 
 	if (AssetViewPtr.IsValid() && !InArgs._AssetPickerConfig.bAutohideSearchBar)
diff --git a/Engine/Source/Editor/PackagesDialog/Private/SPackagesDialog.cpp b/Engine/Source/Editor/PackagesDialog/Private/SPackagesDialog.cpp
--- a/Engine/Source/Editor/PackagesDialog/Private/SPackagesDialog.cpp
+++ b/Engine/Source/Editor/PackagesDialog/Private/SPackagesDialog.cpp
@@ -272,6 +272,8 @@ void SPackagesDialog::AddButton(TSharedPtr<FPackageButton> Button)
 	.Padding( 2 )
 	[
 		SNew(SButton) 
+			.ButtonStyle(&FAppStyle::Get(), Button->GetType() == DRT_Save ? "PrimaryButton" : "Button")
+			.TextStyle(&FAppStyle::Get(), "DialogButtonText")
 			.Text(Button->GetName())
 			.ContentPadding(FMargin(10, 3))
 			.ToolTipText(Button->GetToolTip())
