Remove broken drop shadow from content browser list view

diff --git a/Engine/Source/Editor/ContentBrowser/Private/AssetViewWidgets.cpp b/Engine/Source/Editor/ContentBrowser/Private/AssetViewWidgets.cpp
--- a/Engine/Source/Editor/ContentBrowser/Private/AssetViewWidgets.cpp
+++ b/Engine/Source/Editor/ContentBrowser/Private/AssetViewWidgets.cpp
@@ -1479,13 +1479,7 @@ void SAssetListItem::Construct( const FArguments& InArgs )
 				.WidthOverride( this, &SAssetListItem::GetThumbnailBoxSize )
 				.HeightOverride( this, &SAssetListItem::GetThumbnailBoxSize )
 				[
-					// Drop shadow border
-					SNew(SBorder)
-					.Padding(4.f)
-					.BorderImage(FEditorStyle::GetBrush(ItemShadowBorderName))
-					[
-						ItemContents
-					]
+					ItemContents
 				]
 			]
 
