Fixed the initial Content Browser selection sometimes being "All Assets"

This could happen as SetSelectedPaths wouldn't update LastSelectedPaths, which would result in the next call to Populate clearing the selection unless LoadSettings had provided a selection (which it doesn't by default for an empty config).

This change has SetSelectedPaths also update LastSelectedPaths, and LoadSettings will now clear that again when applying its settings, including if there are no paths saved (meaning "All Assets" was actually requested).

#rnx

diff --git a/Engine/Source/Editor/ContentBrowser/Private/SPathView.cpp b/Engine/Source/Editor/ContentBrowser/Private/SPathView.cpp
--- a/Engine/Source/Editor/ContentBrowser/Private/SPathView.cpp
+++ b/Engine/Source/Editor/ContentBrowser/Private/SPathView.cpp
@@ -241,6 +241,7 @@ void SPathView::SetSelectedPaths(const TArray<FString>& Paths)
 	PendingInitialPaths.Empty();
 
 	// Clear the selection to start, then add the selected paths as they are found
+	LastSelectedPaths.Empty();
 	TreeViewPtr->ClearSelection();
 
 	for (int32 PathIdx = 0; PathIdx < Paths.Num(); ++PathIdx)
@@ -302,6 +303,7 @@ void SPathView::SetSelectedPaths(const TArray<FString>& Paths)
 				}
 
 				// Set the selection to the closest found folder and scroll it into view
+				LastSelectedPaths.Add(TreeItems.Last()->GetItem().GetVirtualPath());
 				TreeViewPtr->SetItemSelection(TreeItems.Last(), true);
 				TreeViewPtr->RequestScrollIntoView(TreeItems.Last());
 			}
@@ -785,68 +787,39 @@ void SPathView::LoadSettings(const FString& IniFilename, const FString& IniSecti
 {
 	// Selected Paths
 	FString SelectedPathsString;
-	TArray<FString> NewSelectedPaths;
 	if ( GConfig->GetString(*IniSection, *(SettingsString + TEXT(".SelectedPaths")), SelectedPathsString, IniFilename) )
 	{
+		TArray<FString> NewSelectedPaths;
 		SelectedPathsString.ParseIntoArray(NewSelectedPaths, TEXT(","), /*bCullEmpty*/true);
-	}
 
-	if ( NewSelectedPaths.Num() > 0 )
-	{
 		UContentBrowserDataSubsystem* ContentBrowserData = IContentBrowserDataModule::Get().GetSubsystem();
 		const bool bDiscoveringAssets = ContentBrowserData->IsDiscoveringItems();
 
+		// Batch the selection changed event
+		FScopedSelectionChangedEvent ScopedSelectionChangedEvent(SharedThis(this));
+
 		if ( bDiscoveringAssets )
 		{
-			// Keep track if we changed at least one source so we know to fire the bulk selection changed delegate later
-			bool bSelectedAtLeastOnePath = false;
+			// Clear any previously selected paths
+			LastSelectedPaths.Empty();
+			TreeViewPtr->ClearSelection();
 
+			// If the selected paths is empty, the path was "All assets"
+			// This should handle that case properly
+			for (int32 PathIdx = 0; PathIdx < NewSelectedPaths.Num(); ++PathIdx)
 			{
-				// Prevent the selection changed delegate since we are selecting one path at a time. A bulk event will be fired later if needed.
-				FScopedPreventTreeItemChangedDelegate DelegatePrevention( SharedThis(this) );				
-
-				// Clear any previously selected paths
-				TreeViewPtr->ClearSelection();
-
-				// If the selected paths is empty, the path was "All assets"
-				// This should handle that case properly
-				for (int32 PathIdx = 0; PathIdx < NewSelectedPaths.Num(); ++PathIdx)
+				const FName Path = *NewSelectedPaths[PathIdx];
+				if ( !ExplicitlyAddPathToSelection(Path) )
 				{
-					const FName Path = *NewSelectedPaths[PathIdx];
-					if ( ExplicitlyAddPathToSelection(Path) )
-					{
-						bSelectedAtLeastOnePath = true;
-					}
-					else
-					{
-						// If we could not initially select these paths, but are still discovering assets, add them to a pending list to select them later
-						PendingInitialPaths.Add(Path);
-					}
+					// If we could not initially select these paths, but are still discovering assets, add them to a pending list to select them later
+					PendingInitialPaths.Add(Path);
 				}
 			}
-
-			if ( bSelectedAtLeastOnePath )
-			{
-				// Send the first selected item with the notification
-				const TArray<TSharedPtr<FTreeItem>> SelectedItems = TreeViewPtr->GetSelectedItems();
-				check(SelectedItems.Num() > 0);
-
-				// Signal a single selection changed event to let any listeners know that paths have changed
-				TreeSelectionChanged( SelectedItems[0], ESelectInfo::Direct );
-			}
 		}
 		else
 		{
 			// If all assets are already discovered, just select paths the best we can
 			SetSelectedPaths(NewSelectedPaths);
-
-			// Send the first selected item with the notification
-			const TArray<TSharedPtr<FTreeItem>> SelectedItems = TreeViewPtr->GetSelectedItems();
-			if (SelectedItems.Num() > 0)
-			{
-				// Signal a single selection changed event to let any listeners know that paths have changed
-				TreeSelectionChanged( SelectedItems[0], ESelectInfo::Direct );
-			}
 		}
 	}
 }
