Re-added in-place frontend filter of recently modified Content Browser items

This can avoid a costly content re-sort if the item still passes the frontend filters

#rnx

diff --git a/Engine/Source/Editor/ContentBrowser/Private/SAssetView.cpp b/Engine/Source/Editor/ContentBrowser/Private/SAssetView.cpp
--- a/Engine/Source/Editor/ContentBrowser/Private/SAssetView.cpp
+++ b/Engine/Source/Editor/ContentBrowser/Private/SAssetView.cpp
@@ -98,6 +98,61 @@ static FText ThumbnailSizeToDisplayName(EThumbnailSize InSize)
 	}
 }
 
+class FAssetViewFrontendFilterHelper
+{
+public:
+	explicit FAssetViewFrontendFilterHelper(SAssetView* InAssetView)
+		: AssetView(InAssetView)
+		, ContentBrowserData(IContentBrowserDataModule::Get().GetSubsystem())
+		, bDisplayEmptyFolders(AssetView->IsShowingEmptyFolders())
+	{
+	}
+
+	bool DoesItemPassQueryFilter(const TSharedPtr<FAssetViewItem>& InItemToFilter)
+	{
+		// Folders aren't subject to additional filtering
+		if (InItemToFilter->IsFolder())
+		{
+			return true;
+		}
+
+		// If we have OnShouldFilterAsset then it is assumed that we really only want to see true assets and 
+		// nothing else so only include things that have asset data and also pass the query filter
+		FAssetData ItemAssetData;
+		if (InItemToFilter->GetItem().Legacy_TryGetAssetData(ItemAssetData))
+		{
+			if (!AssetView->OnShouldFilterAsset.Execute(ItemAssetData))
+			{
+				return true;
+			}
+		}
+
+		return false;
+	}
+
+	bool DoesItemPassFrontendFilter(const TSharedPtr<FAssetViewItem>& InItemToFilter)
+	{
+		// Folders are only subject to "empty" filtering
+		if (InItemToFilter->IsFolder())
+		{
+			return bDisplayEmptyFolders || ContentBrowserData->IsFolderVisibleIfHidingEmpty(InItemToFilter->GetItem().GetVirtualPath());
+		}
+
+		// Run the item through the filters
+		if (!AssetView->IsFrontendFilterActive() || AssetView->PassesCurrentFrontendFilter(InItemToFilter->GetItem()))
+		{
+			return true;
+		}
+
+		return false;
+	}
+
+private:
+	SAssetView* AssetView = nullptr;
+	UContentBrowserDataSubsystem* ContentBrowserData = nullptr;
+	const bool bDisplayEmptyFolders = true;
+};
+
 SAssetView::~SAssetView()
 {
 	if (IContentBrowserDataModule* ContentBrowserDataModule = IContentBrowserDataModule::GetPtr())
@@ -106,6 +161,7 @@ SAssetView::~SAssetView()
 		{
 			ContentBrowserData->OnItemDataUpdated().RemoveAll(this);
 			ContentBrowserData->OnItemDataRefreshed().RemoveAll(this);
+			ContentBrowserData->OnItemDataDiscoveryComplete().RemoveAll(this);
 		}
 	}
 
@@ -139,6 +195,7 @@ void SAssetView::Construct( const FArguments& InArgs )
 	UContentBrowserDataSubsystem* ContentBrowserData = IContentBrowserDataModule::Get().GetSubsystem();
 	ContentBrowserData->OnItemDataUpdated().AddSP(this, &SAssetView::HandleItemDataUpdated);
 	ContentBrowserData->OnItemDataRefreshed().AddSP(this, &SAssetView::RequestSlowFullListRefresh);
+	ContentBrowserData->OnItemDataDiscoveryComplete().AddSP(this, &SAssetView::HandleItemDataDiscoveryComplete);
 
 	FCollectionManagerModule& CollectionManagerModule = FCollectionManagerModule::GetModule();
 	CollectionManagerModule.Get().OnAssetsAdded().AddSP( this, &SAssetView::OnAssetsAddedToCollection );
@@ -900,6 +957,7 @@ void SAssetView::Tick( const FGeometry& AllottedGeometry, const double InCurrent
 		bForceViewUpdate = false;
 
 		const double TickStartTime = FPlatformTime::Seconds();
+		const bool bWasWorking = bIsWorking;
 
 		// Mark the first amortize time
 		if (AmortizeStartTime == 0)
@@ -944,7 +1002,7 @@ void SAssetView::Tick( const FGeometry& AllottedGeometry, const double InCurrent
 				SetMajorityAssetType(HighestType);
 			}
 
-			if (bPendingSortFilteredItems)
+			if (bPendingSortFilteredItems && (bWasWorking || (InCurrentTime > LastSortTime + SortDelaySeconds)))
 			{
 				// Don't sync to selection if we are just going to do it below
 				SortList(!PendingSyncItems.Num());
@@ -1123,47 +1181,7 @@ void SAssetView::ProcessItemsPendingFilter(const double TickStartTime)
 {
 	const double ProcessItemsPendingFilterStartTime = FPlatformTime::Seconds();
 
-	UContentBrowserDataSubsystem* ContentBrowserData = IContentBrowserDataModule::Get().GetSubsystem();
-
-	auto DoesItemPassQueryFilter = [this](const TSharedPtr<FAssetViewItem>& InItemToFilter)
-	{
-		// Folders aren't subject to additional filtering
-		if (InItemToFilter->IsFolder())
-		{
-			return true;
-		}
-
-		// If we have OnShouldFilterAsset then it is assumed that we really only want to see true assets and 
-		// nothing else so only include things that have asset data and also pass the query filter
-		FAssetData ItemAssetData;
-		if (InItemToFilter->GetItem().Legacy_TryGetAssetData(ItemAssetData))
-		{
-			if (!OnShouldFilterAsset.Execute(ItemAssetData))
-			{
-				return true;
-			}
-		}
-
-		return false;
-	};
-
-	const bool bDisplayEmpty = IsShowingEmptyFolders();
-	auto DoesItemPassFrontendFilter = [this, bDisplayEmpty, ContentBrowserData](const TSharedPtr<FAssetViewItem>& InItemToFilter)
-	{
-		// Folders are only subject to "empty" filtering
-		if (InItemToFilter->IsFolder())
-		{
-			return bDisplayEmpty || ContentBrowserData->IsFolderVisibleIfHidingEmpty(InItemToFilter->GetItem().GetVirtualPath());
-		}
-
-		// Run the item through the filters
-		if (!IsFrontendFilterActive() || PassesCurrentFrontendFilter(InItemToFilter->GetItem()))
-		{
-			return true;
-		}
-
-		return false;
-	};
+	FAssetViewFrontendFilterHelper FrontendFilterHelper(this);
 
 	auto UpdateFilteredAssetItemTypeCounts = [this](const TSharedPtr<FAssetViewItem>& InItem)
 	{
@@ -1183,12 +1201,12 @@ void SAssetView::ProcessItemsPendingFilter(const double TickStartTime)
 	bool bRefreshList = false;
 	bool bHasTimeRemaining = true;
 
-	auto FilterItem = [this, bRunQueryFilter, &bRefreshList, &DoesItemPassQueryFilter, &DoesItemPassFrontendFilter, &UpdateFilteredAssetItemTypeCounts](const TSharedPtr<FAssetViewItem>& ItemToFilter)
+	auto FilterItem = [this, bRunQueryFilter, &bRefreshList, &FrontendFilterHelper, &UpdateFilteredAssetItemTypeCounts](const TSharedPtr<FAssetViewItem>& ItemToFilter)
 	{
 		// Run the query filter if required
 		if (bRunQueryFilter)
 		{
-			const bool bPassedBackendFilter = DoesItemPassQueryFilter(ItemToFilter);
+			const bool bPassedBackendFilter = FrontendFilterHelper.DoesItemPassQueryFilter(ItemToFilter);
 			if (!bPassedBackendFilter)
 			{
 				AvailableBackendItems.Remove(FContentBrowserItemKey(ItemToFilter->GetItem()));
@@ -1198,7 +1216,7 @@ void SAssetView::ProcessItemsPendingFilter(const double TickStartTime)
 
 		// Run the frontend filter
 		{
-			const bool bPassedFrontendFilter = DoesItemPassFrontendFilter(ItemToFilter);
+			const bool bPassedFrontendFilter = FrontendFilterHelper.DoesItemPassFrontendFilter(ItemToFilter);
 			if (bPassedFrontendFilter)
 			{
 				checkAssetList(!FilteredAssetItems.Contains(ItemToFilter));
@@ -4169,8 +4187,9 @@ void SAssetView::HandleItemDataUpdated(TArrayView<const FContentBrowserItemDataU
 	}
 
 	bool bRefreshView = false;
+	TSet<TSharedPtr<FAssetViewItem>> ItemsPendingInplaceFrontendFilter;
 
-	auto AddItem = [this](const FContentBrowserItemKey& InItemDataKey, const FContentBrowserItemData& InItemData)
+	auto AddItem = [this, &ItemsPendingInplaceFrontendFilter](const FContentBrowserItemKey& InItemDataKey, const FContentBrowserItemData& InItemData)
 	{
 		TSharedPtr<FAssetViewItem>& ItemToUpdate = AvailableBackendItems.FindOrAdd(InItemDataKey);
 		if (ItemToUpdate)
@@ -4181,11 +4200,9 @@ void SAssetView::HandleItemDataUpdated(TArrayView<const FContentBrowserItemDataU
 			// Update the custom column data
 			ItemToUpdate->CacheCustomColumns(CustomColumns, true, true, true);
 
-			// Remove this item from the list view data until it has been filtered again
-			FilteredAssetItems.RemoveSingle(ItemToUpdate);
-
-			// This item was modified, so put it in this prioritized set of items to be filtered
-			ItemsPendingPriorityFilter.Add(ItemToUpdate);
+			// This item was modified, so put it in the list of items to be in-place re-tested against the active frontend filter (this can avoid a costly re-sort of the view)
+			// If the item can't be queried in-place (because the item isn't in the view) then it will be added to ItemsPendingPriorityFilter instead
+			ItemsPendingInplaceFrontendFilter.Add(ItemToUpdate);
 		}
 		else
 		{
@@ -4196,7 +4213,7 @@ void SAssetView::HandleItemDataUpdated(TArrayView<const FContentBrowserItemDataU
 		}
 	};
 
-	auto RemoveItem = [this, &bRefreshView](const FContentBrowserItemKey& InItemDataKey, const FContentBrowserItemData& InItemData)
+	auto RemoveItem = [this, &bRefreshView, &ItemsPendingInplaceFrontendFilter](const FContentBrowserItemKey& InItemDataKey, const FContentBrowserItemData& InItemData)
 	{
 		const uint32 ItemDataKeyHash = GetTypeHash(InItemDataKey);
 
@@ -4220,6 +4237,7 @@ void SAssetView::HandleItemDataUpdated(TArrayView<const FContentBrowserItemDataU
 			FilteredAssetItems.RemoveSingle(ItemToRemove);
 			ItemsPendingPriorityFilter.RemoveByHash(ItemToRemoveHash, ItemToRemove);
 			ItemsPendingFrontendFilter.RemoveByHash(ItemToRemoveHash, ItemToRemove);
+			ItemsPendingInplaceFrontendFilter.RemoveByHash(ItemToRemoveHash, ItemToRemove);
 
 			// Need to refresh manually after removing items, as adding relies on the pending filter lists to trigger this
 			bRefreshView = true;
@@ -4239,6 +4257,7 @@ void SAssetView::HandleItemDataUpdated(TArrayView<const FContentBrowserItemDataU
 		return false;
 	};
 
+	// Process the main set of updates
 	for (const FContentBrowserItemDataUpdate& ItemDataUpdate : InUpdatedItems)
 	{
 		const FContentBrowserItemData& ItemData = ItemDataUpdate.GetItemData();
@@ -4285,6 +4304,59 @@ void SAssetView::HandleItemDataUpdated(TArrayView<const FContentBrowserItemDataU
 		}
 	}
 
+	// Now patch in the in-place frontend filter requests (if possible)
+	if (ItemsPendingInplaceFrontendFilter.Num() > 0)
+	{
+		FAssetViewFrontendFilterHelper FrontendFilterHelper(this);
+		const bool bRunQueryFilter = OnShouldFilterAsset.IsBound();
+
+		for (auto It = FilteredAssetItems.CreateIterator(); It && ItemsPendingInplaceFrontendFilter.Num() > 0; ++It)
+		{
+			const TSharedPtr<FAssetViewItem> ItemToFilter = *It;
+
+			if (ItemsPendingInplaceFrontendFilter.Remove(ItemToFilter) > 0)
+			{
+				bool bRemoveItem = false;
+
+				// Run the query filter if required
+				if (bRunQueryFilter)
+				{
+					const bool bPassedBackendFilter = FrontendFilterHelper.DoesItemPassQueryFilter(ItemToFilter);
+					if (!bPassedBackendFilter)
+					{
+						bRemoveItem = true;
+						AvailableBackendItems.Remove(FContentBrowserItemKey(ItemToFilter->GetItem()));
+					}
+				}
+
+				// Run the frontend filter
+				if (!bRemoveItem)
+				{
+					const bool bPassedFrontendFilter = FrontendFilterHelper.DoesItemPassFrontendFilter(ItemToFilter);
+					if (!bPassedFrontendFilter)
+					{
+						bRemoveItem = true;
+					}
+				}
+
+				// Remove this item?
+				if (bRemoveItem)
+				{
+					bRefreshView = true;
+					It.RemoveCurrent();
+				}
+			}
+		}
+
+		// Do we still have items that could not be in-place filtered?
+		// If so, add them to ItemsPendingPriorityFilter so they are processed into the view ASAP
+		if (ItemsPendingInplaceFrontendFilter.Num() > 0)
+		{
+			ItemsPendingPriorityFilter.Append(MoveTemp(ItemsPendingInplaceFrontendFilter));
+			ItemsPendingInplaceFrontendFilter.Reset();
+		}
+	}
+
 	if (bRefreshView)
 	{
 		RefreshList();
@@ -4293,6 +4365,15 @@ void SAssetView::HandleItemDataUpdated(TArrayView<const FContentBrowserItemDataU
 	UE_LOG(LogContentBrowser, VeryVerbose, TEXT("AssetView - HandleItemDataUpdated completed in %0.4f seconds for %d items (%d available items)"), FPlatformTime::Seconds() - HandleItemDataUpdatedStartTime, InUpdatedItems.Num(), AvailableBackendItems.Num());
 }
 
+void SAssetView::HandleItemDataDiscoveryComplete()
+{
+	if (bPendingSortFilteredItems)
+	{
+		// If we have a sort pending, then force this to happen next frame now that discovery has finished
+		LastSortTime = 0;
+	}
+}
+
 #undef checkAssetList
 #undef ASSET_VIEW_PARANOIA_LIST_CHECKS
 
diff --git a/Engine/Source/Editor/ContentBrowser/Private/SAssetView.h b/Engine/Source/Editor/ContentBrowser/Private/SAssetView.h
--- a/Engine/Source/Editor/ContentBrowser/Private/SAssetView.h
+++ b/Engine/Source/Editor/ContentBrowser/Private/SAssetView.h
@@ -772,9 +772,15 @@ private:
 
 	FContentBrowserDataFilter CreateBackendDataFilter() const;
 
+	/** Handles updating the view when content items are changed */
 	void HandleItemDataUpdated(TArrayView<const FContentBrowserItemDataUpdate> InUpdatedItems);
 
+	/** Notification for when the content browser has completed it's initial search */
+	void HandleItemDataDiscoveryComplete();
+
 private:
+	friend class FAssetViewFrontendFilterHelper;
+
 	/** The available items from querying the backend data sources */
 	TMap<FContentBrowserItemKey, TSharedPtr<FAssetViewItem>> AvailableBackendItems;
 
