Update content browser search and filter styling

diff --git a/Engine/Source/Editor/ContentBrowser/Private/SAssetView.cpp b/Engine/Source/Editor/ContentBrowser/Private/SAssetView.cpp
--- a/Engine/Source/Editor/ContentBrowser/Private/SAssetView.cpp
+++ b/Engine/Source/Editor/ContentBrowser/Private/SAssetView.cpp
@@ -2354,19 +2354,6 @@ void SAssetView::PopulateViewButtonMenu(UToolMenu* Menu)
 			EUserInterfaceActionType::ToggleButton
 		);
 
-		Section.AddMenuEntry(
-			"DockCollections",
-			LOCTEXT("DockCollectionsOptions", "Dock Collections"),
-			LOCTEXT("DockCollectionsOptionToolTip", "Dock the collections view under the path view?"),
-			FSlateIcon(),
-			FUIAction(
-				FExecuteAction::CreateSP(this, &SAssetView::ToggleDockCollections),
-				FCanExecuteAction::CreateSP(this, &SAssetView::IsToggleDockCollectionsAllowed),
-				FIsActionChecked::CreateSP(this, &SAssetView::HasDockedCollections)
-			),
-			EUserInterfaceActionType::ToggleButton
-		);
-
 		Section.AddMenuEntry(
 			"FilterRecursively",
 			LOCTEXT("FilterRecursivelyOption", "Filter Recursively"),
diff --git a/Engine/Source/Editor/ContentBrowser/Private/SContentBrowser.cpp b/Engine/Source/Editor/ContentBrowser/Private/SContentBrowser.cpp
--- a/Engine/Source/Editor/ContentBrowser/Private/SContentBrowser.cpp
+++ b/Engine/Source/Editor/ContentBrowser/Private/SContentBrowser.cpp
@@ -213,105 +213,107 @@ void SContentBrowser::Construct( const FArguments& InArgs, const FName& InInstan
 			.Padding( FMargin( 3 ) )
 			.BorderImage(FAppStyle::Get().GetBrush("Brushes.Background"))
 			[
-				SNew( SWrapBox )
-				.UseAllottedSize( true )
-				.InnerSlotPadding( FVector2D( 5, 2 ) )
-				+ SWrapBox::Slot()
-				.FillLineWhenSizeLessThan( 600 )
-				.FillEmptySpace( true )
+				SNew(SHorizontalBox)
+				+ SHorizontalBox::Slot()
+				.AutoWidth()
 				.HAlign(HAlign_Left)
+				.Padding(5, 0, 0, 0)
 				[
 					CreateToolBar(Config)
 				]
-				+ SWrapBox::Slot()
-				.FillEmptySpace(true)
+				// History Back Button
+				+SHorizontalBox::Slot()
+				.Padding(10, 0, 0, 0)
+				.AutoWidth()
 				[
-					SNew(SHorizontalBox)
-					// Search
-					+ SHorizontalBox::Slot()
-					.VAlign(VAlign_Center)
-					.Padding(5, 0, 0, 0)
-					.FillWidth(1.0f)
-					[
-						SAssignNew(SearchBoxPtr, SAssetSearchBox)
-						.HintText(this, &SContentBrowser::GetSearchAssetsHintText)
-						.OnTextChanged(this, &SContentBrowser::OnSearchBoxChanged)
-						.OnTextCommitted(this, &SContentBrowser::OnSearchBoxCommitted)
-						.OnKeyDownHandler(this, &SContentBrowser::OnSearchKeyDown)
-						.OnAssetSearchBoxSuggestionFilter(this, &SContentBrowser::OnAssetSearchSuggestionFilter)
-						.OnAssetSearchBoxSuggestionChosen(this, &SContentBrowser::OnAssetSearchSuggestionChosen)
-						.DelayChangeNotificationsWhileTyping(true)
-						.Visibility((Config != nullptr ? Config->bCanShowAssetSearch : true) ? EVisibility::Visible : EVisibility::Collapsed)
-						.AddMetaData<FTagMetaData>(FTagMetaData(TEXT("ContentBrowserSearchAssets")))
-					]
-					
-					// Save Search
-					+ SHorizontalBox::Slot()
-					.AutoWidth()
-					.VAlign(VAlign_Center)
-					.Padding(2.0f, 0.0f, 0.0f, 0.0f)
-					[
-						SNew(SButton)
-						.ButtonStyle(FAppStyle::Get(), "SimpleButton")
-						.ToolTipText(LOCTEXT("SaveSearchButtonTooltip", "Save the current search as a dynamic collection."))
-						.IsEnabled(this, &SContentBrowser::IsSaveSearchButtonEnabled)
-						.OnClicked(this, &SContentBrowser::OnSaveSearchButtonClicked)
-						.ContentPadding(FMargin(1, 1))
-						.Visibility((Config != nullptr ? Config->bCanShowAssetSearch : true) ? EVisibility::Visible : EVisibility::Collapsed)
-						[
-							SNew(STextBlock)
-							.TextStyle(FEditorStyle::Get(), "GenericFilters.TextStyle")
-							.Font(FEditorStyle::Get().GetFontStyle("FontAwesome.10"))
-							.Text(FEditorFontGlyphs::Floppy_O)
-						]
-					]
-					+ SHorizontalBox::Slot()
-					.AutoWidth()
-					.HAlign(HAlign_Left)
-					.Padding(5, 0, 0, 0)
-					[
-						SNew(SComboButton)
-						.ComboButtonStyle(&FAppStyle::Get().GetWidgetStyle<FComboButtonStyle>("SimpleComboButton"))
-						.ForegroundColor(FSlateColor::UseStyle())
-						.ToolTipText(LOCTEXT("AddFilterToolTip", "Add an asset filter."))
-						.OnGetMenuContent(this, &SContentBrowser::MakeAddFilterMenu)
-						.ContentPadding(FMargin(1, 0))
-						.AddMetaData<FTagMetaData>(FTagMetaData(TEXT("ContentBrowserFiltersCombo")))
-						.Visibility((Config != nullptr ? Config->bCanShowFilters : true) ? EVisibility::Visible : EVisibility::Collapsed)
-						.ButtonContent()
-						[
-							SNew(SImage)
-							.Image(FAppStyle::Get().GetBrush("Icons.Filter"))
-							.ColorAndOpacity(FSlateColor::UseForeground())
-						]
-					]
-					+ SHorizontalBox::Slot()
-					.Padding(5.0f,0.0f,0.0f,0.0f)
-					.VAlign(VAlign_Center)
-					.FillWidth(1.5f)
+					SNew(SButton)
+					.VAlign(EVerticalAlignment::VAlign_Center)
+					.ButtonStyle(FEditorStyle::Get(), "SimpleButton")
+					.ToolTipText( this, &SContentBrowser::GetHistoryBackTooltip )
+					.ContentPadding( FMargin(1, 0) )
+					.OnClicked(this, &SContentBrowser::BackClicked)
+					.IsEnabled(this, &SContentBrowser::IsBackEnabled)
+					.AddMetaData<FTagMetaData>(FTagMetaData(TEXT("ContentBrowserHistoryBack")))
 					[
-						SAssignNew(FilterListPtr, SFilterList)
-						.OnFilterChanged(this, &SContentBrowser::OnFilterChanged)
-						.OnGetContextMenu(this, &SContentBrowser::GetFilterContextMenu)
-						.Visibility((Config != nullptr ? Config->bCanShowFilters : true) ? EVisibility::Visible : EVisibility::Collapsed)
-						.FrontendFilters(FrontendFilters)
-						.AddMetaData<FTagMetaData>(FTagMetaData(TEXT("ContentBrowserFilters")))
+						SNew(SImage)
+						.Image(FAppStyle::Get().GetBrush("Icons.CircleArrowLeft"))
+						.ColorAndOpacity(FSlateColor::UseForeground())
 					]
-					+ SHorizontalBox::Slot()
-					.AutoWidth()
-					.HAlign(HAlign_Right)
-					.VAlign(VAlign_Center)
+				]
+				// History Forward Button
+				+ SHorizontalBox::Slot()
+				.Padding(2, 0, 0, 0)
+				.AutoWidth()
+				[
+					SNew(SButton)
+					.VAlign(EVerticalAlignment::VAlign_Center)
+					.ButtonStyle(FEditorStyle::Get(), "SimpleButton")
+					.ToolTipText( this, &SContentBrowser::GetHistoryForwardTooltip )
+					.ContentPadding( FMargin(1, 0) )
+					.OnClicked(this, &SContentBrowser::ForwardClicked)
+					.IsEnabled(this, &SContentBrowser::IsForwardEnabled)
+					.AddMetaData<FTagMetaData>(FTagMetaData(TEXT("ContentBrowserHistoryForward")))
 					[
-						CreateLockButton(Config)
+						SNew(SImage)
+						.Image(FAppStyle::Get().GetBrush("Icons.CircleArrowRight"))
+						.ColorAndOpacity(FSlateColor::UseForeground())
 					]
-					+ SHorizontalBox::Slot()
-					.AutoWidth()
-					.Padding(5.0f, 0.0f, 0.0f, 0.0f)
-					.HAlign(HAlign_Right)
+				]
+				// Path picker
+				+ SHorizontalBox::Slot()
+				.Padding(2, 0, 0, 0)
+				.AutoWidth()
+				.VAlign(VAlign_Fill)
+				[
+					SAssignNew( PathPickerButton, SComboButton )
+					.Visibility( ( Config != nullptr ? Config->bUsePathPicker : true ) ? EVisibility::Visible : EVisibility::Collapsed )
+					.ButtonStyle(FEditorStyle::Get(), "SimpleButton")
+					.ToolTipText( LOCTEXT( "PathPickerTooltip", "Choose a path" ) )
+					.OnGetMenuContent( this, &SContentBrowser::GetPathPickerContent )
+					.HasDownArrow( false )
+					.AddMetaData<FTagMetaData>(FTagMetaData(TEXT("ContentBrowserPathPicker")))
+					.ContentPadding(FMargin(1, 0))
+					.ButtonContent()
 					[
-						ViewOptions
+						SNew(SImage)
+						.Image(FAppStyle::Get().GetBrush("Icons.FolderClosed"))
+						.ColorAndOpacity(FSlateColor::UseForeground())
 					]
 				]
+
+				// Path
+				+ SHorizontalBox::Slot()
+				.VAlign(VAlign_Center)
+				.HAlign(HAlign_Left)
+				.FillWidth(1.0f)
+				.Padding(2, 0, 0, 0)
+				[
+					SAssignNew(PathBreadcrumbTrail, SBreadcrumbTrail<FString>)
+					.ButtonContentPadding(FMargin(2, 2))
+					.ButtonStyle(FAppStyle::Get(), "SimpleButton")
+					.DelimiterImage(FAppStyle::Get().GetBrush("Icons.ChevronRight"))
+					.TextStyle(FAppStyle::Get(), "NormalText")
+					.ShowLeadingDelimiter(false)
+					.InvertTextColorOnHover(false)
+					.OnCrumbClicked(this, &SContentBrowser::OnPathClicked)
+					.HasCrumbMenuContent(this, &SContentBrowser::OnHasCrumbDelimiterContent)
+					.GetCrumbMenuContent(this, &SContentBrowser::OnGetCrumbDelimiterContent)
+					.AddMetaData<FTagMetaData>(FTagMetaData(TEXT("ContentBrowserPath")))
+				]
+				+ SHorizontalBox::Slot()
+				.AutoWidth()
+				.HAlign(HAlign_Right)
+				.VAlign(VAlign_Center)
+				[
+					CreateLockButton(Config)
+				]
+				+ SHorizontalBox::Slot()
+				.AutoWidth()
+				.Padding(5.0f, 0.0f, 0.0f, 0.0f)
+				.HAlign(HAlign_Right)
+				[
+					ViewOptions
+				]
 			]
 		]
 
@@ -328,7 +330,7 @@ void SContentBrowser::Construct( const FArguments& InArgs, const FName& InInstan
 
 			// Sources View
 			+ SSplitter::Slot()
-			.Value(0.25f)
+			.Value(0.15f)
 			[
 				SNew(SBorder)
 				.Padding(FMargin(0,0))
@@ -627,102 +629,74 @@ TSharedRef<SWidget> SContentBrowser::CreateAssetView(const FContentBrowserConfig
 			.Padding(FMargin(0.0f, 5.0f))
 			[
 				SNew(SHorizontalBox)
-				// Sources panel toggle
+				// Search
 				+ SHorizontalBox::Slot()
-				.AutoWidth()
-				.HAlign(HAlign_Left)
+				.VAlign(VAlign_Center)
 				.Padding(5, 0, 0, 0)
+				.FillWidth(TAttribute<float>(this, &SContentBrowser::GetSearchBoxFillWidth))
 				[
-					SNew(SButton)
-					.VAlign(EVerticalAlignment::VAlign_Center)
-					.ButtonStyle(FAppStyle::Get(), "SimpleButton")
-					.ToolTipText(LOCTEXT("SourcesTreeToggleTooltip", "Show or hide the sources panel"))
-					.ContentPadding(FMargin(1, 0))
-					.OnClicked(this, &SContentBrowser::SourcesViewExpandClicked)
-					[
-						SNew(SImage)
-						.Image(this, &SContentBrowser::GetSourcesToggleImage)
-						.ColorAndOpacity(FSlateColor::UseForeground())
-					]
-				]
-				// History Back Button
-				+SHorizontalBox::Slot()
-				.Padding(2, 0, 0, 0)
-				.AutoWidth()
-				[
-					SNew(SButton)
-					.VAlign(EVerticalAlignment::VAlign_Center)
-					.ButtonStyle(FEditorStyle::Get(), "SimpleButton")
-					.ToolTipText( this, &SContentBrowser::GetHistoryBackTooltip )
-					.ContentPadding( FMargin(1, 0) )
-					.OnClicked(this, &SContentBrowser::BackClicked)
-					.IsEnabled(this, &SContentBrowser::IsBackEnabled)
-					.AddMetaData<FTagMetaData>(FTagMetaData(TEXT("ContentBrowserHistoryBack")))
-					[
-						SNew(SImage)
-						.Image(FAppStyle::Get().GetBrush("Icons.CircleArrowLeft"))
-						.ColorAndOpacity(FSlateColor::UseForeground())
-					]
+					SAssignNew(SearchBoxPtr, SAssetSearchBox)
+					.HintText(this, &SContentBrowser::GetSearchAssetsHintText)
+					.OnTextChanged(this, &SContentBrowser::OnSearchBoxChanged)
+					.OnTextCommitted(this, &SContentBrowser::OnSearchBoxCommitted)
+					.OnKeyDownHandler(this, &SContentBrowser::OnSearchKeyDown)
+					.OnAssetSearchBoxSuggestionFilter(this, &SContentBrowser::OnAssetSearchSuggestionFilter)
+					.OnAssetSearchBoxSuggestionChosen(this, &SContentBrowser::OnAssetSearchSuggestionChosen)
+					.DelayChangeNotificationsWhileTyping(true)
+					.Visibility((Config != nullptr ? Config->bCanShowAssetSearch : true) ? EVisibility::Visible : EVisibility::Collapsed)
+					.AddMetaData<FTagMetaData>(FTagMetaData(TEXT("ContentBrowserSearchAssets")))
 				]
-				// History Forward Button
+					
+				// Save Search
 				+ SHorizontalBox::Slot()
-				.Padding(2, 0, 0, 0)
 				.AutoWidth()
+				.VAlign(VAlign_Center)
+				.Padding(2.0f, 0.0f, 0.0f, 0.0f)
 				[
 					SNew(SButton)
-					.VAlign(EVerticalAlignment::VAlign_Center)
-					.ButtonStyle(FEditorStyle::Get(), "SimpleButton")
-					.ToolTipText( this, &SContentBrowser::GetHistoryForwardTooltip )
-					.ContentPadding( FMargin(1, 0) )
-					.OnClicked(this, &SContentBrowser::ForwardClicked)
-					.IsEnabled(this, &SContentBrowser::IsForwardEnabled)
-					.AddMetaData<FTagMetaData>(FTagMetaData(TEXT("ContentBrowserHistoryForward")))
+					.ButtonStyle(FAppStyle::Get(), "SimpleButton")
+					.ToolTipText(LOCTEXT("SaveSearchButtonTooltip", "Save the current search as a dynamic collection."))
+					.IsEnabled(this, &SContentBrowser::IsSaveSearchButtonEnabled)
+					.OnClicked(this, &SContentBrowser::OnSaveSearchButtonClicked)
+					.ContentPadding(FMargin(1, 1))
+					.Visibility((Config != nullptr ? Config->bCanShowAssetSearch : true) ? EVisibility::Visible : EVisibility::Collapsed)
 					[
-						SNew(SImage)
-						.Image(FAppStyle::Get().GetBrush("Icons.CircleArrowRight"))
-						.ColorAndOpacity(FSlateColor::UseForeground())
+						SNew(STextBlock)
+						.TextStyle(FEditorStyle::Get(), "GenericFilters.TextStyle")
+						.Font(FEditorStyle::Get().GetFontStyle("FontAwesome.10"))
+						.Text(FEditorFontGlyphs::Floppy_O)
 					]
 				]
-				// Path picker
 				+ SHorizontalBox::Slot()
-				.Padding(2, 0, 0, 0)
 				.AutoWidth()
-				.VAlign(VAlign_Fill)
+				.HAlign(HAlign_Left)
+				.Padding(5, 0, 0, 0)
 				[
-					SAssignNew( PathPickerButton, SComboButton )
-					.Visibility( ( Config != nullptr ? Config->bUsePathPicker : true ) ? EVisibility::Visible : EVisibility::Collapsed )
-					.ButtonStyle(FEditorStyle::Get(), "SimpleButton")
-					.ToolTipText( LOCTEXT( "PathPickerTooltip", "Choose a path" ) )
-					.OnGetMenuContent( this, &SContentBrowser::GetPathPickerContent )
-					.HasDownArrow( false )
-					.AddMetaData<FTagMetaData>(FTagMetaData(TEXT("ContentBrowserPathPicker")))
+					SNew(SComboButton)
+					.ComboButtonStyle(&FAppStyle::Get().GetWidgetStyle<FComboButtonStyle>("SimpleComboButton"))
+					.ForegroundColor(FSlateColor::UseStyle())
+					.ToolTipText(LOCTEXT("AddFilterToolTip", "Add an asset filter."))
+					.OnGetMenuContent(this, &SContentBrowser::MakeAddFilterMenu)
 					.ContentPadding(FMargin(1, 0))
+					.AddMetaData<FTagMetaData>(FTagMetaData(TEXT("ContentBrowserFiltersCombo")))
+					.Visibility((Config != nullptr ? Config->bCanShowFilters : true) ? EVisibility::Visible : EVisibility::Collapsed)
 					.ButtonContent()
 					[
 						SNew(SImage)
-						.Image(FAppStyle::Get().GetBrush("Icons.FolderClosed"))
+						.Image(FAppStyle::Get().GetBrush("Icons.Filter"))
 						.ColorAndOpacity(FSlateColor::UseForeground())
 					]
 				]
-
-				// Path
 				+ SHorizontalBox::Slot()
+				.Padding(5.0f,0.0f,0.0f,0.0f)
 				.VAlign(VAlign_Center)
-				.HAlign(HAlign_Left)
-				.FillWidth(1.0f)
-				.Padding(2, 0, 0, 0)
 				[
-					SAssignNew(PathBreadcrumbTrail, SBreadcrumbTrail<FString>)
-					.ButtonContentPadding(FMargin(2, 2))
-					.ButtonStyle(FAppStyle::Get(), "SimpleButton")
-					.DelimiterImage(FAppStyle::Get().GetBrush("Icons.ChevronRight"))
-					.TextStyle(FAppStyle::Get(), "NormalText")
-					.ShowLeadingDelimiter(false)
-					.InvertTextColorOnHover(false)
-					.OnCrumbClicked(this, &SContentBrowser::OnPathClicked)
-					.HasCrumbMenuContent(this, &SContentBrowser::OnHasCrumbDelimiterContent)
-					.GetCrumbMenuContent(this, &SContentBrowser::OnGetCrumbDelimiterContent)
-					.AddMetaData<FTagMetaData>(FTagMetaData(TEXT("ContentBrowserPath")))
+					SAssignNew(FilterListPtr, SFilterList)
+					.OnFilterChanged(this, &SContentBrowser::OnFilterChanged)
+					.OnGetContextMenu(this, &SContentBrowser::GetFilterContextMenu)
+					.Visibility((Config != nullptr ? Config->bCanShowFilters : true) ? EVisibility::Visible : EVisibility::Collapsed)
+					.FrontendFilters(FrontendFilters)
+					.AddMetaData<FTagMetaData>(FTagMetaData(TEXT("ContentBrowserFilters")))
 				]
 			]
 		]
@@ -874,13 +848,29 @@ void SContentBrowser::ExtendViewOptionsMenu(const FContentBrowserConfig* Config)
 				{
 					if (TSharedPtr<SContentBrowser> ContentBrowser = Context->OwningContentBrowser.Pin())
 					{
-						FToolMenuSection& Section = InMenu->AddSection("Locking", LOCTEXT("LockingMenuHeader", "Locking"), FToolMenuInsert("AssetViewType", EToolMenuInsertType::After));
-						Section.AddMenuEntry(
-							"ToggleLock",
-							TAttribute<FText>(ContentBrowser.ToSharedRef(), &SContentBrowser::GetLockMenuText),
-							LOCTEXT("LockToggleTooltip", "Toggle lock. If locked, this browser will ignore Find in Content Browser requests."),
-							TAttribute<FSlateIcon>(),
-							FUIAction(FExecuteAction::CreateLambda([ContentBrowser = Context->OwningContentBrowser]() {ContentBrowser.Pin()->ToggleLockClicked(); })));
+						{
+							FToolMenuSection& Section = InMenu->AddSection("Locking", LOCTEXT("LockingMenuHeader", "Locking"), FToolMenuInsert("AssetViewType", EToolMenuInsertType::After));
+							Section.AddMenuEntry(
+								"ToggleLock",
+								TAttribute<FText>(ContentBrowser.ToSharedRef(), &SContentBrowser::GetLockMenuText),
+								LOCTEXT("LockToggleTooltip", "Toggle lock. If locked, this browser will ignore Find in Content Browser requests."),
+								TAttribute<FSlateIcon>(),
+								FUIAction(FExecuteAction::CreateLambda([ContentBrowser = Context->OwningContentBrowser]() {ContentBrowser.Pin()->ToggleLockClicked(); })));
+						}
+
+						{
+							FToolMenuSection& Section = InMenu->FindOrAddSection("View");
+							Section.AddMenuEntry(
+								"ToggleSources",
+								LOCTEXT("ToggleSourcesView", "Show Sources Panel"),
+								LOCTEXT("ToggleSourcesView_Tooltip", "Show or hide the sources panel"),
+								TAttribute<FSlateIcon>(),
+								FUIAction(
+									FExecuteAction::CreateLambda([ContentBrowser = Context->OwningContentBrowser]() {ContentBrowser.Pin()->SourcesViewExpandClicked(); }),
+									FCanExecuteAction(),
+									FIsActionChecked::CreateLambda([ContentBrowser = Context->OwningContentBrowser]() {return ContentBrowser.Pin()->bSourcesViewExpanded; })),
+								EUserInterfaceActionType::Check);
+						}
 					}
 				}
 			}
@@ -1019,6 +1009,20 @@ SSplitter::ESizeRule SContentBrowser::GetCollectionsAreaSizeRule() const
 	return CollectionArea->IsExpanded() ? SSplitter::ESizeRule::FractionOfParent : SSplitter::ESizeRule::SizeToContent;
 }
 
+float SContentBrowser::GetSearchBoxFillWidth() const
+{
+	// Gives more room to the search box when the  content browser is constrained to a small space
+	float LocalWidth = GetTickSpaceGeometry().GetLocalSize().X;
+	if (LocalWidth < 1000.f)
+	{
+		return 1.0f;
+	}
+	else
+	{
+		return .35f;
+	}
+}
+
 FText SContentBrowser::GetHighlightedText() const
 {
 	return TextFilter->GetRawFilterText();
@@ -1389,27 +1393,8 @@ void SContentBrowser::LoadSettings(const FName& InInstanceName)
 {
 	FString SettingsString = InInstanceName.ToString();
 
-	// Test to see if we should load legacy settings from a previous instance name
-	// First make sure there aren't any existing settings with the given instance name
-	bool TestBool;
-	if ( !GConfig->GetBool(*SettingsIniSection, *(SettingsString + TEXT(".SourcesExpanded")), TestBool, GEditorPerProjectIni) )
-	{
-		// If there were not any settings and we are Content Browser 1, see if we have any settings under the legacy name "LevelEditorContentBrowser"
-		if ( InInstanceName.ToString() == TEXT("ContentBrowserTab1") && GConfig->GetBool(*SettingsIniSection, TEXT("LevelEditorContentBrowser.SourcesExpanded"), TestBool, GEditorPerProjectIni) )
-		{
-			// We have found some legacy settings with the old ID, use them. These settings will be saved out to the new id later
-			SettingsString = TEXT("LevelEditorContentBrowser");
-		}
-		// else see if we are Content Browser 2, and see if we have any settings under the legacy name "MajorContentBrowserTab"
-		else if ( InInstanceName.ToString() == TEXT("ContentBrowserTab2") && GConfig->GetBool(*SettingsIniSection, TEXT("MajorContentBrowserTab.SourcesExpanded"), TestBool, GEditorPerProjectIni) )
-		{
-			// We have found some legacy settings with the old ID, use them. These settings will be saved out to the new id later
-			SettingsString = TEXT("MajorContentBrowserTab");
-		}
-	}
-
-
 	// Now that we have determined the appropriate settings string, actually load the settings
+	bSourcesViewExpanded = true;
 	GConfig->GetBool(*SettingsIniSection, *(SettingsString + TEXT(".SourcesExpanded")), bSourcesViewExpanded, GEditorPerProjectIni);
 	GConfig->GetBool(*SettingsIniSection, *(SettingsString + TEXT(".Locked")), bIsLocked, GEditorPerProjectIni);
 
@@ -2392,18 +2377,6 @@ EVisibility SContentBrowser::GetSourcesViewVisibility() const
 	return bSourcesViewExpanded ? EVisibility::Visible : EVisibility::Collapsed;
 }
 
-const FSlateBrush* SContentBrowser::GetSourcesToggleImage() const
-{
-	if ( bSourcesViewExpanded )
-	{
-		return FEditorStyle::GetBrush("ContentBrowser.HideSourcesView");
-	}
-	else
-	{
-		return FEditorStyle::GetBrush("ContentBrowser.ShowSourcesView");
-	}
-}
-
 FReply SContentBrowser::SourcesViewExpandClicked()
 {
 	bSourcesViewExpanded = !bSourcesViewExpanded;
@@ -2419,11 +2392,6 @@ FReply SContentBrowser::SourcesViewExpandClicked()
 	return FReply::Handled();
 }
 
-EVisibility SContentBrowser::GetPathExpanderVisibility() const
-{
-	return bSourcesViewExpanded ? EVisibility::Collapsed : EVisibility::Visible;
-}
-
 EVisibility SContentBrowser::GetSourcesSwitcherVisibility() const
 {
 	return GetDefault<UContentBrowserSettings>()->GetDockCollections() ? EVisibility::Collapsed : EVisibility::Visible;
diff --git a/Engine/Source/Editor/ContentBrowser/Private/SContentBrowser.h b/Engine/Source/Editor/ContentBrowser/Private/SContentBrowser.h
--- a/Engine/Source/Editor/ContentBrowser/Private/SContentBrowser.h
+++ b/Engine/Source/Editor/ContentBrowser/Private/SContentBrowser.h
@@ -330,15 +330,9 @@ private:
 	/** Gets the visibility state of the asset tree */
 	EVisibility GetSourcesViewVisibility() const;
 
-	/** Gets the brush used to on the sources toggle button */
-	const FSlateBrush* GetSourcesToggleImage() const;
-
 	/** Handler for clicking the tree expand/collapse button */
 	FReply SourcesViewExpandClicked();
 
-	/** Gets the visibility of the path expander button */
-	EVisibility GetPathExpanderVisibility() const;
-
 	/** Gets the visibility of the source switch button */
 	EVisibility GetSourcesSwitcherVisibility() const;
 
@@ -519,6 +513,8 @@ private:
 	SSplitter::ESizeRule GetFavoritesAreaSizeRule() const;
 	SSplitter::ESizeRule GetPathAreaSizeRule() const;
 	SSplitter::ESizeRule GetCollectionsAreaSizeRule() const;
+
+	float GetSearchBoxFillWidth() const;
 private:
 
 	/** The tab that contains this browser */
diff --git a/Engine/Source/Editor/ContentBrowser/Private/SFilterList.cpp b/Engine/Source/Editor/ContentBrowser/Private/SFilterList.cpp
--- a/Engine/Source/Editor/ContentBrowser/Private/SFilterList.cpp
+++ b/Engine/Source/Editor/ContentBrowser/Private/SFilterList.cpp
@@ -183,13 +183,13 @@ public:
 		ChildSlot
 		[
 			SNew(SBorder)
- 			.Padding(0)
-			.BorderImage(this, &SFilter::GetFilterBorderBackground)
+ 			.Padding(1.0f)
+			.BorderImage(FAppStyle::Get().GetBrush("ContentBrowser.FilterBackground"))
  			[
 				SAssignNew( ToggleButtonPtr, SFilterCheckBox )
-				.Style(FEditorStyle::Get(), "ContentBrowser.FilterButton")
+				.Style(FAppStyle::Get(), "ContentBrowser.FilterButton")
 				.ToolTipText(FilterToolTip)
-				.Padding(this, &SFilter::GetFilterNamePadding)
+				.Padding(0.0f)
 				.IsChecked(this, &SFilter::IsChecked)
 				.OnCheckStateChanged(this, &SFilter::FilterToggled)
 				.OnGetMenuContent(this, &SFilter::GetRightClickMenuContent)
@@ -197,7 +197,6 @@ public:
 					SNew(SHorizontalBox)
 					+SHorizontalBox::Slot()
 					.VAlign(VAlign_Center)
-					.Padding(2.0f)
 					.AutoWidth()
 					[
 						SNew(SImage)
@@ -205,12 +204,12 @@ public:
 						.ColorAndOpacity(this, &SFilter::GetFilterImageColorAndOpacity)
 					]
 					+SHorizontalBox::Slot()
-					.Padding(2.0f)
+					.Padding(TAttribute<FMargin>(this, &SFilter::GetFilterNamePadding))
 					.VAlign(VAlign_Center)
 					[
 						SNew(STextBlock)
 						.Text(this, &SFilter::GetFilterName)
-						//.ColorAndOpacity(this, &SFilter::GetFilterTextColorAndOpacity);
+						.IsEnabled_Lambda([this] {return bEnabled;})
 					]
 				]
 			]
@@ -421,21 +420,18 @@ private:
 	/** Handler to determine the color of the checkbox when it is checked */
 	FSlateColor GetFilterImageColorAndOpacity() const
 	{
-		return bEnabled ? FilterColor : FilterColor.Desaturate(.4f);
+		return bEnabled ? FilterColor : FAppStyle::Get().GetSlateColor("Colors.Recessed");
 	}
 
-	const FSlateBrush* GetFilterBorderBackground() const
+	EVisibility GetFilterOverlayVisibility() const
 	{
-		static const FName Unchecked("ContentBrowser.FilterBackgroundUnchecked");
-		static const FName Checked("ContentBrowser.FilterBackgroundChecked");
-
-		return FAppStyle::Get().GetBrush(bEnabled ? Checked : Unchecked);
+		return bEnabled ? EVisibility::Collapsed : EVisibility::HitTestInvisible;
 	}
 
 	/** Handler to determine the padding of the checkbox text when it is pressed */
 	FMargin GetFilterNamePadding() const
 	{
-		return ToggleButtonPtr->IsPressed() ? FMargin(3,2,4,0) : FMargin(3,1,4,1);
+		return ToggleButtonPtr->IsPressed() ? FMargin(4,2,4,0) : FMargin(4,1,4,1);
 	}
 
 
diff --git a/Engine/Source/Editor/EditorStyle/Private/StarshipStyle.cpp b/Engine/Source/Editor/EditorStyle/Private/StarshipStyle.cpp
--- a/Engine/Source/Editor/EditorStyle/Private/StarshipStyle.cpp
+++ b/Engine/Source/Editor/EditorStyle/Private/StarshipStyle.cpp
@@ -4235,6 +4235,7 @@ void FStarshipEditorStyle::FStyle::SetupLevelEditorStyle()
 		Set("MainFrame.ReportABug",                new IMAGE_BRUSH_SVG("Starship/Common/Bug", Icon16x16));
 
 		Set("MainFrame.VisitAskAQuestionPage",     new IMAGE_BRUSH_SVG("Starship/Menus/Bug", Icon16x16));
+		Set("MainFrame.VisitOnlineLearning",	   new IMAGE_BRUSH_SVG("Starship/Common/Documentation", Icon16x16));
 		Set("MainFrame.VisitWiki",                 new IMAGE_BRUSH_SVG("Starship/Common/WebBrowser", Icon16x16));
 		Set("MainFrame.VisitForums",               new IMAGE_BRUSH_SVG("Starship/Common/WebBrowser", Icon16x16));
 		Set("MainFrame.VisitSearchForAnswersPage", new IMAGE_BRUSH_SVG("Starship/Common/WebBrowser", Icon16x16));
@@ -6019,14 +6020,14 @@ void FStarshipEditorStyle::FStyle::SetupContentBrowserStyle()
 			.SetCheckedPressedImage(FSlateNoResource())
 			.SetForegroundColor(FStyleColors::Foreground)
 			.SetHoveredForegroundColor(FStyleColors::ForegroundHover)
-			.SetCheckedForegroundColor(FStyleColors::Primary)
-			.SetCheckedHoveredForegroundColor(FStyleColors::PrimaryHover)
-			.SetCheckedPressedForegroundColor(FStyleColors::PrimaryHover)
+			.SetCheckedForegroundColor(FStyleColors::Foreground)
+			.SetCheckedHoveredForegroundColor(FStyleColors::ForegroundHover)
+			.SetCheckedPressedForegroundColor(FStyleColors::ForegroundHover)
 			.SetPressedForegroundColor(FStyleColors::ForegroundHover);
 
-		Set("ContentBrowser.FilterImage",  new FSlateRoundedBoxBrush(FStyleColors::White, Icon8x8));// new CORE_IMAGE_BRUSH_SVG("Starship/Common/filled-circle", Icon8x8));
-		Set("ContentBrowser.FilterBackgroundUnchecked", new FSlateRoundedBoxBrush(FStyleColors::Secondary, 4.0f));
-		Set("ContentBrowser.FilterBackgroundChecked", new FSlateRoundedBoxBrush(FStyleColors::Recessed, 4.0f));
+		Set("ContentBrowser.FilterImage", new IMAGE_BRUSH_SVG("Starship/ContentBrowser/AssetFilterColorSegment", FVector2D(8, 22)));
+		Set("ContentBrowser.FilterBackground", new FSlateRoundedBoxBrush(FStyleColors::Secondary, 3.0f));
+
 		/* ... and add the new style */
 		Set("ContentBrowser.FilterButton", ContentBrowserFilterButtonCheckBoxStyle );
 
diff --git a/Engine/Source/Editor/UnrealEd/Private/AssetThumbnail.cpp b/Engine/Source/Editor/UnrealEd/Private/AssetThumbnail.cpp
--- a/Engine/Source/Editor/UnrealEd/Private/AssetThumbnail.cpp
+++ b/Engine/Source/Editor/UnrealEd/Private/AssetThumbnail.cpp
@@ -137,8 +137,6 @@ public:
 					.Font(GetTextFont())
 					.Justification(ETextJustify::Center)
 					.ColorAndOpacity(FEditorStyle::GetColor(Style, ".ColorAndOpacity"))
-					.ShadowOffset(FEditorStyle::GetVector(Style, ".ShadowOffset"))
-					.ShadowColorAndOpacity( FEditorStyle::GetColor(Style, ".ShadowColorAndOpacity"))
 					.HighlightText(HighlightedText)
 				]
 
@@ -224,8 +222,6 @@ public:
 						.Text(GetLabelText())
 						.Font(GetHintTextFont())
 						.ColorAndOpacity(FEditorStyle::GetColor(Style, ".HintColorAndOpacity"))
-						.ShadowOffset(FEditorStyle::GetVector(Style, ".HintShadowOffset"))
-						.ShadowColorAndOpacity(FEditorStyle::GetColor(Style, ".HintShadowColorAndOpacity"))
 						.HighlightText(HighlightedText)
 					]
 				];
