Fixed SPathView sometimes switching to "All Assets" during asset discovery

#rnx

diff --git a/Engine/Source/Editor/ContentBrowser/Private/SPathView.cpp b/Engine/Source/Editor/ContentBrowser/Private/SPathView.cpp
--- a/Engine/Source/Editor/ContentBrowser/Private/SPathView.cpp
+++ b/Engine/Source/Editor/ContentBrowser/Private/SPathView.cpp
@@ -880,6 +880,7 @@ bool SPathView::ExplicitlyAddPathToSelection(const FName Path)
 	{
 		// Set the selection to the closest found folder and scroll it into view
 		RecursiveExpandParents(FoundItem);
+		LastSelectedPaths.Add(FoundItem->GetItem().GetVirtualPath());
 		TreeViewPtr->SetItemSelection(FoundItem, true);
 		TreeViewPtr->RequestScrollIntoView(FoundItem);
 
