Fix Add Collection not working in the content browser of the collections area is collapsed

diff --git a/Engine/Source/Editor/ContentBrowser/Private/SCollectionView.cpp b/Engine/Source/Editor/ContentBrowser/Private/SCollectionView.cpp
--- a/Engine/Source/Editor/ContentBrowser/Private/SCollectionView.cpp
+++ b/Engine/Source/Editor/ContentBrowser/Private/SCollectionView.cpp
@@ -152,7 +152,7 @@ void SCollectionView::Construct( const FArguments& InArgs )
 				SNew(SButton)
 				.ButtonStyle(FEditorStyle::Get(), "SimpleButton")
 				.ToolTipText(LOCTEXT("AddCollectionButtonTooltip", "Add a collection."))
-				.OnClicked(this, &SCollectionView::MakeAddCollectionMenu)
+				.OnClicked(this, &SCollectionView::OnAddCollectionClicked)
 				.ContentPadding( FMargin(2, 2) )
 				.Visibility(bAllowCollectionButtons ? EVisibility::Visible : EVisibility::Collapsed)
 				[
@@ -812,12 +812,18 @@ void SCollectionView::MakeSaveDynamicCollectionMenu(FText InQueryString)
 	}
 }
 
+FReply SCollectionView::OnAddCollectionClicked()
+{
+	MakeAddCollectionMenu(AsShared());
+	return FReply::Handled();
+}
+
 bool SCollectionView::ShouldAllowSelectionChangedDelegate() const
 {
 	return PreventSelectionChangedDelegateCount == 0;
 }
 
-FReply SCollectionView::MakeAddCollectionMenu()
+void SCollectionView::MakeAddCollectionMenu(TSharedRef<SWidget> MenuParent)
 {
 	// Get all menu extenders for this context menu from the content browser module
 	FContentBrowserModule& ContentBrowserModule = FModuleManager::GetModuleChecked<FContentBrowserModule>( TEXT("ContentBrowser") );
@@ -840,14 +846,12 @@ FReply SCollectionView::MakeAddCollectionMenu()
 	CollectionContextMenu->MakeNewCollectionSubMenu(MenuBuilder, ECollectionStorageMode::Static, SCollectionView::FCreateCollectionPayload());
 
 	FSlateApplication::Get().PushMenu(
-		AsShared(),
+		MenuParent,
 		FWidgetPath(),
 		MenuBuilder.MakeWidget(),
 		FSlateApplication::Get().GetCursorPos(),
 		FPopupTransitionEffect( FPopupTransitionEffect::TopMenu )
 		);
-
-	return FReply::Handled();
 }
 
 void SCollectionView::CreateCollectionItem( ECollectionShareType::Type CollectionType, ECollectionStorageMode::Type StorageMode, const FCreateCollectionPayload& InCreationPayload )
diff --git a/Engine/Source/Editor/ContentBrowser/Private/SCollectionView.h b/Engine/Source/Editor/ContentBrowser/Private/SCollectionView.h
--- a/Engine/Source/Editor/ContentBrowser/Private/SCollectionView.h
+++ b/Engine/Source/Editor/ContentBrowser/Private/SCollectionView.h
@@ -105,7 +105,8 @@ public:
 	void MakeSaveDynamicCollectionMenu(FText InQueryString);
 
 	/** Creates the menu for the add collection button */
-	FReply MakeAddCollectionMenu();
+	FReply OnAddCollectionClicked();
+	void MakeAddCollectionMenu(TSharedRef<SWidget> MenuParent);
 
 private:
 
diff --git a/Engine/Source/Editor/ContentBrowser/Private/SContentBrowser.cpp b/Engine/Source/Editor/ContentBrowser/Private/SContentBrowser.cpp
--- a/Engine/Source/Editor/ContentBrowser/Private/SContentBrowser.cpp
+++ b/Engine/Source/Editor/ContentBrowser/Private/SContentBrowser.cpp
@@ -814,7 +814,7 @@ TSharedRef<SWidget> SContentBrowser::CreateDockedCollectionsView(const FContentB
 				SNew(SButton)
 				.ButtonStyle(FEditorStyle::Get(), "SimpleButton")
 				.ToolTipText(LOCTEXT("AddCollectionButtonTooltip", "Add a collection."))
-				.OnClicked(CollectionViewPtr.ToSharedRef(), &SCollectionView::MakeAddCollectionMenu)
+				.OnClicked(this, &SContentBrowser::OnAddCollectionClicked)
 				.ContentPadding(FMargin(1, 0))
 				[
 					SNew(SImage)
@@ -2432,6 +2432,15 @@ FReply SContentBrowser::ForwardClicked()
 	return FReply::Handled();
 }
 
+FReply SContentBrowser::OnAddCollectionClicked()
+{
+	CollectionArea->SetExpanded(true);
+
+	CollectionViewPtr->MakeAddCollectionMenu(AsShared());
+
+	return FReply::Handled();
+}
+
 bool SContentBrowser::HandleRenameCommandCanExecute() const
 {
 	// The order of these conditions are carefully crafted to match the logic of the context menu summoning, as this callback 
diff --git a/Engine/Source/Editor/ContentBrowser/Private/SContentBrowser.h b/Engine/Source/Editor/ContentBrowser/Private/SContentBrowser.h
--- a/Engine/Source/Editor/ContentBrowser/Private/SContentBrowser.h
+++ b/Engine/Source/Editor/ContentBrowser/Private/SContentBrowser.h
@@ -354,6 +354,9 @@ private:
 	/** Handler for clicking the history forward button */
 	FReply ForwardClicked();
 
+	/** Handler for clicking the add collection button */
+	FReply OnAddCollectionClicked();
+
 	/** Handler to check to see if a rename command is allowed */
 	bool HandleRenameCommandCanExecute() const;
 
