Replace some usage of recessed color with background color to flatten things up a bit

diff --git a/Engine/Source/Editor/ContentBrowser/Private/SContentBrowser.cpp b/Engine/Source/Editor/ContentBrowser/Private/SContentBrowser.cpp
--- a/Engine/Source/Editor/ContentBrowser/Private/SContentBrowser.cpp
+++ b/Engine/Source/Editor/ContentBrowser/Private/SContentBrowser.cpp
@@ -316,7 +316,13 @@ void SContentBrowser::Construct( const FArguments& InArgs, const FName& InInstan
 				]
 			]
 		]
-
+		+ SVerticalBox::Slot()
+		.AutoHeight()
+		[
+			SNew(SSeparator)
+			.SeparatorImage(FAppStyle::Get().GetBrush("Brushes.Recessed"))
+			.Thickness(2.0f)
+		]
 
 		// Assets/tree
 		+ SVerticalBox::Slot()
@@ -337,53 +343,48 @@ void SContentBrowser::Construct( const FArguments& InArgs, const FName& InInstan
 				.BorderImage(FEditorStyle::GetBrush("Brushes.Background"))
 				.Visibility(this, &SContentBrowser::GetSourcesViewVisibility)
 				[
-					SNew(SBorder)
-					.Padding(FMargin(0.0f, 2.0f, 0.0f, 3.0f))
-					.BorderImage(FEditorStyle::GetBrush("Brushes.Recessed"))
-					[
-						// Note: If adding more widgets here, fix ContentBrowserSourcesWidgetSwitcherIndex and the code that uses it!
-						SAssignNew(SourcesWidgetSwitcher, SWidgetSwitcher)
+					// Note: If adding more widgets here, fix ContentBrowserSourcesWidgetSwitcherIndex and the code that uses it!
+					SAssignNew(SourcesWidgetSwitcher, SWidgetSwitcher)
 
-						// Paths View
-						+SWidgetSwitcher::Slot()
+					// Paths View
+					+SWidgetSwitcher::Slot()
+					[
+						SAssignNew(PathFavoriteSplitterPtr, SSplitter)
+						.Style(FEditorStyle::Get(), "ContentBrowser.Splitter")
+						.PhysicalSplitterHandleSize(3.0f)
+						.HitDetectionSplitterHandleSize(3.0f)
+						.Orientation(EOrientation::Orient_Vertical)
+						.MinimumSlotHeight(26.0f)
+						.Visibility( this, &SContentBrowser::GetSourcesViewVisibility )
+						+SSplitter::Slot()
+						.SizeRule(TAttribute<SSplitter::ESizeRule>(this, &SContentBrowser::GetFavoritesAreaSizeRule))
+						.Value(0.2f)
 						[
-							SAssignNew(PathFavoriteSplitterPtr, SSplitter)
-							.Style(FEditorStyle::Get(), "ContentBrowser.Splitter")
-							.PhysicalSplitterHandleSize(3.0f)
-							.HitDetectionSplitterHandleSize(3.0f)
-							.Orientation(EOrientation::Orient_Vertical)
-							.MinimumSlotHeight(26.0f)
-							.Visibility( this, &SContentBrowser::GetSourcesViewVisibility )
-							+SSplitter::Slot()
-							.SizeRule(TAttribute<SSplitter::ESizeRule>(this, &SContentBrowser::GetFavoritesAreaSizeRule))
-							.Value(0.2f)
-							[
-								CreateFavoritesView(Config)
-							]
+							CreateFavoritesView(Config)
+						]
 							
-							+SSplitter::Slot()
-							.SizeRule(TAttribute<SSplitter::ESizeRule>(this, &SContentBrowser::GetPathAreaSizeRule))
-							.Value(0.8f)
-							[
-								CreatePathView(Config)
-							]
-
-							+SSplitter::Slot()
-							.SizeRule(TAttribute<SSplitter::ESizeRule>(this, &SContentBrowser::GetCollectionsAreaSizeRule))
-							.Value(0.4f)
-							[
-								CreateDockedCollectionsView(Config)
-							]
+						+SSplitter::Slot()
+						.SizeRule(TAttribute<SSplitter::ESizeRule>(this, &SContentBrowser::GetPathAreaSizeRule))
+						.Value(0.8f)
+						[
+							CreatePathView(Config)
 						]
 
-						// Collections View
-						+SWidgetSwitcher::Slot()
+						+SSplitter::Slot()
+						.SizeRule(TAttribute<SSplitter::ESizeRule>(this, &SContentBrowser::GetCollectionsAreaSizeRule))
+						.Value(0.4f)
+						[
+							CreateDockedCollectionsView(Config)
+						]
+					]
+
+					// Collections View
+					+SWidgetSwitcher::Slot()
+					[
+						SNew(SBox)
+						.Padding(FMargin(0.0f, 4.0f, 0.0f, 0.0f))
 						[
-							SNew(SBox)
-							.Padding(FMargin(0.0f, 4.0f, 0.0f, 0.0f))
-							[
-								CollectionViewPtr.ToSharedRef()
-							]
+							CollectionViewPtr.ToSharedRef()
 						]
 					]
 				]
@@ -714,7 +715,7 @@ TSharedRef<SWidget> SContentBrowser::CreateFavoritesView(const FContentBrowserCo
 	return
 		SAssignNew(FavoritesArea, SExpandableArea)
 		.BorderImage(FAppStyle::Get().GetBrush("Brushes.Background"))
-		.BodyBorderImage(FAppStyle::Get().GetBrush("Brushes.Recessed"))
+		.BodyBorderImage(FAppStyle::Get().GetBrush("Brushes.Background"))
 		.HeaderPadding(FMargin(5.0f, 3.0f))
 		.Visibility(this, &SContentBrowser::GetFavoriteFolderVisibility)
 		.AllowAnimatedTransition(false)
@@ -745,7 +746,7 @@ TSharedRef<SWidget> SContentBrowser::CreatePathView(const FContentBrowserConfig*
 	return
 		SAssignNew(PathArea, SExpandableArea)
 		.BorderImage(FAppStyle::Get().GetBrush("Brushes.Background"))
-		.BodyBorderImage(FAppStyle::Get().GetBrush("Brushes.Recessed"))
+		.BodyBorderImage(FAppStyle::Get().GetBrush("Brushes.Background"))
 		.HeaderPadding(FMargin(5.0f, 3.0f))
 		.AllowAnimatedTransition(false)
 		.HeaderContent()
@@ -788,7 +789,7 @@ TSharedRef<SWidget> SContentBrowser::CreateDockedCollectionsView(const FContentB
 	return
 		SAssignNew(CollectionArea, SExpandableArea)
 		.BorderImage(FAppStyle::Get().GetBrush("Brushes.Background"))
-		.BodyBorderImage(FAppStyle::Get().GetBrush("Brushes.Recessed"))
+		.BodyBorderImage(FAppStyle::Get().GetBrush("Brushes.Background"))
 		.HeaderPadding(FMargin(5.0f, 3.0f))
 		.Visibility(this, &SContentBrowser::GetDockedCollectionsVisibility)
 		.AllowAnimatedTransition(false)
diff --git a/Engine/Source/Editor/SceneOutliner/Private/SSceneOutliner.cpp b/Engine/Source/Editor/SceneOutliner/Private/SSceneOutliner.cpp
--- a/Engine/Source/Editor/SceneOutliner/Private/SSceneOutliner.cpp
+++ b/Engine/Source/Editor/SceneOutliner/Private/SSceneOutliner.cpp
@@ -179,7 +179,7 @@ void SSceneOutliner::Construct(const FArguments& InArgs, const FSceneOutlinerIni
 
 		+SOverlay::Slot()
 		[
-			SNew(SBorder).BorderImage( FAppStyle::Get().GetBrush("Brushes.Recessed") )
+			SNew(SBorder).BorderImage( FAppStyle::Get().GetBrush("Brushes.Background") )
 		]
 
 		+SOverlay::Slot()
diff --git a/Engine/Source/Runtime/SlateCore/Private/Styling/StarshipCoreStyle.cpp b/Engine/Source/Runtime/SlateCore/Private/Styling/StarshipCoreStyle.cpp
--- a/Engine/Source/Runtime/SlateCore/Private/Styling/StarshipCoreStyle.cpp
+++ b/Engine/Source/Runtime/SlateCore/Private/Styling/StarshipCoreStyle.cpp
@@ -1571,10 +1571,10 @@ void FStarshipCoreStyle::SetupTableViewStyles(TSharedRef<FStyle>& Style)
 	const FSlateColor SelectorColor = Style->GetSlateColor("SelectorColor");
 
 	const FTableRowStyle DefaultTableRowStyle = FTableRowStyle()
-		.SetEvenRowBackgroundBrush(FSlateColorBrush(FStyleColors::Recessed))
+		.SetEvenRowBackgroundBrush(FSlateColorBrush(FStyleColors::Background))
 		.SetEvenRowBackgroundHoveredBrush(FSlateColorBrush(FStyleColors::SelectHover))
 
-		.SetOddRowBackgroundBrush(FSlateColorBrush(FStyleColors::Recessed))
+		.SetOddRowBackgroundBrush(FSlateColorBrush(FStyleColors::Background))
 		.SetOddRowBackgroundHoveredBrush(FSlateColorBrush(FStyleColors::SelectHover))
 
 		.SetSelectorFocusedBrush(BORDER_BRUSH("Common/Selector", FMargin(4.f / 16.f), SelectorColor))
