Batch up dispatch of files and folders to Explorer/Finder/File Manager into groups of 10
Unify ContentBrowser "Explore" / "Show in Explorer" logic across several user endpoints

diff --git a/Engine/Source/Editor/ContentBrowser/Private/AssetContextMenu.cpp b/Engine/Source/Editor/ContentBrowser/Private/AssetContextMenu.cpp
--- a/Engine/Source/Editor/ContentBrowser/Private/AssetContextMenu.cpp
+++ b/Engine/Source/Editor/ContentBrowser/Private/AssetContextMenu.cpp
@@ -670,14 +670,7 @@ void FAssetContextMenu::ExecuteSyncToAssetTree()
 
 void FAssetContextMenu::ExecuteFindInExplorer()
 {
-	for (const FContentBrowserItem& SelectedItem : SelectedFiles)
-	{
-		FString ItemFilename;
-		if (SelectedItem.GetItemPhysicalPath(ItemFilename) && FPaths::FileExists(ItemFilename))
-		{
-			FPlatformProcess::ExploreFolder(*IFileManager::Get().ConvertToAbsolutePathForExternalAppForRead(*ItemFilename));
-		}
-	}
+	ContentBrowserUtils::ExploreFolders(SelectedFiles, AssetView.Pin().ToSharedRef());
 }
 
 bool FAssetContextMenu::CanExecuteEditItems() const
diff --git a/Engine/Source/Editor/ContentBrowser/Private/ContentBrowserUtils.cpp b/Engine/Source/Editor/ContentBrowser/Private/ContentBrowserUtils.cpp
--- a/Engine/Source/Editor/ContentBrowser/Private/ContentBrowserUtils.cpp
+++ b/Engine/Source/Editor/ContentBrowser/Private/ContentBrowserUtils.cpp
@@ -468,6 +468,44 @@ FText ContentBrowserUtils::GetExploreFolderText()
 	return FText::Format(NSLOCTEXT("GenericPlatform", "ShowInFileManager", "Show in {FileManagerName}"), Args);
 }
 
+void ContentBrowserUtils::ExploreFolders(const TArray<FContentBrowserItem>& InItems, const TSharedRef<SWidget>& InParentContent)
+{
+	TArray<FString> ExploreItems;
+
+	for (const FContentBrowserItem& SelectedItem : InItems)
+	{
+		FString ItemFilename;
+		if (SelectedItem.GetItemPhysicalPath(ItemFilename))
+		{
+			const bool bExists = SelectedItem.IsFile() ? FPaths::FileExists(ItemFilename) : FPaths::DirectoryExists(ItemFilename);
+			if (bExists)
+			{
+				ExploreItems.Add(IFileManager::Get().ConvertToAbsolutePathForExternalAppForRead(*ItemFilename));
+			}
+		}
+	}
+
+	const int32 BatchSize = 10;
+	const FText FileManagerName = FPlatformMisc::GetFileManagerName();
+	const bool bHasMultipleBatches = ExploreItems.Num() > BatchSize;
+	for (int32 i = 0; i < ExploreItems.Num(); ++i)
+	{
+		bool bIsBatchBoundary = (i % BatchSize) == 0;
+		if (bHasMultipleBatches && bIsBatchBoundary)
+		{
+			int32 RemainingCount = ExploreItems.Num() - i;
+			int32 NextCount = FMath::Min(BatchSize, RemainingCount);
+			FText Prompt = FText::Format(LOCTEXT("ExecuteExploreConfirm", "Show {0} {0}|plural(one=item,other=items) in {1}?\nThere {2}|plural(one=is,other=are) {2} remaining."), NextCount, FileManagerName, RemainingCount);
+			if (FMessageDialog::Open(EAppMsgType::YesNo, Prompt) != EAppReturnType::Yes)
+			{
+				return;
+			}
+		}
+
+		FPlatformProcess::ExploreFolder(*ExploreItems[i]);
+	}
+}
+
 template <typename OutputContainerType>
 void ConvertLegacySelectionToVirtualPathsImpl(TArrayView<const FAssetData> InAssets, TArrayView<const FString> InFolders, const bool InUseFolderPaths, OutputContainerType& OutVirtualPaths)
 {
diff --git a/Engine/Source/Editor/ContentBrowser/Private/ContentBrowserUtils.h b/Engine/Source/Editor/ContentBrowser/Private/ContentBrowserUtils.h
--- a/Engine/Source/Editor/ContentBrowser/Private/ContentBrowserUtils.h
+++ b/Engine/Source/Editor/ContentBrowser/Private/ContentBrowserUtils.h
@@ -67,6 +67,9 @@ namespace ContentBrowserUtils
 	/** Gets the platform specific text for the "explore" command (FPlatformProcess::ExploreFolder) */
 	FText GetExploreFolderText();
 
+	/** Perform a batched "explore" operation on the specified file and/or folder paths */
+	void ExploreFolders(const TArray<FContentBrowserItem>& InItems, const TSharedRef<SWidget>& InParentContent);
+
 	/** Convert a legacy asset and path selection to their corresponding virtual paths for content browser data items */
 	void ConvertLegacySelectionToVirtualPaths(TArrayView<const FAssetData> InAssets, TArrayView<const FString> InFolders, const bool InUseFolderPaths, TArray<FName>& OutVirtualPaths);
 	void ConvertLegacySelectionToVirtualPaths(TArrayView<const FAssetData> InAssets, TArrayView<const FString> InFolders, const bool InUseFolderPaths, TSet<FName>& OutVirtualPaths);
diff --git a/Engine/Source/Editor/ContentBrowser/Private/PathContextMenu.cpp b/Engine/Source/Editor/ContentBrowser/Private/PathContextMenu.cpp
--- a/Engine/Source/Editor/ContentBrowser/Private/PathContextMenu.cpp
+++ b/Engine/Source/Editor/ContentBrowser/Private/PathContextMenu.cpp
@@ -293,14 +293,7 @@ void FPathContextMenu::MakeSetColorSubMenu(UToolMenu* Menu)
 
 void FPathContextMenu::ExecuteExplore()
 {
-	for (const FContentBrowserItem& SelectedItem : SelectedFolders)
-	{
-		FString ItemFilename;
-		if (SelectedItem.GetItemPhysicalPath(ItemFilename) && FPaths::DirectoryExists(ItemFilename))
-		{
-			FPlatformProcess::ExploreFolder(*IFileManager::Get().ConvertToAbsolutePathForExternalAppForRead(*ItemFilename));
-		}
-	}
+	ContentBrowserUtils::ExploreFolders(SelectedFolders, ParentContent.Pin().ToSharedRef());
 }
 
 bool FPathContextMenu::CanExecuteRename() const
diff --git a/Engine/Source/Editor/ContentBrowser/Private/SAssetDialog.cpp b/Engine/Source/Editor/ContentBrowser/Private/SAssetDialog.cpp
--- a/Engine/Source/Editor/ContentBrowser/Private/SAssetDialog.cpp
+++ b/Engine/Source/Editor/ContentBrowser/Private/SAssetDialog.cpp
@@ -465,19 +465,7 @@ FReply SAssetDialog::ExecuteDeleteFolderConfirmed()
 
 void SAssetDialog::ExecuteExplore()
 {
-	const TArray<FContentBrowserItem> SelectedItems = AssetPicker->GetAssetView()->GetSelectedItems();
-	for (const FContentBrowserItem& SelectedItem : SelectedItems)
-	{
-		FString ItemFilename;
-		if (SelectedItem.GetItemPhysicalPath(ItemFilename))
-		{
-			const bool bExists = SelectedItem.IsFile() ? FPaths::FileExists(ItemFilename) : FPaths::DirectoryExists(ItemFilename);
-			if (bExists)
-			{
-				FPlatformProcess::ExploreFolder(*IFileManager::Get().ConvertToAbsolutePathForExternalAppForRead(*ItemFilename));
-			}
-		}
-	}
+	ContentBrowserUtils::ExploreFolders(AssetPicker->GetAssetView()->GetSelectedItems(), AssetPicker->GetAssetView().ToSharedRef());
 }
 
 bool SAssetDialog::CanExecuteCreateNewFolder() const
