Content browser drawer improvements
- Always focus search field even when clicking drawer button
- Drawer is now summoned if there is no content browser when clicking the browse button in the details panel or asset editor.
- Creation of the content browser drawer is now handled by the content browser singleton.

diff --git a/Engine/Source/Editor/ContentBrowser/ContentBrowser.Build.cs b/Engine/Source/Editor/ContentBrowser/ContentBrowser.Build.cs
--- a/Engine/Source/Editor/ContentBrowser/ContentBrowser.Build.cs
+++ b/Engine/Source/Editor/ContentBrowser/ContentBrowser.Build.cs
@@ -44,6 +44,7 @@ public class ContentBrowser : ModuleRules
 				"AssetRegistry",
 				"AssetTagsEditor",
 				"ToolMenus",
+				"StatusBar",
 			}
 		);
 
diff --git a/Engine/Source/Editor/ContentBrowser/Private/ContentBrowserSingleton.cpp b/Engine/Source/Editor/ContentBrowser/Private/ContentBrowserSingleton.cpp
--- a/Engine/Source/Editor/ContentBrowser/Private/ContentBrowserSingleton.cpp
+++ b/Engine/Source/Editor/ContentBrowser/Private/ContentBrowserSingleton.cpp
@@ -28,6 +28,8 @@
 #include "CollectionAssetRegistryBridge.h"
 #include "ContentBrowserCommands.h"
 #include "CoreGlobals.h"
+#include "StatusBarSubsystem.h"
+#include "Toolkits/GlobalEditorCommonCommands.h"
 
 #define LOCTEXT_NAMESPACE "ContentBrowser"
 
@@ -107,6 +109,26 @@ TSharedRef<class SWidget> FContentBrowserSingleton::CreateCollectionPicker(const
 		.CollectionPickerConfig(CollectionPickerConfig);
 }
 
+TSharedRef<class SWidget> FContentBrowserSingleton::CreateContentBrowserDrawer(const FContentBrowserConfig& ContentBrowserConfig)
+{
+	TSharedPtr<SContentBrowser> ContentBrowserDrawerPinned;
+	if(!ContentBrowserDrawer.IsValid())
+	{
+		static const FName ContentBrowserDrawerInstanceName("ContentBrowserDrawer");
+		ContentBrowserDrawerPinned =
+			SNew(SContentBrowser, ContentBrowserDrawerInstanceName, &ContentBrowserConfig)
+			.IsEnabled(FSlateApplication::Get().GetNormalExecutionAttribute());
+
+		ContentBrowserDrawer = ContentBrowserDrawerPinned;
+	}
+	else
+	{
+		ContentBrowserDrawerPinned = ContentBrowserDrawer.Pin();
+	}
+
+	return ContentBrowserDrawerPinned.ToSharedRef();
+}
+
 void FContentBrowserSingleton::CreateOpenAssetDialog(const FOpenAssetDialogConfig& InConfig,
 													 const FOnAssetsChosenForOpen& InOnAssetsChosenForOpen,
 													 const FOnAssetDialogCancelled& InOnAssetDialogCancelled)
@@ -223,22 +245,30 @@ void FContentBrowserSingleton::FocusPrimaryContentBrowser(bool bFocusSearch)
 	}
 }
 
-void FContentBrowserSingleton::FocusContentBrowserSearchField(FName InstanceName)
+void FContentBrowserSingleton::FocusContentBrowserSearchField(TSharedPtr<SWidget> ContentBrowserWidget)
 {
-	if (InstanceName.IsValid() && !InstanceName.IsNone())
+	TSharedPtr<SContentBrowser> BrowserToFocus;
+
+	if (ContentBrowserWidget.IsValid() )
 	{
 		for (TWeakPtr<SContentBrowser>& Browser : AllContentBrowsers)
 		{
-			if (TSharedPtr<SContentBrowser> BrowserPinned = Browser.Pin())
+			if (Browser == ContentBrowserWidget)
 			{
-				if (BrowserPinned->GetInstanceName() == InstanceName)
-				{
-					BrowserPinned->SetKeyboardFocusOnSearch();
-					break;
-				}
+				BrowserToFocus = Browser.Pin();
 			}
 		}
 	}
+
+	if (!BrowserToFocus.IsValid() && ContentBrowserDrawer == ContentBrowserWidget)
+	{
+		BrowserToFocus = ContentBrowserDrawer.Pin();
+	}
+
+	if (BrowserToFocus.IsValid())
+	{
+		BrowserToFocus->SetKeyboardFocusOnSearch();
+	}
 }
 
 void FContentBrowserSingleton::CreateNewAsset(const FString& DefaultAssetName, const FString& PackagePath, UClass* AssetClass, UFactory* Factory)
@@ -284,13 +314,22 @@ TSharedPtr<SContentBrowser> FContentBrowserSingleton::FindContentBrowserToSync(b
 	}
 	else
 	{
-		// If there is no primary or it is locked, find the first non-locked valid browser
-		for (int32 BrowserIdx = 0; BrowserIdx < AllContentBrowsers.Num(); ++BrowserIdx)
+
+		// If there is no primary or it is locked, sync the content browser drawer
+		if (ContentBrowserDrawer.IsValid())
 		{
-			if ( AllContentBrowsers[BrowserIdx].IsValid() && (bAllowLockedBrowsers || !AllContentBrowsers[BrowserIdx].Pin()->IsLocked()) )
+			ContentBrowserToSync = ContentBrowserDrawer.Pin();
+		}
+		else
+		{
+			// if the drawer doesn't exist find the first non-locked valid browser
+			for (int32 BrowserIdx = 0; BrowserIdx < AllContentBrowsers.Num(); ++BrowserIdx)
 			{
-				ContentBrowserToSync = AllContentBrowsers[BrowserIdx].Pin();
-				break;
+				if (AllContentBrowsers[BrowserIdx].IsValid() && (bAllowLockedBrowsers || !AllContentBrowsers[BrowserIdx].Pin()->IsLocked()))
+				{
+					ContentBrowserToSync = AllContentBrowsers[BrowserIdx].Pin();
+					break;
+				}
 			}
 		}
 	}
@@ -567,11 +606,18 @@ void FContentBrowserSingleton::FocusContentBrowser(const TSharedPtr<SContentBrow
 {
 	if ( BrowserToFocus.IsValid() )
 	{
-		TSharedRef<SContentBrowser> Browser = BrowserToFocus.ToSharedRef();
-		TSharedPtr<FTabManager> TabManager = Browser->GetTabManager();
-		if ( TabManager.IsValid() )
+		if (BrowserToFocus == ContentBrowserDrawer)
+		{
+			GEditor->GetEditorSubsystem<UStatusBarSubsystem>()->OpenContentBrowserDrawer();
+		}
+		else
 		{
-			TabManager->TryInvokeTab(Browser->GetInstanceName());
+			TSharedRef<SContentBrowser> Browser = BrowserToFocus.ToSharedRef();
+			TSharedPtr<FTabManager> TabManager = Browser->GetTabManager();
+			if (TabManager.IsValid())
+			{
+				TabManager->TryInvokeTab(Browser->GetInstanceName());
+			}
 		}
 	}
 }
@@ -755,22 +801,30 @@ void FContentBrowserSingleton::ForceShowPluginContent(bool bEnginePlugin)
 	}
 }
 
-void FContentBrowserSingleton::SaveContentBrowserSettings(FName InstanceName)
+void FContentBrowserSingleton::SaveContentBrowserSettings(TSharedPtr<SWidget> ContentBrowserWidget)
 {
-	if (InstanceName.IsValid() && !InstanceName.IsNone())
+	TSharedPtr<SContentBrowser> BrowserToSave;
+
+	if (ContentBrowserWidget.IsValid())
 	{
 		for (TWeakPtr<SContentBrowser>& Browser : AllContentBrowsers)
 		{
-			if (TSharedPtr<SContentBrowser> BrowserPinned = Browser.Pin())
+			if (Browser == ContentBrowserWidget)
 			{
-				if (BrowserPinned->GetInstanceName() == InstanceName)
-				{
-					BrowserPinned->SaveSettings();
-					break;
-				}
+				BrowserToSave = Browser.Pin();
 			}
 		}
 	}
+
+	if (!BrowserToSave.IsValid() && ContentBrowserDrawer == ContentBrowserWidget)
+	{
+		BrowserToSave = ContentBrowserDrawer.Pin();
+	}
+
+	if (BrowserToSave.IsValid())
+	{
+		BrowserToSave->SaveSettings();
+	}
 }
 
 void FContentBrowserSingleton::PopulateConfigValues()
diff --git a/Engine/Source/Editor/ContentBrowser/Private/ContentBrowserSingleton.h b/Engine/Source/Editor/ContentBrowser/Private/ContentBrowserSingleton.h
--- a/Engine/Source/Editor/ContentBrowser/Private/ContentBrowserSingleton.h
+++ b/Engine/Source/Editor/ContentBrowser/Private/ContentBrowserSingleton.h
@@ -49,13 +49,14 @@ public:
 	virtual TSharedRef<class SWidget> CreateAssetPicker(const FAssetPickerConfig& AssetPickerConfig) override;
 	virtual TSharedRef<class SWidget> CreatePathPicker(const FPathPickerConfig& PathPickerConfig) override;
 	virtual TSharedRef<class SWidget> CreateCollectionPicker(const FCollectionPickerConfig& CollectionPickerConfig) override;
+	virtual TSharedRef<class SWidget> CreateContentBrowserDrawer(const FContentBrowserConfig& ContentBrowserConfig) override;
 	virtual void CreateOpenAssetDialog(const FOpenAssetDialogConfig& OpenAssetConfig, const FOnAssetsChosenForOpen& OnAssetsChosenForOpen, const FOnAssetDialogCancelled& OnAssetDialogCancelled) override;
 	virtual TArray<FAssetData> CreateModalOpenAssetDialog(const FOpenAssetDialogConfig& InConfig) override;
 	virtual void CreateSaveAssetDialog(const FSaveAssetDialogConfig& SaveAssetConfig, const FOnObjectPathChosenForSave& OnAssetNameChosenForSave, const FOnAssetDialogCancelled& OnAssetDialogCancelled) override;
 	virtual FString CreateModalSaveAssetDialog(const FSaveAssetDialogConfig& SaveAssetConfig) override;
 	virtual bool HasPrimaryContentBrowser() const override;
 	virtual void FocusPrimaryContentBrowser(bool bFocusSearch) override;
-	virtual void FocusContentBrowserSearchField(FName InstanceName) override;
+	virtual void FocusContentBrowserSearchField(TSharedPtr<SWidget> ContentBrowserWidget) override;
 	virtual void CreateNewAsset(const FString& DefaultAssetName, const FString& PackagePath, UClass* AssetClass, UFactory* Factory) override;
 	virtual void SyncBrowserToAssets(const TArray<struct FAssetData>& AssetDataList, bool bAllowLockedBrowsers = false, bool bFocusContentBrowser = true, const FName& InstanceName = FName(), bool bNewSpawnBrowser = false) override;
 	virtual void SyncBrowserToAssets(const TArray<UObject*>& AssetList, bool bAllowLockedBrowsers = false, bool bFocusContentBrowser = true, const FName& InstanceName = FName(), bool bNewSpawnBrowser = false) override;
@@ -69,7 +70,7 @@ public:
 	virtual void CaptureThumbnailFromViewport(FViewport* InViewport, TArray<FAssetData>& SelectedAssets) override;
 	virtual void SetSelectedPaths(const TArray<FString>& FolderPaths, bool bNeedsRefresh = false) override;
 	virtual void ForceShowPluginContent(bool bEnginePlugin) override;
-	virtual void SaveContentBrowserSettings(FName InstanceName) override;
+	virtual void SaveContentBrowserSettings(TSharedPtr<SWidget> ContentBrowserWidget) override;
 
 	/** Gets the content browser singleton as a FContentBrowserSingleton */
 	static FContentBrowserSingleton& Get();
@@ -130,6 +131,8 @@ public:
 private:
 	TArray<TWeakPtr<SContentBrowser>> AllContentBrowsers;
 
+	TWeakPtr<SContentBrowser> ContentBrowserDrawer;
+
 	TMap<FName, TWeakPtr<FTabManager>> BrowserToLastKnownTabManagerMap;
 
 	TWeakPtr<SContentBrowser> PrimaryContentBrowser;
diff --git a/Engine/Source/Editor/ContentBrowser/Public/IContentBrowserSingleton.h b/Engine/Source/Editor/ContentBrowser/Public/IContentBrowserSingleton.h
--- a/Engine/Source/Editor/ContentBrowser/Public/IContentBrowserSingleton.h
+++ b/Engine/Source/Editor/ContentBrowser/Public/IContentBrowserSingleton.h
@@ -492,6 +492,15 @@ public:
 	 */
 	virtual TSharedRef<class SWidget> CreateCollectionPicker(const FCollectionPickerConfig& CollectionPickerConfig) = 0;
 
+	/**
+	 * Generates a content browser for use in a drawer. This content browser is a singleton and is reused among all drawers.
+	 *
+	 * @param ContentBrowserConfig	Initial defaults for the new content browser
+	 *
+	 * @return The content browser drawer widget
+	 */
+	virtual TSharedRef<class SWidget> CreateContentBrowserDrawer(const FContentBrowserConfig& ContentBrowserConfig) = 0;
+
 	/**
 	 * Opens the Open Asset dialog in a non-modal window
 	 *
@@ -533,7 +542,7 @@ public:
 	virtual void FocusPrimaryContentBrowser(bool bFocusSearch) = 0;
 
 	/** Focuses the search field of a content browser widget */
-	virtual void FocusContentBrowserSearchField(FName InstanceName) = 0;
+	virtual void FocusContentBrowserSearchField(TSharedPtr<SWidget> ContentBrowserWidget) = 0;
 
 	/** Sets up an inline-name for the creation of a new asset in the primary content browser using the specified path and the specified class and/or factory */
 	virtual void CreateNewAsset(const FString& DefaultAssetName, const FString& PackagePath, UClass* AssetClass, UFactory* Factory) = 0;
@@ -626,7 +635,7 @@ public:
 	/**
 	 * Saves the settings for a particular content browser instance
 	 *
-	 * @param InstanceName The name of the content browser to save
+	 * @param ContentBrowserWidget The content browser widget to save
 	 */
-	virtual void SaveContentBrowserSettings(FName InstanceName) = 0;
-};
+	virtual void SaveContentBrowserSettings(TSharedPtr<SWidget> ContentBrowserWidget) = 0;
+};
\ No newline at end of file
diff --git a/Engine/Source/Editor/StatusBar/Private/SStatusBar.cpp b/Engine/Source/Editor/StatusBar/Private/SStatusBar.cpp
--- a/Engine/Source/Editor/StatusBar/Private/SStatusBar.cpp
+++ b/Engine/Source/Editor/StatusBar/Private/SStatusBar.cpp
@@ -438,8 +438,27 @@ void SStatusBar::OnGlobalFocusChanging(const FFocusEvent& FocusEvent, const FWea
 
 void SStatusBar::OnActiveTabChanged(TSharedPtr<SDockTab> PreviouslyActive, TSharedPtr<SDockTab> NewlyActivated)
 {
-	// Remove the content browser as the tab well's content is changing 
-	RemoveContentBrowser();
+	bool bShouldRemoveContentBrowser = false;
+	if (!PreviouslyActive || !NewlyActivated)
+	{
+		// Remove the content browser if there is some invalid state with the tabs
+		bShouldRemoveContentBrowser = true;
+	}
+	else if(NewlyActivated->GetTabRole() == ETabRole::MajorTab)
+	{
+		// Remove the content browser if a newly activated tab is a major tab
+		bShouldRemoveContentBrowser = true;
+	}
+	else if (PreviouslyActive->GetTabManager() != NewlyActivated->GetTabManager())
+	{
+		// Remove the content browser if we're switching tab managers (indicates a new status bar is becoming active)
+		bShouldRemoveContentBrowser = true;
+	}
+
+	if (bShouldRemoveContentBrowser)
+	{
+		RemoveContentBrowser();
+	}
 }
 
 bool SStatusBar::IsContentBrowserOpened() const
@@ -567,6 +586,7 @@ TSharedRef<SWidget> SStatusBar::MakeContentBrowserWidget()
 {
 	return
 		SNew(SButton)
+		.IsFocusable(false)
 		.ButtonStyle(&FAppStyle::Get().GetWidgetStyle<FButtonStyle>("StatusBar.StatusBarButton"))
 		.OnClicked(this, &SStatusBar::OnContentBrowserButtonClicked)
 		.ToolTipText(FText::Format(LOCTEXT("StatusBar_ContentBrowserDrawerToolTip","Opens a temporary content browser above this status which will dismiss when it loses focus ({0})"), FGlobalEditorCommonCommands::Get().OpenContentBrowserDrawer->GetInputText()))
diff --git a/Engine/Source/Editor/StatusBar/Private/StatusBarSubsystem.cpp b/Engine/Source/Editor/StatusBar/Private/StatusBarSubsystem.cpp
--- a/Engine/Source/Editor/StatusBar/Private/StatusBarSubsystem.cpp
+++ b/Engine/Source/Editor/StatusBar/Private/StatusBarSubsystem.cpp
@@ -18,8 +18,6 @@
 
 int32 UStatusBarSubsystem::HandleCounter = 0;
 
-static const FName StatusBarContentBrowserName = "StatusBarContentBrowser";
-
 void UStatusBarSubsystem::Initialize(FSubsystemCollectionBase& Collection)
 {
 	FSourceControlCommands::Register();
@@ -52,6 +50,29 @@ bool UStatusBarSubsystem::FocusDebugConsole(TSharedRef<SWindow> ParentWindow)
 	return bFocusedSuccessfully;
 }
 
+bool UStatusBarSubsystem::OpenContentBrowserDrawer()
+{
+	TSharedPtr<SWindow> ParentWindow = FSlateApplication::Get().GetActiveTopLevelWindow();
+	if (!ParentWindow.IsValid())
+	{
+		if (TSharedPtr<SDockTab> ActiveTab = FGlobalTabmanager::Get()->GetActiveTab())
+		{
+			if (TSharedPtr<SDockTab> ActiveMajorTab = FGlobalTabmanager::Get()->GetMajorTabForTabManager(ActiveTab->GetTabManager()))
+			{
+				ParentWindow = ActiveMajorTab->GetParentWindow();
+			}
+		}
+	}
+
+	if (ParentWindow.IsValid() && ParentWindow->GetType() == EWindowType::Normal)
+	{
+		TSharedRef<SWindow> WindowRef = ParentWindow.ToSharedRef();
+		return ToggleContentBrowser(ParentWindow.ToSharedRef());
+	}
+
+	return false;
+}
+
 bool UStatusBarSubsystem::ToggleContentBrowser(TSharedRef<SWindow> ParentWindow)
 {
 	bool bWasDismissed = false;
@@ -74,8 +95,6 @@ bool UStatusBarSubsystem::ToggleContentBrowser(TSharedRef<SWindow> ParentWindow)
 
 	if(!bWasDismissed)
 	{
-		CreateContentBrowserIfNeeded();
-
 		TSharedPtr<SWindow> Window = ParentWindow;
 		GEditor->GetTimerManager()->SetTimerForNextTick(FTimerDelegate::CreateUObject(this, &UStatusBarSubsystem::HandleDeferredOpenContentBrowser, Window));
 	}
@@ -85,6 +104,8 @@ bool UStatusBarSubsystem::ToggleContentBrowser(TSharedRef<SWindow> ParentWindow)
 
 TSharedRef<SWidget> UStatusBarSubsystem::MakeStatusBarWidget(FName StatusBarName, const TSharedRef<SDockTab>& InParentTab)
 {
+	CreateContentBrowserIfNeeded();
+
 	TSharedRef<SStatusBar> StatusBar =
 		SNew(SStatusBar, StatusBarName, InParentTab)
 		.OnConsoleClosed_UObject(this, &UStatusBarSubsystem::OnDebugConsoleClosed)
@@ -159,7 +180,7 @@ void UStatusBarSubsystem::CreateContentBrowserIfNeeded()
 		FContentBrowserConfig Config;
 		Config.bCanSetAsPrimaryBrowser = false;
 
-		StatusBarContentBrowser = ContentBrowserSingleton.CreateContentBrowser(StatusBarContentBrowserName, nullptr, &Config);
+		StatusBarContentBrowser = ContentBrowserSingleton.CreateContentBrowserDrawer(Config);
 	}
 }
 
@@ -188,6 +209,13 @@ void UStatusBarSubsystem::OnContentBrowserOpened(TSharedRef<SStatusBar>& StatusB
 			}
 		}
 	}
+
+	IContentBrowserSingleton& ContentBrowserSingleton = FModuleManager::Get().LoadModuleChecked<FContentBrowserModule>("ContentBrowser").Get();;
+
+	// Cache off the previously focused widget so we can restore focus if the user hits the focus key again
+	PreviousKeyboardFocusedWidget = FSlateApplication::Get().GetKeyboardFocusedWidget();
+
+	ContentBrowserSingleton.FocusContentBrowserSearchField(StatusBarContentBrowser);
 }
 
 void UStatusBarSubsystem::OnContentBrowserDismissed(const TSharedPtr<SWidget>& NewlyFocusedWidget)
@@ -198,7 +226,7 @@ void UStatusBarSubsystem::OnContentBrowserDismissed(const TSharedPtr<SWidget>& N
 	}
 
 	IContentBrowserSingleton& ContentBrowserSingleton = FModuleManager::Get().LoadModuleChecked<FContentBrowserModule>("ContentBrowser").Get();;
-	ContentBrowserSingleton.SaveContentBrowserSettings(StatusBarContentBrowserName);
+	ContentBrowserSingleton.SaveContentBrowserSettings(StatusBarContentBrowser);
 
 	PreviousKeyboardFocusedWidget.Reset();
 }
@@ -214,11 +242,6 @@ void UStatusBarSubsystem::HandleDeferredOpenContentBrowser(TSharedPtr<SWindow> P
 			{
 				StatusBarPinned->OpenContentBrowser();
 				IContentBrowserSingleton& ContentBrowserSingleton = FModuleManager::Get().GetModuleChecked<FContentBrowserModule>("ContentBrowser").Get();
-
-				// Cache off the previously focused widget so we can restore focus if the user hits the focus key again
-				PreviousKeyboardFocusedWidget = FSlateApplication::Get().GetKeyboardFocusedWidget();
-
-				ContentBrowserSingleton.FocusContentBrowserSearchField(StatusBarContentBrowserName);
 				break;
 			}
 		}
diff --git a/Engine/Source/Editor/StatusBar/Public/StatusBarSubsystem.h b/Engine/Source/Editor/StatusBar/Public/StatusBarSubsystem.h
--- a/Engine/Source/Editor/StatusBar/Public/StatusBarSubsystem.h
+++ b/Engine/Source/Editor/StatusBar/Public/StatusBarSubsystem.h
@@ -72,9 +72,11 @@ public:
 	bool FocusDebugConsole(TSharedRef<SWindow> ParentWindow);
 
 	/**
-	 * Opens or closes the content browser drawer for a status bar residing in the passed in parent window
+	 * Opens the content browser drawer for a status bar residing in the active window 
+	 * 
+	 * @return true if the content browser was opened or false if no status bar in the active window was found
 	 */
-	bool ToggleContentBrowser(TSharedRef<SWindow> ParentWindow);
+	bool OpenContentBrowserDrawer();
 
 	/** 
 	 * Creates a new instance of a status bar widget
@@ -110,6 +112,7 @@ public:
 	void ClearStatusBarMessages(FName StatusBarName);
 
 private:
+	bool ToggleContentBrowser(TSharedRef<SWindow> ParentWindow);
 	void OnDebugConsoleClosed();
 	void CreateContentBrowserIfNeeded();
 	TSharedPtr<SStatusBar> GetStatusBar(FName StatusBarName) const;
diff --git a/Engine/Source/Editor/UnrealEd/Private/Toolkits/GlobalEditorCommonCommands.cpp b/Engine/Source/Editor/UnrealEd/Private/Toolkits/GlobalEditorCommonCommands.cpp
--- a/Engine/Source/Editor/UnrealEd/Private/Toolkits/GlobalEditorCommonCommands.cpp
+++ b/Engine/Source/Editor/UnrealEd/Private/Toolkits/GlobalEditorCommonCommands.cpp
@@ -154,25 +154,7 @@ void FGlobalEditorCommonCommands::OnSummonedConsoleCommandBox()
 
 void FGlobalEditorCommonCommands::OnOpenContentBrowserDrawer()
 {
-	TSharedPtr<SWindow> ParentWindow = FSlateApplication::Get().GetActiveTopLevelWindow();
-	if (!ParentWindow.IsValid())
-	{
-		if (TSharedPtr<SDockTab> ActiveTab = FGlobalTabmanager::Get()->GetActiveTab())
-		{
-			if (TSharedPtr<SDockTab> ActiveMajorTab = FGlobalTabmanager::Get()->GetMajorTabForTabManager(ActiveTab->GetTabManager()))
-			{
-				ParentWindow = ActiveMajorTab->GetParentWindow();
-			}
-		}
-	}
-
-	if (ParentWindow.IsValid() && ParentWindow->GetType() == EWindowType::Normal)
-	{
-		TSharedRef<SWindow> WindowRef = ParentWindow.ToSharedRef();
-		FOutputLogModule& OutputLogModule = FModuleManager::LoadModuleChecked<FOutputLogModule>(TEXT("OutputLog"));
-
-		GEditor->GetEditorSubsystem<UStatusBarSubsystem>()->ToggleContentBrowser(ParentWindow.ToSharedRef());
-	}
+	GEditor->GetEditorSubsystem<UStatusBarSubsystem>()->OpenContentBrowserDrawer();
 }
 
 //////////////////////////////////////////////////////////////////////////
