Clear asset selection on shutdown

Any lingering references just used to become null, but they now need to be cleared to avoid leaking typed element references.

#rnx

diff --git a/Engine/Source/Editor/ContentBrowser/Private/SContentBrowser.cpp b/Engine/Source/Editor/ContentBrowser/Private/SContentBrowser.cpp
--- a/Engine/Source/Editor/ContentBrowser/Private/SContentBrowser.cpp
+++ b/Engine/Source/Editor/ContentBrowser/Private/SContentBrowser.cpp
@@ -104,6 +104,14 @@ SContentBrowser::~SContentBrowser()
 			ContentBrowserData->OnItemDataUpdated().RemoveAll(this);
 		}
 	}
+
+	if (bIsPrimaryBrowser && GEditor)
+	{
+		if (USelection* EditorSelection = GEditor->GetSelectedObjects())
+		{
+			EditorSelection->DeselectAll();
+		}
+	}
 }
 
 BEGIN_SLATE_FUNCTION_BUILD_OPTIMIZATION
