Content Browser: Control Rig:  Expose CustomBlackListFolder to PathPickerConfig so we can only see one root and not every folder for the Pose Library.

diff --git a/Engine/Source/Editor/ContentBrowser/Private/SAssetPicker.cpp b/Engine/Source/Editor/ContentBrowser/Private/SAssetPicker.cpp
--- a/Engine/Source/Editor/ContentBrowser/Private/SAssetPicker.cpp
+++ b/Engine/Source/Editor/ContentBrowser/Private/SAssetPicker.cpp
@@ -542,10 +542,9 @@ void SAssetPicker::HandleItemSelectionChanged(const FContentBrowserItem& InSelec
 	if (InSelectInfo != ESelectInfo::Direct)
 	{
 		FAssetData ItemAssetData;
-		if (InSelectedItem.Legacy_TryGetAssetData(ItemAssetData))
-		{
-			OnAssetSelected.ExecuteIfBound(ItemAssetData);
-		}
+		InSelectedItem.Legacy_TryGetAssetData(ItemAssetData);
+		OnAssetSelected.ExecuteIfBound(ItemAssetData);
+		
 	}
 }
 
diff --git a/Engine/Source/Editor/ContentBrowser/Private/SPathPicker.cpp b/Engine/Source/Editor/ContentBrowser/Private/SPathPicker.cpp
--- a/Engine/Source/Editor/ContentBrowser/Private/SPathPicker.cpp
+++ b/Engine/Source/Editor/ContentBrowser/Private/SPathPicker.cpp
@@ -39,6 +39,7 @@ void SPathPicker::Construct( const FArguments& InArgs )
 		.AllowClassesFolder(InArgs._PathPickerConfig.bAllowClassesFolder)
 		.AllowReadOnlyFolders(InArgs._PathPickerConfig.bAllowReadOnlyFolders)
 		.SelectionMode(ESelectionMode::Single)
+		.CustomFolderBlacklist(InArgs._PathPickerConfig.CustomFolderBlacklist)
 	];
 
 	const FString& DefaultPath = InArgs._PathPickerConfig.DefaultPath;
diff --git a/Engine/Source/Editor/ContentBrowser/Private/SPathView.cpp b/Engine/Source/Editor/ContentBrowser/Private/SPathView.cpp
--- a/Engine/Source/Editor/ContentBrowser/Private/SPathView.cpp
+++ b/Engine/Source/Editor/ContentBrowser/Private/SPathView.cpp
@@ -210,6 +210,7 @@ void SPathView::Construct( const FArguments& InArgs )
 		]
 	];
 
+	CustomFolderBlacklist = InArgs._CustomFolderBlacklist;
 	// Add all paths currently gathered from the asset registry
 	Populate();
 
@@ -607,6 +608,15 @@ FContentBrowserDataCompiledFilter SPathView::CreateCompiledFolderFilter() const
 		}
 	}
 
+	if (CustomFolderBlacklist.IsValid())
+	{
+		if (!CombinedFolderBlacklist.IsValid())
+		{
+			CombinedFolderBlacklist = MakeShared<FBlacklistPaths>();
+		}
+		CombinedFolderBlacklist->Append(*CustomFolderBlacklist);
+	}
+
 	ContentBrowserUtils::AppendAssetFilterToContentBrowserFilter(FARFilter(), nullptr, CombinedFolderBlacklist, DataFilter);
 
 	FContentBrowserDataCompiledFilter CompiledDataFilter;
diff --git a/Engine/Source/Editor/ContentBrowser/Private/SPathView.h b/Engine/Source/Editor/ContentBrowser/Private/SPathView.h
--- a/Engine/Source/Editor/ContentBrowser/Private/SPathView.h
+++ b/Engine/Source/Editor/ContentBrowser/Private/SPathView.h
@@ -83,6 +83,9 @@ public:
 		/** Optional external search. Will hide and replace our internal search UI */
 		SLATE_ARGUMENT( TSharedPtr<FSourcesSearch>, ExternalSearch )
 
+		/** Optional Custom Folder Blacklist to be used to filter folders. */
+		SLATE_ARGUMENT( TSharedPtr<FBlacklistPaths>, CustomFolderBlacklist)
+
 	SLATE_END_ARGS()
 
 	/** Destructor */
@@ -328,6 +331,9 @@ protected:
 	/** Writable folder filter */
 	TSharedPtr<FBlacklistPaths> WritableFolderBlacklist;
 
+	/** Custom Folder Blacklist*/
+	TSharedPtr<FBlacklistPaths> CustomFolderBlacklist;
+
 private:
 
 	/** The paths that were last reported by OnPathExpanded event. Used in preserving expansion when filtering folders */
diff --git a/Engine/Source/Editor/ContentBrowser/Public/IContentBrowserSingleton.h b/Engine/Source/Editor/ContentBrowser/Public/IContentBrowserSingleton.h
--- a/Engine/Source/Editor/ContentBrowser/Public/IContentBrowserSingleton.h
+++ b/Engine/Source/Editor/ContentBrowser/Public/IContentBrowserSingleton.h
@@ -14,6 +14,7 @@
 #include "Misc/FilterCollection.h"
 #include "Framework/Views/ITypedTableView.h"
 #include "AssetThumbnail.h"
+#include "Misc/BlacklistNames.h"
 
 class FViewport;
 class UFactory;
@@ -325,6 +326,9 @@ struct FPathPickerConfig
 	/** The initial path to select. Leave empty to skip initial selection. */
 	FString DefaultPath;
 
+	/** Custom Folder Blacklist to be used to filter folders in this Path Picker. */
+	TSharedPtr<FBlacklistPaths> CustomFolderBlacklist;
+
 	/** The delegate that fires when a path was selected */
 	FOnPathSelected OnPathSelected;
 
