Fix content browser folder tree clipping

diff --git a/Engine/Source/Editor/ContentBrowser/Private/SContentBrowser.cpp b/Engine/Source/Editor/ContentBrowser/Private/SContentBrowser.cpp
--- a/Engine/Source/Editor/ContentBrowser/Private/SContentBrowser.cpp
+++ b/Engine/Source/Editor/ContentBrowser/Private/SContentBrowser.cpp
@@ -348,6 +348,7 @@ void SContentBrowser::Construct( const FArguments& InArgs, const FName& InInstan
 					+SWidgetSwitcher::Slot()
 					[
 						SAssignNew(PathFavoriteSplitterPtr, SSplitter)
+						.Clipping(EWidgetClipping::ClipToBounds)
 						.PhysicalSplitterHandleSize(3.0f)
 						.HitDetectionSplitterHandleSize(3.0f)
 						.Orientation(EOrientation::Orient_Vertical)
