Fix cook assert when blutility module loads content browser module.

#rnx

diff --git a/Engine/Source/Editor/ContentBrowser/Private/ContentBrowserSingleton.cpp b/Engine/Source/Editor/ContentBrowser/Private/ContentBrowserSingleton.cpp
--- a/Engine/Source/Editor/ContentBrowser/Private/ContentBrowserSingleton.cpp
+++ b/Engine/Source/Editor/ContentBrowser/Private/ContentBrowserSingleton.cpp
@@ -55,7 +55,7 @@ FContentBrowserSingleton::FContentBrowserSingleton()
 		LOCTEXT("FocusContentBrowser_Label", "Focus Content Browser"),
 		LOCTEXT("FocusContentBrowser_Desc", "Focuses the most recently active content browser tab."),
 		FSlateIcon(),
-		FUIAction(FExecuteAction::CreateRaw(this, &FContentBrowserSingleton::FocusPrimaryContentBrowser, true))
+		FUIAction(FExecuteAction::CreateRaw(this, &FContentBrowserSingleton::FocusPrimaryContentBrowser, true), FCanExecuteAction())
 	);
 	Section.AddSeparator(NAME_None);
 
