Introduction of TextFilterKeyValueHandler
- Added a text filter key value handler class that can custom handle the key and value for the text filter in the content browser
- Handlers can be added from other modules/plugins and specified in the editor config file with a corresponding key and its class
- Added a class to query the configurable settings for each handler to find the correct handler based on key text

diff --git a/Engine/Source/Editor/ContentBrowser/Private/FrontendFilters.cpp b/Engine/Source/Editor/ContentBrowser/Private/FrontendFilters.cpp
--- a/Engine/Source/Editor/ContentBrowser/Private/FrontendFilters.cpp
+++ b/Engine/Source/Editor/ContentBrowser/Private/FrontendFilters.cpp
@@ -23,6 +23,7 @@
 #include "MRUFavoritesList.h"
 #include "Settings/ContentBrowserSettings.h"
 #include "HAL/FileManager.h"
+#include "TextFilterKeyValueHandlers.h"
 #include "AssetCompilingManager.h"
 
 /** Helper functions for frontend filters */
@@ -380,6 +381,11 @@ public:
 
 	virtual bool TestComplexExpression(const FName& InKey, const FTextFilterString& InValue, const ETextFilterComparisonOperation InComparisonOperation, const ETextFilterTextComparisonMode InTextComparisonMode) const override
 	{
+		if (UTextFilterKeyValueHandlers::HandleTextFilterKeyValue(*AssetPtr, InKey, InValue, InComparisonOperation, InTextComparisonMode))
+		{
+			return true;
+		}
+
 		// Special case for the asset name, as this isn't contained within the asset registry meta-data
 		if (InKey == NameKeyName)
 		{
diff --git a/Engine/Source/Editor/ContentBrowser/Private/TextFilterKeyValueHandlers.cpp b/Engine/Source/Editor/ContentBrowser/Private/TextFilterKeyValueHandlers.cpp
new file mode 100644
--- /dev/null
+++ b/Engine/Source/Editor/ContentBrowser/Private/TextFilterKeyValueHandlers.cpp
@@ -0,0 +1,21 @@
+// Copyright Epic Games, Inc. All Rights Reserved.
+
+#include "TextFilterKeyValueHandlers.h"
+
+#include "TextFilterKeyValueHandler.h"
+
+bool UTextFilterKeyValueHandlers::HandleTextFilterKeyValue(const FContentBrowserItem& InContentBrowserItem, const FName& InKey, const FTextFilterString& InValue, const ETextFilterComparisonOperation InComparisonOperation, const ETextFilterTextComparisonMode InTextComparisonMode)
+{
+	for (const FTextFilterKeyValueHandlerEntry& Handler : GetDefault<UTextFilterKeyValueHandlers>()->TextFilterKeyValueHandlers)
+	{
+		if (InKey.IsEqual(Handler.Key))
+		{
+			if (UClass* TextFilterKeyValueHandlerClass = Handler.HandlerClass.LoadSynchronous())
+			{
+				return GetDefault<UTextFilterKeyValueHandler>(TextFilterKeyValueHandlerClass)->HandleTextFilterKeyValue(InContentBrowserItem, InKey, InValue, InComparisonOperation, InTextComparisonMode);
+			}
+		}
+	}
+
+	return false;
+}
\ No newline at end of file
diff --git a/Engine/Source/Editor/ContentBrowser/Private/TextFilterKeyValueHandlers.h b/Engine/Source/Editor/ContentBrowser/Private/TextFilterKeyValueHandlers.h
new file mode 100644
--- /dev/null
+++ b/Engine/Source/Editor/ContentBrowser/Private/TextFilterKeyValueHandlers.h
@@ -0,0 +1,32 @@
+// Copyright Epic Games, Inc. All Rights Reserved.
+
+#pragma once
+
+#include "CoreMinimal.h"
+#include "ContentBrowserItem.h"
+#include "TextFilterKeyValueHandler.h"
+
+#include "TextFilterKeyValueHandlers.generated.h"
+
+USTRUCT()
+struct FTextFilterKeyValueHandlerEntry
+{
+	GENERATED_BODY()
+
+	UPROPERTY(config)
+	FName Key;
+
+	UPROPERTY(config)
+	TSoftClassPtr<UTextFilterKeyValueHandler> HandlerClass;
+};
+
+UCLASS(transient, config = Editor)
+class UTextFilterKeyValueHandlers : public UObject
+{
+	GENERATED_BODY()
+public:
+	UPROPERTY(config)
+	TArray<FTextFilterKeyValueHandlerEntry> TextFilterKeyValueHandlers;
+
+	static bool HandleTextFilterKeyValue(const FContentBrowserItem& InContentBrowserItem, const FName& InKey, const FTextFilterString& InValue, const ETextFilterComparisonOperation InComparisonOperation, const ETextFilterTextComparisonMode InTextComparisonMode);
+};
\ No newline at end of file
diff --git a/Engine/Source/Editor/ContentBrowser/Public/TextFilterKeyValueHandler.h b/Engine/Source/Editor/ContentBrowser/Public/TextFilterKeyValueHandler.h
new file mode 100644
--- /dev/null
+++ b/Engine/Source/Editor/ContentBrowser/Public/TextFilterKeyValueHandler.h
@@ -0,0 +1,16 @@
+// Copyright Epic Games, Inc. All Rights Reserved.
+
+#pragma once
+
+#include "CoreMinimal.h"
+#include "ContentBrowserItem.h"
+
+#include "TextFilterKeyValueHandler.generated.h"
+
+UCLASS(abstract, transient, MinimalAPI)
+class UTextFilterKeyValueHandler : public UObject
+{
+	GENERATED_BODY()
+public:
+	virtual bool HandleTextFilterKeyValue(const FContentBrowserItem& InContentBrowserItem, const FName& InKey, const FTextFilterString& InValue, const ETextFilterComparisonOperation InComparisonOperation, const ETextFilterTextComparisonMode InTextComparisonMode) const { return false; }
+};
\ No newline at end of file
