Prevent creating assets from creating a new content browser when there is an open content browser drawer.

diff --git a/Engine/Source/Editor/ContentBrowser/Private/ContentBrowserSingleton.cpp b/Engine/Source/Editor/ContentBrowser/Private/ContentBrowserSingleton.cpp
--- a/Engine/Source/Editor/ContentBrowser/Private/ContentBrowserSingleton.cpp
+++ b/Engine/Source/Editor/ContentBrowser/Private/ContentBrowserSingleton.cpp
@@ -297,11 +297,21 @@ void FContentBrowserSingleton::FocusContentBrowserSearchField(TSharedPtr<SWidget
 
 void FContentBrowserSingleton::CreateNewAsset(const FString& DefaultAssetName, const FString& PackagePath, UClass* AssetClass, UFactory* Factory)
 {
-	FocusPrimaryContentBrowser(false);
+	const bool bAllowLockedBrowsers = true;
+	const FName ContentBrowserInstanceName = NAME_None;
+	const bool bCreateNewContentBrowser = false;
 
-	if ( PrimaryContentBrowser.IsValid() )
+	TSharedPtr<SContentBrowser> ContentBrowserToSync = FindContentBrowserToSync(bAllowLockedBrowsers, ContentBrowserInstanceName, bCreateNewContentBrowser);
+
+	if (!ContentBrowserToSync.IsValid())
+	{
+		FocusPrimaryContentBrowser(false);
+		ContentBrowserToSync = PrimaryContentBrowser.Pin();
+	}
+
+	if (ContentBrowserToSync.IsValid())
 	{
-		PrimaryContentBrowser.Pin()->CreateNewAsset(DefaultAssetName, PackagePath, AssetClass, Factory);
+		ContentBrowserToSync->CreateNewAsset(DefaultAssetName, PackagePath, AssetClass, Factory);
 	}
 }
 
