Content Browse Singleton: Add functions to rename assets or folders or add a new folder. This is needed for custom context menus, like one in the new Pose Library, so they can interact correctly with the hidden text edit widgets.

diff --git a/Engine/Source/Editor/ContentBrowser/Private/ContentBrowserSingleton.cpp b/Engine/Source/Editor/ContentBrowser/Private/ContentBrowserSingleton.cpp
--- a/Engine/Source/Editor/ContentBrowser/Private/ContentBrowserSingleton.cpp
+++ b/Engine/Source/Editor/ContentBrowser/Private/ContentBrowserSingleton.cpp
@@ -861,6 +861,37 @@ void FContentBrowserSingleton::SaveContentBrowserSettings(TSharedPtr<SWidget> Co
 	}
 }
 
+
+void FContentBrowserSingleton::ExecuteRename(TSharedPtr<SWidget> PickerWidget)
+{
+	if (PickerWidget.IsValid())
+	{
+		if (PickerWidget->GetType() == FName(TEXT("SAssetPicker")))
+		{
+			TSharedPtr<SAssetPicker> AssetPicker = StaticCastSharedPtr<SAssetPicker>(PickerWidget);
+			AssetPicker->ExecuteRenameCommand();
+		}
+		else if (PickerWidget->GetType() == FName(TEXT("SPathPicker")))
+		{
+			TSharedPtr<SPathPicker> PathPicker = StaticCastSharedPtr<SPathPicker>(PickerWidget);
+			PathPicker->ExecuteRenameFolder();
+		}
+	}
+}
+
+
+void FContentBrowserSingleton::ExecuteAddFolder(TSharedPtr<SWidget> Widget)
+{
+	if (Widget.IsValid())
+	{
+		if (Widget->GetType() == FName(TEXT("SPathPicker")))
+		{
+			TSharedPtr<SPathPicker> PathPicker = StaticCastSharedPtr<SPathPicker>(Widget);
+			PathPicker->ExecuteAddFolder();
+		}
+	}
+}
+
 void FContentBrowserSingleton::PopulateConfigValues()
 {
 	const FString ContentBrowserSection = TEXT("ContentBrowser");
diff --git a/Engine/Source/Editor/ContentBrowser/Private/ContentBrowserSingleton.h b/Engine/Source/Editor/ContentBrowser/Private/ContentBrowserSingleton.h
--- a/Engine/Source/Editor/ContentBrowser/Private/ContentBrowserSingleton.h
+++ b/Engine/Source/Editor/ContentBrowser/Private/ContentBrowserSingleton.h
@@ -71,6 +71,8 @@ public:
 	virtual void SetSelectedPaths(const TArray<FString>& FolderPaths, bool bNeedsRefresh = false) override;
 	virtual void ForceShowPluginContent(bool bEnginePlugin) override;
 	virtual void SaveContentBrowserSettings(TSharedPtr<SWidget> ContentBrowserWidget) override;
+	virtual void ExecuteRename(TSharedPtr<SWidget> PickerWidget) override;
+	virtual void ExecuteAddFolder(TSharedPtr<SWidget> PathPickerWidget) override;
 
 	/** Gets the content browser singleton as a FContentBrowserSingleton */
 	static FContentBrowserSingleton& Get();
diff --git a/Engine/Source/Editor/ContentBrowser/Private/SAssetPicker.cpp b/Engine/Source/Editor/ContentBrowser/Private/SAssetPicker.cpp
--- a/Engine/Source/Editor/ContentBrowser/Private/SAssetPicker.cpp
+++ b/Engine/Source/Editor/ContentBrowser/Private/SAssetPicker.cpp
@@ -672,6 +672,14 @@ bool SAssetPicker::CanExecuteRenameRequested()
 	return ContentBrowserUtils::CanRenameFromAssetView(AssetViewPtr);
 }
 
+void SAssetPicker::ExecuteRenameCommand()
+{
+	if (Commands.IsValid())
+	{
+		Commands->TryExecuteAction(FGenericCommands::Get().Rename.ToSharedRef());
+	}
+}
+
 void SAssetPicker::BindCommands()
 {
 	Commands = MakeShareable(new FUICommandList);
diff --git a/Engine/Source/Editor/ContentBrowser/Private/SAssetPicker.h b/Engine/Source/Editor/ContentBrowser/Private/SAssetPicker.h
--- a/Engine/Source/Editor/ContentBrowser/Private/SAssetPicker.h
+++ b/Engine/Source/Editor/ContentBrowser/Private/SAssetPicker.h
@@ -46,6 +46,9 @@ public:
 	/** Return the associated AssetView */
 	const TSharedPtr<SAssetView>& GetAssetView() const { return AssetViewPtr; }
 
+	/** Function that's called when someone wants to execute a Rename Command*/
+	void ExecuteRenameCommand();
+
 private:
 	/** Focuses the search box post-construct */
 	EActiveTimerReturnType SetFocusPostConstruct( double InCurrentTime, float InDeltaTime );
diff --git a/Engine/Source/Editor/ContentBrowser/Private/SPathPicker.cpp b/Engine/Source/Editor/ContentBrowser/Private/SPathPicker.cpp
--- a/Engine/Source/Editor/ContentBrowser/Private/SPathPicker.cpp
+++ b/Engine/Source/Editor/ContentBrowser/Private/SPathPicker.cpp
@@ -145,6 +145,31 @@ TSharedPtr<SWidget> SPathPicker::GetFolderContextMenu(const TArray<FString> & Se
 	return MenuBuilder.MakeWidget();
 }
 
+void SPathPicker::ExecuteRenameFolder()
+{
+	if (PathViewPtr.IsValid())
+	{
+		const TArray<FContentBrowserItem> SelectedItems = PathViewPtr->GetSelectedFolderItems();
+		if (SelectedItems.Num() == 1)
+		{
+			PathViewPtr->RenameFolderItem(SelectedItems[0]);
+		}
+	}
+}
+
+void SPathPicker::ExecuteAddFolder()
+{
+	if (PathViewPtr.IsValid())
+	{
+		const TArray<FString> SelectedItems = PathViewPtr->GetSelectedPaths();
+		if (SelectedItems.Num() == 1)
+		{
+			FOnCreateNewFolder OnCreateNewFolder = FOnCreateNewFolder::CreateSP(PathViewPtr.Get(), &SPathView::NewFolderItemRequested);
+			CreateNewFolder(SelectedItems[0], OnCreateNewFolder);
+		}
+	}
+}
+
 void SPathPicker::CreateNewFolder(FString FolderPath, FOnCreateNewFolder InOnCreateNewFolder)
 {
 	const FText DefaultFolderBaseName = LOCTEXT("DefaultFolderName", "NewFolder");
diff --git a/Engine/Source/Editor/ContentBrowser/Private/SPathPicker.h b/Engine/Source/Editor/ContentBrowser/Private/SPathPicker.h
--- a/Engine/Source/Editor/ContentBrowser/Private/SPathPicker.h
+++ b/Engine/Source/Editor/ContentBrowser/Private/SPathPicker.h
@@ -37,6 +37,12 @@ public:
 	/** Handler for creating a new folder in the path picker */
 	void CreateNewFolder(FString FolderPath, FOnCreateNewFolder InOnCreateNewFolder);
 
+	/** Rename the selected, will just use the first folder selected*/
+	void ExecuteRenameFolder();
+
+	/** Add a folder to the first selected folder*/
+	void ExecuteAddFolder();
+
 private:
 
 	/** Handle for when selection changes */
diff --git a/Engine/Source/Editor/ContentBrowser/Public/IContentBrowserSingleton.h b/Engine/Source/Editor/ContentBrowser/Public/IContentBrowserSingleton.h
--- a/Engine/Source/Editor/ContentBrowser/Public/IContentBrowserSingleton.h
+++ b/Engine/Source/Editor/ContentBrowser/Public/IContentBrowserSingleton.h
@@ -646,4 +646,18 @@ public:
 	 * @param ContentBrowserWidget The content browser widget to save
 	 */
 	virtual void SaveContentBrowserSettings(TSharedPtr<SWidget> ContentBrowserWidget) = 0;
+
+	/**
+	 * Rename current first selected content item on the passed in widget.
+	 *
+	 * @param PickerWidget The picker widget whose asset we want to rename, should be a asset or path picker widget.
+	 */
+	virtual void ExecuteRename(TSharedPtr<SWidget> PickerWidget) = 0;
+
+	/**
+	 * Add a folder to the path picker widget under the current selected path.
+	 *
+	 * @param PathPickerWidget The path picker widget where we want to add an folder
+	 */
+	virtual void ExecuteAddFolder(TSharedPtr<SWidget> PathPickerWidget) = 0;
 };
\ No newline at end of file
