Control Rig: Pose Library. When we add a new new folder to our whitelist it doesn't show up since the path view doesn't refresh. So added API call to ContentBrowserSingleton->PathPicker to refresh/populate the view due to the filter change.

diff --git a/Engine/Source/Editor/ContentBrowser/Private/ContentBrowserSingleton.cpp b/Engine/Source/Editor/ContentBrowser/Private/ContentBrowserSingleton.cpp
--- a/Engine/Source/Editor/ContentBrowser/Private/ContentBrowserSingleton.cpp
+++ b/Engine/Source/Editor/ContentBrowser/Private/ContentBrowserSingleton.cpp
@@ -892,6 +892,18 @@ void FContentBrowserSingleton::ExecuteAddFolder(TSharedPtr<SWidget> Widget)
 	}
 }
 
+void FContentBrowserSingleton::RefreshPathView(TSharedPtr<SWidget> Widget)
+{
+	if (Widget.IsValid())
+	{
+		if (Widget->GetType() == FName(TEXT("SPathPicker")))
+		{
+			TSharedPtr<SPathPicker> PathPicker = StaticCastSharedPtr<SPathPicker>(Widget);
+			PathPicker->RefreshPathView();
+		}
+	}
+}
+
 void FContentBrowserSingleton::PopulateConfigValues()
 {
 	const FString ContentBrowserSection = TEXT("ContentBrowser");
diff --git a/Engine/Source/Editor/ContentBrowser/Private/ContentBrowserSingleton.h b/Engine/Source/Editor/ContentBrowser/Private/ContentBrowserSingleton.h
--- a/Engine/Source/Editor/ContentBrowser/Private/ContentBrowserSingleton.h
+++ b/Engine/Source/Editor/ContentBrowser/Private/ContentBrowserSingleton.h
@@ -73,6 +73,7 @@ public:
 	virtual void SaveContentBrowserSettings(TSharedPtr<SWidget> ContentBrowserWidget) override;
 	virtual void ExecuteRename(TSharedPtr<SWidget> PickerWidget) override;
 	virtual void ExecuteAddFolder(TSharedPtr<SWidget> PathPickerWidget) override;
+	virtual void RefreshPathView(TSharedPtr<SWidget> PathPickerWidget) override;
 
 	/** Gets the content browser singleton as a FContentBrowserSingleton */
 	static FContentBrowserSingleton& Get();
diff --git a/Engine/Source/Editor/ContentBrowser/Private/SPathPicker.cpp b/Engine/Source/Editor/ContentBrowser/Private/SPathPicker.cpp
--- a/Engine/Source/Editor/ContentBrowser/Private/SPathPicker.cpp
+++ b/Engine/Source/Editor/ContentBrowser/Private/SPathPicker.cpp
@@ -204,6 +204,11 @@ void SPathPicker::CreateNewFolder(FString FolderPath, FOnCreateNewFolder InOnCre
 	}
 }
 
+void SPathPicker::RefreshPathView()
+{
+	PathViewPtr->Populate(true);
+}
+
 void SPathPicker::SetPaths(const TArray<FString>& NewPaths)
 {
 	PathViewPtr->SetSelectedPaths(NewPaths);
diff --git a/Engine/Source/Editor/ContentBrowser/Private/SPathPicker.h b/Engine/Source/Editor/ContentBrowser/Private/SPathPicker.h
--- a/Engine/Source/Editor/ContentBrowser/Private/SPathPicker.h
+++ b/Engine/Source/Editor/ContentBrowser/Private/SPathPicker.h
@@ -43,6 +43,9 @@ public:
 	/** Add a folder to the first selected folder*/
 	void ExecuteAddFolder();
 
+	/** Refresh the path view*/
+	void RefreshPathView();
+
 private:
 
 	/** Handle for when selection changes */
diff --git a/Engine/Source/Editor/ContentBrowser/Public/IContentBrowserSingleton.h b/Engine/Source/Editor/ContentBrowser/Public/IContentBrowserSingleton.h
--- a/Engine/Source/Editor/ContentBrowser/Public/IContentBrowserSingleton.h
+++ b/Engine/Source/Editor/ContentBrowser/Public/IContentBrowserSingleton.h
@@ -660,4 +660,11 @@ public:
 	 * @param PathPickerWidget The path picker widget where we want to add an folder
 	 */
 	virtual void ExecuteAddFolder(TSharedPtr<SWidget> PathPickerWidget) = 0;
+
+	/**
+	 * Force refresh on the path picker widget.  You may need to do this if you have changed the filter/white list on the path picker.
+	 *
+	 * @param PathPickerWidget The path picker widget where we want to add an folder
+	*/
+	virtual void RefreshPathView(TSharedPtr<SWidget> PathPickerWidget) = 0;
 };
\ No newline at end of file
