Fixed "Add feature or content pack" not appearing in the content browser +new menu

diff --git a/Engine/Source/Editor/ContentBrowser/Private/SContentBrowser.cpp b/Engine/Source/Editor/ContentBrowser/Private/SContentBrowser.cpp
--- a/Engine/Source/Editor/ContentBrowser/Private/SContentBrowser.cpp
+++ b/Engine/Source/Editor/ContentBrowser/Private/SContentBrowser.cpp
@@ -2078,13 +2078,25 @@ void SContentBrowser::AppendNewMenuContextObjects(const EContentBrowserDataMenuC
 		UToolMenu* Menu = UToolMenus::Get()->RegisterMenu("ContentBrowser.AddNewContextMenu");
 		Menu->AddDynamicSection("DynamicSection_Common", FNewToolMenuDelegate::CreateLambda([](UToolMenu* InMenu)
 		{
-			if (const UContentBrowserMenuContext* ContextObject = InMenu->FindContext<UContentBrowserMenuContext>())
+			TSharedPtr<SContentBrowser> ContentBrowser;
+			const UContentBrowserMenuContext* MenuContext = InMenu->FindContext<UContentBrowserMenuContext>();
+			if (MenuContext)
 			{
-				if (TSharedPtr<SContentBrowser> ContentBrowser = ContextObject->ContentBrowser.Pin())
+				ContentBrowser = MenuContext->ContentBrowser.Pin();
+			}
+			else
+			{
+				const UContentBrowserToolbarMenuContext* ToolbarContext  = InMenu->FindContext<UContentBrowserToolbarMenuContext>();
+				if (ToolbarContext)
 				{
-					ContentBrowser->PopulateAddNewContextMenu(InMenu);
+					ContentBrowser = ToolbarContext->ContentBrowser.Pin();
 				}
 			}
+
+			if (ContentBrowser)
+			{
+				ContentBrowser->PopulateAddNewContextMenu(InMenu);
+			}
 		}));
 	}
 
