Updates TableView Background style, adds frames for SceneOutliner and ContentBrowser tree

diff --git a/Engine/Source/Editor/ContentBrowser/Private/SContentBrowser.cpp b/Engine/Source/Editor/ContentBrowser/Private/SContentBrowser.cpp
--- a/Engine/Source/Editor/ContentBrowser/Private/SContentBrowser.cpp
+++ b/Engine/Source/Editor/ContentBrowser/Private/SContentBrowser.cpp
@@ -326,10 +326,18 @@ void SContentBrowser::Construct( const FArguments& InArgs, const FName& InInstan
 				]
 			]
 		]
+
+		+SVerticalBox::Slot()
+		.AutoHeight()
+		[
+			SNew(SSeparator)
+			.Thickness(2.0f)
+		]
+
 		// Assets/tree
 		+ SVerticalBox::Slot()
 		.FillHeight(1.0f)
-		.Padding(0.0f,2.0f, 0.0f, 0.0f)
+		.Padding(0.0f)
 		[
 			// The tree/assets splitter
 			SAssignNew(PathAssetSplitterPtr, SSplitter)
@@ -339,53 +347,57 @@ void SContentBrowser::Construct( const FArguments& InArgs, const FName& InInstan
 			+ SSplitter::Slot()
 			.Value(0.15f)
 			[
-				SNew(SBorder)
-				.Padding(FMargin(0,0))
-				.BorderImage(FEditorStyle::GetBrush("Brushes.Background"))
-				.Visibility(this, &SContentBrowser::GetSourcesViewVisibility)
+				SNew(SBox)
+				.Padding(FMargin(4.f))
 				[
-					// Note: If adding more widgets here, fix ContentBrowserSourcesWidgetSwitcherIndex and the code that uses it!
-					SAssignNew(SourcesWidgetSwitcher, SWidgetSwitcher)
-
-					// Paths View
-					+SWidgetSwitcher::Slot()
+					SNew(SBorder)
+					.Padding(FMargin(0))
+					.BorderImage(FEditorStyle::GetBrush("Brushes.Recessed"))
+					.Visibility(this, &SContentBrowser::GetSourcesViewVisibility)
 					[
-						SAssignNew(PathFavoriteSplitterPtr, SSplitter)
-						.Clipping(EWidgetClipping::ClipToBounds)
-						.PhysicalSplitterHandleSize(3.0f)
-						.HitDetectionSplitterHandleSize(3.0f)
-						.Orientation(EOrientation::Orient_Vertical)
-						.MinimumSlotHeight(26.0f)
-						.Visibility( this, &SContentBrowser::GetSourcesViewVisibility )
-						+SSplitter::Slot()
-						.SizeRule(TAttribute<SSplitter::ESizeRule>(this, &SContentBrowser::GetFavoritesAreaSizeRule))
-						.Value(0.2f)
-						[
-							CreateFavoritesView(Config)
-						]
-							
-						+SSplitter::Slot()
-						.SizeRule(TAttribute<SSplitter::ESizeRule>(this, &SContentBrowser::GetPathAreaSizeRule))
-						.Value(0.8f)
-						[
-							CreatePathView(Config)
-						]
+						// Note: If adding more widgets here, fix ContentBrowserSourcesWidgetSwitcherIndex and the code that uses it!
+						SAssignNew(SourcesWidgetSwitcher, SWidgetSwitcher)
 
-						+SSplitter::Slot()
-						.SizeRule(TAttribute<SSplitter::ESizeRule>(this, &SContentBrowser::GetCollectionsAreaSizeRule))
-						.Value(0.4f)
+						// Paths View
+						+SWidgetSwitcher::Slot()
 						[
-							CreateDockedCollectionsView(Config)
+							SAssignNew(PathFavoriteSplitterPtr, SSplitter)
+							.Clipping(EWidgetClipping::ClipToBounds)
+							.PhysicalSplitterHandleSize(3.0f)
+							.HitDetectionSplitterHandleSize(3.0f)
+							.Orientation(EOrientation::Orient_Vertical)
+							.MinimumSlotHeight(26.0f)
+							.Visibility( this, &SContentBrowser::GetSourcesViewVisibility )
+							+SSplitter::Slot()
+							.SizeRule(TAttribute<SSplitter::ESizeRule>(this, &SContentBrowser::GetFavoritesAreaSizeRule))
+							.Value(0.2f)
+							[
+								CreateFavoritesView(Config)
+							]
+								
+							+SSplitter::Slot()
+							.SizeRule(TAttribute<SSplitter::ESizeRule>(this, &SContentBrowser::GetPathAreaSizeRule))
+							.Value(0.8f)
+							[
+								CreatePathView(Config)
+							]
+
+							+SSplitter::Slot()
+							.SizeRule(TAttribute<SSplitter::ESizeRule>(this, &SContentBrowser::GetCollectionsAreaSizeRule))
+							.Value(0.4f)
+							[
+								CreateDockedCollectionsView(Config)
+							]
 						]
-					]
 
-					// Collections View
-					+SWidgetSwitcher::Slot()
-					[
-						SNew(SBox)
-						.Padding(FMargin(0.0f, 4.0f, 0.0f, 0.0f))
+						// Collections View
+						+SWidgetSwitcher::Slot()
 						[
-							CollectionViewPtr.ToSharedRef()
+							SNew(SBox)
+							.Padding(FMargin(0.0f, 4.0f, 0.0f, 0.0f))
+							[
+								CollectionViewPtr.ToSharedRef()
+							]
 						]
 					]
 				]
@@ -717,8 +729,8 @@ TSharedRef<SWidget> SContentBrowser::CreateFavoritesView(const FContentBrowserCo
 {
 	return
 		SAssignNew(FavoritesArea, SExpandableArea)
-		.BorderImage(FAppStyle::Get().GetBrush("Brushes.Background"))
-		.BodyBorderImage(FAppStyle::Get().GetBrush("Brushes.Background"))
+		.BorderImage(FAppStyle::Get().GetBrush("Brushes.Header"))
+		.BodyBorderImage(FAppStyle::Get().GetBrush("Brushes.Recessed"))
 		.HeaderPadding(FMargin(5.0f, 3.0f))
 		.Visibility(this, &SContentBrowser::GetFavoriteFolderVisibility)
 		.AllowAnimatedTransition(false)
@@ -748,8 +760,8 @@ TSharedRef<SWidget> SContentBrowser::CreatePathView(const FContentBrowserConfig*
 {	
 	return
 		SAssignNew(PathArea, SExpandableArea)
-		.BorderImage(FAppStyle::Get().GetBrush("Brushes.Background"))
-		.BodyBorderImage(FAppStyle::Get().GetBrush("Brushes.Background"))
+		.BorderImage(FAppStyle::Get().GetBrush("Brushes.Header"))
+		.BodyBorderImage(FAppStyle::Get().GetBrush("Brushes.Recessed"))
 		.HeaderPadding(FMargin(5.0f, 3.0f))
 		.AllowAnimatedTransition(false)
 		.HeaderContent()
@@ -791,8 +803,8 @@ TSharedRef<SWidget> SContentBrowser::CreateDockedCollectionsView(const FContentB
 {
 	return
 		SAssignNew(CollectionArea, SExpandableArea)
-		.BorderImage(FAppStyle::Get().GetBrush("Brushes.Background"))
-		.BodyBorderImage(FAppStyle::Get().GetBrush("Brushes.Background"))
+		.BorderImage(FAppStyle::Get().GetBrush("Brushes.Header"))
+		.BodyBorderImage(FAppStyle::Get().GetBrush("Brushes.Recessed"))
 		.HeaderPadding(FMargin(5.0f, 3.0f))
 		.Visibility(this, &SContentBrowser::GetDockedCollectionsVisibility)
 		.AllowAnimatedTransition(false)
diff --git a/Engine/Source/Editor/SceneOutliner/Private/SSceneOutliner.cpp b/Engine/Source/Editor/SceneOutliner/Private/SSceneOutliner.cpp
--- a/Engine/Source/Editor/SceneOutliner/Private/SSceneOutliner.cpp
+++ b/Engine/Source/Editor/SceneOutliner/Private/SSceneOutliner.cpp
@@ -160,14 +160,29 @@ void SSceneOutliner::Construct(const FArguments& InArgs, const FSceneOutlinerIni
 
 	VerticalBox->AddSlot()
 	.AutoHeight()
-	.Padding( 8.0f, 4.0f, 4.0f, 2.0f )
+	.Padding( 8.0f, 8.0f, 8.0f, 4.0f )
 	[
 		Toolbar
 	];
 
+	TSharedRef<SVerticalBox> InnerBox = SNew(SVerticalBox);
+
 	VerticalBox->AddSlot()
 	.FillHeight( 1.0f )
+	.Padding(4.0, 4.0)
 	[
+		SNew(SBorder)
+		.BorderImage( FAppStyle::Get().GetBrush("Brushes.Recessed") )
+		.Padding(0.0)
+		[
+			InnerBox
+		]
+	];	
+
+	InnerBox->AddSlot()
+	.FillHeight(1.0)
+	[
+
 		SNew( SOverlay )
 		+SOverlay::Slot()
 		.HAlign( HAlign_Center )
@@ -180,7 +195,7 @@ void SSceneOutliner::Construct(const FArguments& InArgs, const FSceneOutlinerIni
 
 		+SOverlay::Slot()
 		[
-			SNew(SBorder).BorderImage( FAppStyle::Get().GetBrush("Brushes.Background") )
+			SNew(SBorder).BorderImage( FAppStyle::Get().GetBrush("Brushes.Recessed") )
 		]
 
 		+SOverlay::Slot()
@@ -229,16 +244,14 @@ void SSceneOutliner::Construct(const FArguments& InArgs, const FSceneOutlinerIni
 	// Bottom panel status bar, if enabled by the mode
 	if (Mode->ShowStatusBar())
 	{
-		VerticalBox->AddSlot()
+		InnerBox->AddSlot()
 		.AutoHeight()
 		[
-			SNew(SHorizontalBox)
-
-			// Asset count
-			+SHorizontalBox::Slot()
-			.FillWidth(1.f)
+			SNew(SBorder)
+			.BorderImage( FAppStyle::Get().GetBrush("Brushes.Header") )
 			.VAlign(VAlign_Center)
-			.Padding(14, 9)
+			.HAlign(HAlign_Left)
+			.Padding(FMargin(14, 9))
 			[
 				SNew( STextBlock )
 				.Text(this, &SSceneOutliner::GetFilterStatusText)
diff --git a/Engine/Source/Runtime/SlateCore/Private/Styling/StarshipCoreStyle.cpp b/Engine/Source/Runtime/SlateCore/Private/Styling/StarshipCoreStyle.cpp
--- a/Engine/Source/Runtime/SlateCore/Private/Styling/StarshipCoreStyle.cpp
+++ b/Engine/Source/Runtime/SlateCore/Private/Styling/StarshipCoreStyle.cpp
@@ -1634,10 +1634,10 @@ void FStarshipCoreStyle::SetupTableViewStyles(TSharedRef<FStyle>& Style)
 	const FSlateColor SelectorColor = Style->GetSlateColor("SelectorColor");
 
 	const FTableRowStyle DefaultTableRowStyle = FTableRowStyle()
-		.SetEvenRowBackgroundBrush(FSlateColorBrush(FStyleColors::Background))
+		.SetEvenRowBackgroundBrush(FSlateColorBrush(FStyleColors::Recessed))
 		.SetEvenRowBackgroundHoveredBrush(FSlateColorBrush(FStyleColors::SelectHover))
 
-		.SetOddRowBackgroundBrush(FSlateColorBrush(FStyleColors::Background))
+		.SetOddRowBackgroundBrush(FSlateColorBrush(FStyleColors::Recessed))
 		.SetOddRowBackgroundHoveredBrush(FSlateColorBrush(FStyleColors::SelectHover))
 
 		.SetSelectorFocusedBrush(BORDER_BRUSH("Common/Selector", FMargin(4.f / 16.f), SelectorColor))
@@ -1706,7 +1706,7 @@ void FStarshipCoreStyle::SetupTableViewStyles(TSharedRef<FStyle>& Style)
 		.SetLastColumnStyle(TableLastColumnHeaderStyle)
 		.SetColumnSplitterStyle(TableHeaderSplitterStyle)
 		.SetSplitterHandleSize(1.0)
-		.SetBackgroundBrush(FSlateColorBrush(FStyleColors::Background))
+		.SetBackgroundBrush(FSlateColorBrush(FStyleColors::Recessed))
 		.SetForegroundColor(FStyleColors::Foreground)
 		.SetHorizontalSeparatorBrush(FSlateColorBrush(FStyleColors::Recessed))
 		.SetHorizontalSeparatorThickness(2.0f)
