Fixing PVS 7.7 Issues:

warning V547: Expression 'XYZ' is always true/false.

diff --git a/Engine/Plugins/FX/Niagara/Source/NiagaraEditor/Private/Widgets/SNiagaraParameterMapView.cpp b/Engine/Plugins/FX/Niagara/Source/NiagaraEditor/Private/Widgets/SNiagaraParameterMapView.cpp
--- a/Engine/Plugins/FX/Niagara/Source/NiagaraEditor/Private/Widgets/SNiagaraParameterMapView.cpp
+++ b/Engine/Plugins/FX/Niagara/Source/NiagaraEditor/Private/Widgets/SNiagaraParameterMapView.cpp
@@ -1925,7 +1925,7 @@ void SNiagaraParameterMapView::RenameParameter(TSharedPtr<FNiagaraParameterActio
 			}
 		}
 	}
-	else if (ToolkitType == SYSTEM)
+	else if (ToolkitType == SYSTEM) //-V547
 	{
 		UNiagaraSystem* System = CachedSystem.Get();
 		if (System != nullptr)
diff --git a/Engine/Source/Developer/ProjectLauncher/Private/Widgets/Deploy/SProjectLauncherDeployPage.cpp b/Engine/Source/Developer/ProjectLauncher/Private/Widgets/Deploy/SProjectLauncherDeployPage.cpp
--- a/Engine/Source/Developer/ProjectLauncher/Private/Widgets/Deploy/SProjectLauncherDeployPage.cpp
+++ b/Engine/Source/Developer/ProjectLauncher/Private/Widgets/Deploy/SProjectLauncherDeployPage.cpp
@@ -165,29 +165,19 @@ FText SProjectLauncherDeployPage::HandleDeploymentModeComboButtonContentText() c
 
 	if (SelectedProfile.IsValid())
 	{
-		ELauncherProfileDeploymentModes::Type DeploymentMode = SelectedProfile->GetDeploymentMode();
-
-		if (DeploymentMode == ELauncherProfileDeploymentModes::CopyToDevice)
+		switch (SelectedProfile->GetDeploymentMode())
 		{
+		case ELauncherProfileDeploymentModes::CopyToDevice:
 			return LOCTEXT("CopyToDeviceAction", "Copy to device");
-		}
-
-		if (DeploymentMode == ELauncherProfileDeploymentModes::DoNotDeploy)
-		{
+		case ELauncherProfileDeploymentModes::DoNotDeploy:
 			return LOCTEXT("DoNotDeployAction", "Do not deploy");
-		}
-
-		if (DeploymentMode == ELauncherProfileDeploymentModes::FileServer)
-		{
+		case ELauncherProfileDeploymentModes::FileServer:
 			return LOCTEXT("FileServerAction", "File server");
-		}
-
-		if (DeploymentMode == ELauncherProfileDeploymentModes::CopyRepository)
-		{
+		case ELauncherProfileDeploymentModes::CopyRepository:
 			return LOCTEXT("CopyRepositoryAction", "Copy repository");
+		default:
+			return LOCTEXT("DeploymentModeComboButtonDefaultText", "Select...");
 		}
-
-		return LOCTEXT("DeploymentModeComboButtonDefaultText", "Select...");
 	}
 
 	return FText::GetEmpty();
diff --git a/Engine/Source/Developer/ShaderCompilerCommon/Private/HlslParser.cpp b/Engine/Source/Developer/ShaderCompilerCommon/Private/HlslParser.cpp
--- a/Engine/Source/Developer/ShaderCompilerCommon/Private/HlslParser.cpp
+++ b/Engine/Source/Developer/ShaderCompilerCommon/Private/HlslParser.cpp
@@ -1082,7 +1082,7 @@ namespace CrossCompiler
 			{
 				break;
 			}
-			else if (Result == EParseResult::NotMatched)
+			else if (Result == EParseResult::NotMatched) //-V547
 			{
 				Parser.Scanner.SourceError(TEXT("Internal error on function parameter!\n"));
 				return ParseResultError();
diff --git a/Engine/Source/Developer/ShaderFormatOpenGL/Private/GlslBackend.cpp b/Engine/Source/Developer/ShaderFormatOpenGL/Private/GlslBackend.cpp
--- a/Engine/Source/Developer/ShaderFormatOpenGL/Private/GlslBackend.cpp
+++ b/Engine/Source/Developer/ShaderFormatOpenGL/Private/GlslBackend.cpp
@@ -1786,7 +1786,7 @@ class ir_gen_glsl_visitor : public ir_visitor
 				}
 			}
 		}
-		else if ( deref->op == ir_image_dimensions)
+		else if ( deref->op == ir_image_dimensions) //-V547
 		{
 			check(!bIsStructured);
 			ralloc_asprintf_append( buffer, "imageSize( " );
diff --git a/Engine/Source/Developer/VulkanShaderFormat/Private/VulkanBackend.cpp b/Engine/Source/Developer/VulkanShaderFormat/Private/VulkanBackend.cpp
--- a/Engine/Source/Developer/VulkanShaderFormat/Private/VulkanBackend.cpp
+++ b/Engine/Source/Developer/VulkanShaderFormat/Private/VulkanBackend.cpp
@@ -2031,7 +2031,7 @@ class FGenerateVulkanVisitor : public ir_visitor
 				}
 			}
 		}
-		else if (deref->op == ir_image_dimensions)
+		else if (deref->op == ir_image_dimensions) //-V547
 		{
 			check(!bIsStructured);
 			ralloc_asprintf_append(buffer, "imageSize( ");
diff --git a/Engine/Source/Editor/ContentBrowser/Private/SAssetDialog.cpp b/Engine/Source/Editor/ContentBrowser/Private/SAssetDialog.cpp
--- a/Engine/Source/Editor/ContentBrowser/Private/SAssetDialog.cpp
+++ b/Engine/Source/Editor/ContentBrowser/Private/SAssetDialog.cpp
@@ -101,7 +101,9 @@ void SAssetDialog::Construct(const FArguments& InArgs, const FSharedAssetDialogC
 	// Open and save specific configuration
 	FText ConfirmButtonText;
 	bool bIncludeNameBox = false;
-	if ( DialogType == EAssetDialogType::Open )
+	switch (DialogType)
+	{
+	case EAssetDialogType::Open:
 	{
 		const FOpenAssetDialogConfig& OpenAssetConfig = static_cast<const FOpenAssetDialogConfig&>(InConfig);
 		PathPickerConfig.bAllowContextMenu = true;
@@ -109,8 +111,10 @@ void SAssetDialog::Construct(const FArguments& InArgs, const FSharedAssetDialogC
 		AssetPickerConfig.SelectionMode = OpenAssetConfig.bAllowMultipleSelection ? ESelectionMode::Multi : ESelectionMode::Single;
 		AssetPickerConfig.bFocusSearchBoxWhenOpened = true;
 		bIncludeNameBox = false;
+		break;
 	}
-	else if ( DialogType == EAssetDialogType::Save )
+
+	case EAssetDialogType::Save:
 	{
 		const FSaveAssetDialogConfig& SaveAssetConfig = static_cast<const FSaveAssetDialogConfig&>(InConfig);
 		PathPickerConfig.bAllowContextMenu = true;
@@ -120,10 +124,12 @@ void SAssetDialog::Construct(const FArguments& InArgs, const FSharedAssetDialogC
 		bIncludeNameBox = true;
 		ExistingAssetPolicy = SaveAssetConfig.ExistingAssetPolicy;
 		SetCurrentlyEnteredAssetName(SaveAssetConfig.DefaultAssetName);
+		break;
 	}
-	else
-	{
+
+	default:
 		ensureMsgf(0, TEXT("AssetDialog type %d is not supported."), DialogType);
+		break;
 	}
 
 	PathPicker = StaticCastSharedRef<SPathPicker>(FContentBrowserSingleton::Get().CreatePathPicker(PathPickerConfig));
@@ -649,42 +655,41 @@ void SAssetDialog::HandleAssetViewFolderEntered(const FString& NewPath)
 
 bool SAssetDialog::IsConfirmButtonEnabled() const
 {
-	if ( DialogType == EAssetDialogType::Open )
+	switch (DialogType)
 	{
+	case EAssetDialogType::Open:
 		return CurrentlySelectedAssets.Num() > 0;
-	}
-	else if ( DialogType == EAssetDialogType::Save )
-	{
+	case EAssetDialogType::Save:
 		return bLastInputValidityCheckSuccessful;
+	default:
+	    ensureMsgf(0, TEXT("AssetDialog type %d is not supported."), DialogType);
+	    return false;
 	}
-	else
-	{
-		ensureMsgf(0, TEXT("AssetDialog type %d is not supported."), DialogType);
-	}
-
-	return false;
 }
 
 FReply SAssetDialog::OnConfirmClicked()
 {
-	if ( DialogType == EAssetDialogType::Open )
+	switch (DialogType)
+	{
+	case EAssetDialogType::Open:
 	{
 		TArray<FAssetData> SelectedAssets = GetCurrentSelectionDelegate.Execute();
 		if (SelectedAssets.Num() > 0)
 		{
 			ChooseAssetsForOpen(SelectedAssets);
 		}
+		break;
 	}
-	else if ( DialogType == EAssetDialogType::Save )
-	{
+
+	case EAssetDialogType::Save:
 		// @todo save asset validation (e.g. "asset already exists" check)
 		CommitObjectPathForSave();
-	}
-	else
-	{
+		break;
+
+	default:
 		ensureMsgf(0, TEXT("AssetDialog type %d is not supported."), DialogType);
+		break;
 	}
-
 	return FReply::Handled();
 }
 
@@ -711,20 +716,24 @@ void SAssetDialog::OnAssetsActivated(const TArray<FAssetData>& SelectedAssets, E
 	const bool bCorrectActivationMethod = (ActivationType == EAssetTypeActivationMethod::DoubleClicked || ActivationType == EAssetTypeActivationMethod::Opened);
 	if (SelectedAssets.Num() > 0 && bCorrectActivationMethod)
 	{
-		if ( DialogType == EAssetDialogType::Open )
+		switch (DialogType)
 		{
+		case EAssetDialogType::Open:
 			ChooseAssetsForOpen(SelectedAssets);
-		}
-		else if ( DialogType == EAssetDialogType::Save )
+			break;
+
+		case EAssetDialogType::Save:
 		{
 			const FAssetData& AssetData = SelectedAssets[0];
 			SetCurrentlySelectedPath(AssetData.PackagePath.ToString());
 			SetCurrentlyEnteredAssetName(AssetData.AssetName.ToString());
 			CommitObjectPathForSave();
+			break;
 		}
-		else
-		{
+
+		default:
 			ensureMsgf(0, TEXT("AssetDialog type %d is not supported."), DialogType);
+			break;
 		}
 	}
 }
diff --git a/Engine/Source/Runtime/D3D12RHI/Private/D3D12RayTracing.cpp b/Engine/Source/Runtime/D3D12RHI/Private/D3D12RayTracing.cpp
--- a/Engine/Source/Runtime/D3D12RHI/Private/D3D12RayTracing.cpp
+++ b/Engine/Source/Runtime/D3D12RHI/Private/D3D12RayTracing.cpp
@@ -2479,8 +2479,9 @@ void FD3D12RayTracingGeometry::BuildAccelerationStructure(FD3D12CommandContext&
 		// #dxr_todo: Add explicit vertex count to FRayTracingGeometrySegment.
 		const uint32 MaxSegmentVertices = (VertexBuffer->ResourceLocation.GetSize() - Segment.VertexBufferOffset) / Segment.VertexBufferStride;
 
-		if (GeometryType == D3D12_RAYTRACING_GEOMETRY_TYPE_TRIANGLES)
+		switch (GeometryType)
 		{
+		case D3D12_RAYTRACING_GEOMETRY_TYPE_TRIANGLES:
 			switch (Segment.VertexBufferElementType)
 			{
 			case VET_Float4:
@@ -2534,16 +2535,17 @@ void FD3D12RayTracingGeometry::BuildAccelerationStructure(FD3D12CommandContext&
 
 			Desc.Triangles.VertexBuffer.StartAddress = SystemParameters.VertexBuffer;
 			Desc.Triangles.VertexBuffer.StrideInBytes = Segment.VertexBufferStride;
-		}
-		else if (GeometryType == D3D12_RAYTRACING_GEOMETRY_TYPE_PROCEDURAL_PRIMITIVE_AABBS)
-		{
+			break;
+
+		case D3D12_RAYTRACING_GEOMETRY_TYPE_PROCEDURAL_PRIMITIVE_AABBS:
 			Desc.AABBs.AABBCount = Segment.NumPrimitives;
 			Desc.AABBs.AABBs.StartAddress = SystemParameters.VertexBuffer;
 			Desc.AABBs.AABBs.StrideInBytes = Segment.VertexBufferStride;
-		}
-		else
-		{
+			break;
+
+		default:
 			checkf(false, TEXT("Unexpected ray tracing geometry type"));
+			break;
 		}
 
 		VertexBuffer->ResourceLocation.GetResource()->UpdateResidency(CommandContext.CommandListHandle);
diff --git a/Engine/Source/Runtime/Engine/Private/Components/SkinnedMeshComponent.cpp b/Engine/Source/Runtime/Engine/Private/Components/SkinnedMeshComponent.cpp
--- a/Engine/Source/Runtime/Engine/Private/Components/SkinnedMeshComponent.cpp
+++ b/Engine/Source/Runtime/Engine/Private/Components/SkinnedMeshComponent.cpp
@@ -2212,7 +2212,9 @@ FVector USkinnedMeshComponent::GetBoneLocation(FName BoneName, EBoneSpaces::Type
 		return FVector::ZeroVector;
 	}
 
-	if( Space == EBoneSpaces::ComponentSpace )
+	switch (Space)
+	{
+	case EBoneSpaces::ComponentSpace:
 	{
 		const USkinnedMeshComponent* const MasterPoseComponentInst = MasterPoseComponent.Get();
 		if(MasterPoseComponentInst)
@@ -2227,22 +2229,21 @@ FVector USkinnedMeshComponent::GetBoneLocation(FName BoneName, EBoneSpaces::Type
 					return MasterPoseComponentInst->GetComponentSpaceTransforms()[ParentBoneIndex].GetLocation();
 				}
 			}
-			
+
 			// return empty vector
-			return FVector::ZeroVector;			
+			return FVector::ZeroVector;
 		}
 		else
 		{
 			return GetComponentSpaceTransforms()[BoneIndex].GetLocation();
 		}
 	}
-	else if (Space == EBoneSpaces::WorldSpace)
-	{
+
+	case EBoneSpaces::WorldSpace:
 		// To support non-uniform scale (via LocalToWorld), use GetBoneMatrix
 		return GetBoneMatrix(BoneIndex).GetOrigin();
-	}
-	else
-	{
+
+	default:
 		check(false); // Unknown BoneSpace
 		return FVector::ZeroVector;
 	}
diff --git a/Engine/Source/Runtime/OpenGLDrv/Private/OpenGLQuery.cpp b/Engine/Source/Runtime/OpenGLDrv/Private/OpenGLQuery.cpp
--- a/Engine/Source/Runtime/OpenGLDrv/Private/OpenGLQuery.cpp
+++ b/Engine/Source/Runtime/OpenGLDrv/Private/OpenGLQuery.cpp
@@ -712,24 +712,25 @@ void FOpenGLEventQuery::WaitForCompletion()
 	// Wait up to 1/2 second for sync execution
 	FOpenGL::EFenceResult Status = FOpenGL::ClientWaitSync( Sync, 0, 500*1000*1000);
 
-	if ( Status != FOpenGL::FR_AlreadySignaled && Status != FOpenGL::FR_ConditionSatisfied )
+	switch (Status)
 	{
-		//failure of some type, determine type and send diagnostic message
-		if ( Status == FOpenGL::FR_TimeoutExpired )
-		{
-			UE_LOG(LogRHI, Log, TEXT("Timed out while waiting for GPU to catch up. (500 ms)"));
-		}
-		else if ( Status == FOpenGL::FR_WaitFailed )
-		{
-			UE_LOG(LogRHI, Log, TEXT("Wait on GPU failed in driver"));
-		}
-		else
-		{
-			UE_LOG(LogRHI, Log, TEXT("Unknown error while waiting on GPU"));
-			check(0);
-		}
-	}
-
+	case FOpenGL::FR_AlreadySignaled:
+	case FOpenGL::FR_ConditionSatisfied:
+		break;
+
+	case FOpenGL::FR_TimeoutExpired:
+		UE_LOG(LogRHI, Log, TEXT("Timed out while waiting for GPU to catch up. (500 ms)"));
+		break;
+
+	case FOpenGL::FR_WaitFailed:
+		UE_LOG(LogRHI, Log, TEXT("Wait on GPU failed in driver"));
+		break;
+
+	default:
+	    UE_LOG(LogRHI, Log, TEXT("Unknown error while waiting on GPU"));
+	    check(0);
+		break;
+	}	
 }
 
 void FOpenGLEventQuery::InitDynamicRHI()
diff --git a/Engine/Source/Runtime/Slate/Private/Framework/Docking/TabManager.cpp b/Engine/Source/Runtime/Slate/Private/Framework/Docking/TabManager.cpp
--- a/Engine/Source/Runtime/Slate/Private/Framework/Docking/TabManager.cpp
+++ b/Engine/Source/Runtime/Slate/Private/Framework/Docking/TabManager.cpp
@@ -1200,37 +1200,49 @@ void FTabManager::InsertNewDocumentTab(FName PlaceholderId, const FSearchPrefere
 
 void FTabManager::InsertNewDocumentTab( FName PlaceholderId, ESearchPreference::Type SearchPreference, const TSharedRef<SDockTab>& UnmanagedTab )
 {
-	if ( SearchPreference == ESearchPreference::PreferLiveTab )
+	switch (SearchPreference)
+	{
+	case ESearchPreference::PreferLiveTab:
 	{
 		FLiveTabSearch Search;
 		InsertDocumentTab(PlaceholderId, Search, UnmanagedTab, true);
+		break;
 	}
-	else if ( SearchPreference == ESearchPreference::RequireClosedTab )
+
+	case ESearchPreference::RequireClosedTab:
 	{
 		FRequireClosedTab Search;
 		InsertDocumentTab(PlaceholderId, Search, UnmanagedTab, true);
+		break;
 	}
-	else
-	{
+
+	default:
 		check(false);
+		break;
 	}
 }
 
 void FTabManager::RestoreDocumentTab( FName PlaceholderId, ESearchPreference::Type SearchPreference, const TSharedRef<SDockTab>& UnmanagedTab )
 {
-	if ( SearchPreference == ESearchPreference::PreferLiveTab )
+	switch (SearchPreference)
+	{
+	case ESearchPreference::PreferLiveTab:
 	{
 		FLiveTabSearch Search;
 		InsertDocumentTab(PlaceholderId, Search, UnmanagedTab, false);
+		break;
 	}
-	else if ( SearchPreference == ESearchPreference::RequireClosedTab )
+
+	case ESearchPreference::RequireClosedTab:
 	{
 		FRequireClosedTab Search;
 		InsertDocumentTab(PlaceholderId, Search, UnmanagedTab, false);
+		break;
 	}
-	else
-	{
+
+	default:
 		check(false);
+		break;
 	}
 }
 
