Content Browser New button changed to use SEditorHeaderButton and text changed to "Add".

diff --git a/Engine/Source/Editor/ContentBrowser/Private/SContentBrowser.cpp b/Engine/Source/Editor/ContentBrowser/Private/SContentBrowser.cpp
--- a/Engine/Source/Editor/ContentBrowser/Private/SContentBrowser.cpp
+++ b/Engine/Source/Editor/ContentBrowser/Private/SContentBrowser.cpp
@@ -72,7 +72,7 @@
 #include "ToolMenu.h"
 #include "Widgets/Input/SExpandableButton.h"
 #include "SExpandableSearchArea.h"
-
+#include "SEditorHeaderButton.h"
 
 #define LOCTEXT_NAMESPACE "ContentBrowser"
 
@@ -910,53 +910,28 @@ void SContentBrowser::RegisterContentBrowserToolBar()
 		FToolMenuSection& Section = ToolBar->AddSection("New");
 		
 		Section.AddDynamicEntry("New", FNewToolMenuSectionDelegate::CreateLambda([](FToolMenuSection& InSection)
-			{
-				UContentBrowserToolbarMenuContext* Context = InSection.FindContext<UContentBrowserToolbarMenuContext>();
-				TSharedRef<SContentBrowser> ContentBrowser = Context->ContentBrowser.Pin().ToSharedRef();
+		{
+			UContentBrowserToolbarMenuContext* Context = InSection.FindContext<UContentBrowserToolbarMenuContext>();
+			TSharedRef<SContentBrowser> ContentBrowser = Context->ContentBrowser.Pin().ToSharedRef();
 		
-				TSharedRef<SComboButton> NewButton =
-					SNew(SComboButton)
-					.ComboButtonStyle(&FAppStyle::Get().GetWidgetStyle<FComboButtonStyle>("SimpleComboButton"))
-					.ForegroundColor(FSlateColor::UseStyle())
-					.ContentPadding(2)
-					.OnGetMenuContent_Lambda([Context] { return Context->ContentBrowser.Pin()->MakeAddNewContextMenu(EContentBrowserDataMenuContext_AddNewMenuDomain::Toolbar, Context); })
-					.ToolTipText(ContentBrowser, &SContentBrowser::GetAddNewToolTipText)
-					.IsEnabled(ContentBrowser, &SContentBrowser::IsAddNewEnabled)
-					.AddMetaData<FTagMetaData>(FTagMetaData(TEXT("ContentBrowserNewAsset")))
-					.HasDownArrow(false)
-					.ButtonContent()
-					[
-						SNew(SHorizontalBox)
-						+ SHorizontalBox::Slot()
-						.AutoWidth()
-						.HAlign(HAlign_Center)
-						.VAlign(VAlign_Center)
-						[
-							SNew(SImage)
-							.Image(FAppStyle::Get().GetBrush("Icons.Plus"))
-							.ColorAndOpacity(FSlateColor::UseForeground())
-						]
-						+ SHorizontalBox::Slot()
-						.Padding(FMargin(3, 0, 0, 0))
-						.VAlign(VAlign_Center)
-						.AutoWidth()
-						[
-							SNew(STextBlock)
-							.TextStyle(FAppStyle::Get(), "NormalText")
-							.Text(LOCTEXT("NewAssetButton", "New"))
-						]
-					];
+			TSharedRef<SEditorHeaderButton> NewButton = SNew(SEditorHeaderButton)
+				.OnGetMenuContent_Lambda([Context] { return Context->ContentBrowser.Pin()->MakeAddNewContextMenu(EContentBrowserDataMenuContext_AddNewMenuDomain::Toolbar, Context); })
+				.ToolTipText(ContentBrowser, &SContentBrowser::GetAddNewToolTipText)
+				.IsEnabled(ContentBrowser, &SContentBrowser::IsAddNewEnabled)
+				.AddMetaData<FTagMetaData>(FTagMetaData(TEXT("ContentBrowserNewAsset")))
+				.Icon(FAppStyle::Get().GetBrush("Icons.Plus"))
+				.Text(LOCTEXT("AddAssetButton", "Add"));
 		
-				InSection.AddEntry(
-					FToolMenuEntry::InitWidget(
-						"NewButton",
-						NewButton,
-						FText::GetEmpty(),
-						true,
-						false
-					));
-			}));
-		}
+			InSection.AddEntry(
+				FToolMenuEntry::InitWidget(
+					"NewButton",
+					NewButton,
+					FText::GetEmpty(),
+					true,
+					false
+				));
+		}));
+	}
 
 	{
 		FToolMenuSection& Section = ToolBar->AddSection("Save");
