Fix content browser settings icon being too hard to see with no text label
Fix content browser tile medium size thumbnails not allowing 3 lines for the asset name

diff --git a/Engine/Source/Editor/ContentBrowser/Private/AssetViewWidgets.cpp b/Engine/Source/Editor/ContentBrowser/Private/AssetViewWidgets.cpp
--- a/Engine/Source/Editor/ContentBrowser/Private/AssetViewWidgets.cpp
+++ b/Engine/Source/Editor/ContentBrowser/Private/AssetViewWidgets.cpp
@@ -38,6 +38,7 @@
 #include "ContentBrowserUtils.h"
 #include "ContentBrowserDataSource.h"
 #include "Widgets/Images/SLayeredImage.h"
+#include "Fonts/FontMeasure.h"
 
 #define LOCTEXT_NAMESPACE "ContentBrowser"
 
@@ -435,10 +436,10 @@ void SAssetViewItem::Tick( const FGeometry& AllottedGeometry, const double InCur
 		const float WrapWidth = GetNameTextWrapWidth();
 		InlineRenameWidget->SetWrapTextAt(WrapWidth);
 
-		if (ClassTextWidget.IsValid())
+		/*if (ClassTextWidget.IsValid())
 		{
 			ClassTextWidget->SetWrapTextAt(WrapWidth);
-		}
+		}*/
 	}
 
 	UpdateDirtyState();
@@ -1641,6 +1642,10 @@ void SAssetTileItem::Construct( const FArguments& InArgs )
 	ItemWidth = InArgs._ItemWidth;
 	ThumbnailPadding = IsFolder() ? InArgs._ThumbnailPadding + 5.0f : InArgs._ThumbnailPadding;
 
+	CurrentThumbnailSize = InArgs._CurrentThumbnailSize;
+
+	InitializeAssetNameHeights();
+
 	TSharedPtr<SWidget> Thumbnail;
 	if ( AssetItem.IsValid() && AssetThumbnail.IsValid() )
 	{
@@ -1719,24 +1724,29 @@ void SAssetTileItem::Construct( const FArguments& InArgs )
 						[
 							SNew(SVerticalBox)
 							+ SVerticalBox::Slot()
-							.Padding(2.0f)
+							.Padding(2.0f,0.0f,0.0f,0.0f)
 							.VAlign(VAlign_Top)
 							[
-								SAssignNew(InlineRenameWidget, SInlineEditableTextBlock)
-								.Font(this, &SAssetTileItem::GetThumbnailFont)
-								.Text(GetNameText())
-								.OnBeginTextEdit(this, &SAssetTileItem::HandleBeginNameChange)
-								.OnTextCommitted(this, &SAssetTileItem::HandleNameCommitted)
-								.OnVerifyTextChanged(this, &SAssetTileItem::HandleVerifyNameChanged)
-								.HighlightText(InArgs._HighlightText)
-								.IsSelected(InArgs._IsSelectedExclusively)
-								.IsReadOnly(this, &SAssetTileItem::IsNameReadOnly)
-								.Justification(IsFolder() ? ETextJustify::Center : ETextJustify::Left)
-								.LineBreakPolicy(FBreakIterator::CreateCamelCaseBreakIterator())
+								SNew(SBox)
+								.MaxDesiredHeight(AssetNameHeights[(int32)CurrentThumbnailSize])
+								[
+									SAssignNew(InlineRenameWidget, SInlineEditableTextBlock)
+									.Font(this, &SAssetTileItem::GetThumbnailFont)
+									.Text(GetNameText())
+									.OnBeginTextEdit(this, &SAssetTileItem::HandleBeginNameChange)
+									.OnTextCommitted(this, &SAssetTileItem::HandleNameCommitted)
+									.OnVerifyTextChanged(this, &SAssetTileItem::HandleVerifyNameChanged)
+									.HighlightText(InArgs._HighlightText)
+									.IsSelected(InArgs._IsSelectedExclusively)
+									.IsReadOnly(this, &SAssetTileItem::IsNameReadOnly)
+									.Justification(IsFolder() ? ETextJustify::Center : ETextJustify::Left)
+									.LineBreakPolicy(FBreakIterator::CreateCamelCaseBreakIterator())
+								]
 							]
 							+ SVerticalBox::Slot()
 							.VAlign(VAlign_Bottom)
-							.Padding(2.0f)
+							.AutoHeight()
+							.Padding(4.0f,0.0f, 0.0f, 2.0f)
 							[
 								SAssignNew(ClassTextWidget, STextBlock)
 								.Visibility(this, &SAssetTileItem::GetAssetClassLabelVisibility)
@@ -1744,7 +1754,7 @@ void SAssetTileItem::Construct( const FArguments& InArgs )
 								.TransformPolicy(ETextTransformPolicy::ToUpper)
 								.Text(this, &SAssetTileItem::GetAssetClassText)
 								.ColorAndOpacity(FSlateColor::UseSubduedForeground())
-								.LineBreakPolicy(FBreakIterator::CreateWordBreakIterator())
+								//.LineBreakPolicy(FBreakIterator::CreateWordBreakIterator())
 							]
 							
 						]
@@ -1885,8 +1895,6 @@ FSlateFontInfo SAssetTileItem::GetThumbnailFont() const
 	return FEditorStyle::GetFontStyle(RegularFont);
 }
 
-
-
 const FSlateBrush* SAssetTileItem::GetFolderBackgroundImage() const
 {
 	const bool bIsSelected = IsSelected.IsBound() ? IsSelected.Execute() : false;
@@ -1915,6 +1923,41 @@ const FSlateBrush* SAssetTileItem::GetFolderBackgroundShadowImage() const
 	return FStyleDefaults::GetNoBrush();
 }
 
+void SAssetTileItem::InitializeAssetNameHeights()
+{
+	// The height of the asset name field for each thumbnail size
+	static bool bInitializedHeights = false;
+
+	if (!bInitializedHeights)
+	{
+		AssetNameHeights[(int32)EThumbnailSize::Tiny] = 0;
+
+		{
+			const static FName SmallFontName("ContentBrowser.AssetTileViewNameFontSmall");
+			FSlateFontInfo Font = FEditorStyle::GetFontStyle(SmallFontName);
+			TSharedRef<FSlateFontMeasure> FontMeasureService = FSlateApplication::Get().GetRenderer()->GetFontMeasureService();
+			float MaxHeight = FontMeasureService->GetMaxCharacterHeight(Font);
+			AssetNameHeights[(int32)EThumbnailSize::Small] = MaxHeight * 2;
+		}
+
+
+		{
+			const static FName SmallFontName("ContentBrowser.AssetTileViewNameFont");
+			FSlateFontInfo Font = FEditorStyle::GetFontStyle(SmallFontName);
+			TSharedRef<FSlateFontMeasure> FontMeasureService = FSlateApplication::Get().GetRenderer()->GetFontMeasureService();
+			float MaxHeight = FontMeasureService->GetMaxCharacterHeight(Font);
+			AssetNameHeights[(int32)EThumbnailSize::Medium] = MaxHeight * 3;
+			AssetNameHeights[(int32)EThumbnailSize::Large] = MaxHeight * 4;
+			AssetNameHeights[(int32)EThumbnailSize::Huge] = MaxHeight * 5;
+		}
+
+		bInitializedHeights = true;
+	}
+
+}
+
+float SAssetTileItem::AssetNameHeights[(int32)EThumbnailSize::MAX];
+
 ///////////////////////////////
 // SAssetColumnItem
 ///////////////////////////////
diff --git a/Engine/Source/Editor/ContentBrowser/Private/AssetViewWidgets.h b/Engine/Source/Editor/ContentBrowser/Private/AssetViewWidgets.h
--- a/Engine/Source/Editor/ContentBrowser/Private/AssetViewWidgets.h
+++ b/Engine/Source/Editor/ContentBrowser/Private/AssetViewWidgets.h
@@ -42,6 +42,20 @@ DECLARE_DELEGATE_OneParam( FOnItemDestroyed, const TSharedPtr<FAssetViewItem>& /
 class SAssetListItem;
 class SAssetTileItem;
 
+enum class EThumbnailSize : uint8
+{
+	Tiny = 0,
+	Small,
+	Medium,
+	Large,
+	Huge,
+
+	// Not a size
+	MAX
+};
+
+
+
 namespace FAssetViewModeUtils 
 {
 	FReply OnViewModeKeyDown( const TSet< TSharedPtr<FAssetViewItem> >& SelectedItems, const FKeyEvent& InKeyEvent );
@@ -456,6 +470,7 @@ public:
 		, _ThumbnailLabel( EThumbnailLabel::ClassName )
 		, _ThumbnailHintColorAndOpacity( FLinearColor( 0.0f, 0.0f, 0.0f, 0.0f ) )
 		, _AllowThumbnailHintLabel(true)
+		, _CurrentThumbnailSize(EThumbnailSize::Medium)
 		, _ItemWidth(16)
 		, _ShouldAllowToolTip(true)
 		, _ThumbnailEditMode(false)
@@ -479,6 +494,9 @@ public:
 		/** Whether the thumbnail should ever show it's hint label */
 		SLATE_ARGUMENT( bool, AllowThumbnailHintLabel )
 
+		/** Current size of the thumbnail that was generated */
+		SLATE_ARGUMENT(EThumbnailSize, CurrentThumbnailSize)
+
 		/** The width of the item */
 		SLATE_ATTRIBUTE( float, ItemWidth )
 
@@ -565,8 +583,14 @@ protected:
 	/** Returns the font to use for the thumbnail label */
 	FSlateFontInfo GetThumbnailFont() const;
 
+	/** Gets the max size of the name area for an asset */
+	FOptionalSize GetMaxAssetNameHeight() const;
+
 	const FSlateBrush* GetFolderBackgroundImage() const;
 	const FSlateBrush* GetFolderBackgroundShadowImage() const;
+
+	static void InitializeAssetNameHeights();
+
 private:
 	/** The handle to the thumbnail that this item is rendering */
 	TSharedPtr<FAssetThumbnail> AssetThumbnail;
@@ -574,8 +598,14 @@ private:
 	/** The width of the item. Used to enforce a square thumbnail. */
 	TAttribute<float> ItemWidth;
 
+	/** Max name height for each thumbnail size */
+	static float AssetNameHeights[(int32)EThumbnailSize::MAX];
+
 	/** The padding allotted for the thumbnail */
 	float ThumbnailPadding;
+
+	/** Current thumbnail size when this widget was generated */
+	EThumbnailSize CurrentThumbnailSize;
 };
 
 /** An item in the asset column view */
diff --git a/Engine/Source/Editor/ContentBrowser/Private/SAssetView.cpp b/Engine/Source/Editor/ContentBrowser/Private/SAssetView.cpp
--- a/Engine/Source/Editor/ContentBrowser/Private/SAssetView.cpp
+++ b/Engine/Source/Editor/ContentBrowser/Private/SAssetView.cpp
@@ -3207,6 +3207,7 @@ TSharedRef<ITableRow> SAssetView::MakeTileViewWidget(TSharedPtr<FAssetViewItem>
 			.AssetThumbnail(AssetThumbnail)
 			.AssetItem(AssetItem)
 			.ThumbnailPadding(TileViewThumbnailPadding)
+			.CurrentThumbnailSize(ThumbnailSize)
 			.ItemWidth(this, &SAssetView::GetTileViewItemWidth)
 			.OnRenameBegin(this, &SAssetView::AssetRenameBegin)
 			.OnRenameCommit(this, &SAssetView::AssetRenameCommit)
diff --git a/Engine/Source/Editor/ContentBrowser/Private/SAssetView.h b/Engine/Source/Editor/ContentBrowser/Private/SAssetView.h
--- a/Engine/Source/Editor/ContentBrowser/Private/SAssetView.h
+++ b/Engine/Source/Editor/ContentBrowser/Private/SAssetView.h
@@ -46,18 +46,7 @@ DECLARE_DELEGATE(FOnSearchOptionChanged);
 /** Fires whenever asset view options menu is being opened, gives chance for external code to set additional context */
 DECLARE_DELEGATE_OneParam(FOnExtendAssetViewOptionsMenuContext, FToolMenuContext&);
 
-enum class EThumbnailSize : uint8
-{
-	Tiny = 0,
-	Small,
-	Medium,
-	Large,
-	Huge,
-
-	// Not a size
-	MAX
-};
-
+enum class EThumbnailSize : uint8;
 
 /**
  * A widget to display a list of filtered assets
diff --git a/Engine/Source/Editor/ContentBrowser/Private/SContentBrowser.cpp b/Engine/Source/Editor/ContentBrowser/Private/SContentBrowser.cpp
--- a/Engine/Source/Editor/ContentBrowser/Private/SContentBrowser.cpp
+++ b/Engine/Source/Editor/ContentBrowser/Private/SContentBrowser.cpp
@@ -206,9 +206,24 @@ void SContentBrowser::Construct( const FArguments& InArgs, const FName& InInstan
 			.HasDownArrow(false)
 			.ButtonContent()
 			[
-				SNew(SImage)
-				.ColorAndOpacity(FSlateColor::UseForeground())
-				.Image(FAppStyle::Get().GetBrush("Icons.Settings"))
+				SNew(SHorizontalBox)
+				+ SHorizontalBox::Slot()
+				.AutoWidth()
+				.VAlign(VAlign_Center)
+				.Padding(4.0, 0.0f)
+				[
+					SNew(SImage)
+					.ColorAndOpacity(FSlateColor::UseForeground())
+					.Image(FAppStyle::Get().GetBrush("Icons.Settings"))
+				]
+				+ SHorizontalBox::Slot()
+				.VAlign(VAlign_Center)
+				.Padding(4.0, 0.0f)
+				[
+					SNew(STextBlock)
+					.Text(LOCTEXT("Settings", "Settings"))
+					.ColorAndOpacity(FSlateColor::UseForeground())
+				]
 			];
 	}
 
