Edigrate 15459510.
Don't allow assets to be created in content browser while a rename is occuring.

diff --git a/Engine/Source/Editor/ContentBrowser/Private/SAssetView.cpp b/Engine/Source/Editor/ContentBrowser/Private/SAssetView.cpp
--- a/Engine/Source/Editor/ContentBrowser/Private/SAssetView.cpp
+++ b/Engine/Source/Editor/ContentBrowser/Private/SAssetView.cpp
@@ -568,6 +568,12 @@ void SAssetView::AppendBackendFilter(FARFilter& FilterToAppendTo) const
 
 void SAssetView::NewFolderItemRequested(const FContentBrowserItemTemporaryContext& NewItemContext)
 {
+	// Don't allow asset creation while renaming
+	if (IsRenamingAsset())
+	{
+		return;
+	}
+
 	// we should only be creating one deferred folder per tick
 	check(!DeferredItemToCreate.IsValid());
 
@@ -590,6 +596,12 @@ void SAssetView::NewFolderItemRequested(const FContentBrowserItemTemporaryContex
 
 void SAssetView::NewFileItemRequested(const FContentBrowserItemDataTemporaryContext& NewItemContext)
 {
+	// Don't allow asset creation while renaming
+	if (IsRenamingAsset())
+	{
+		return;
+	}
+
 	// We should only be creating one deferred file at a time
 	check(!DeferredItemToCreate.IsValid());
 
