Fix inconsistent sizing with categories in the content browser path view

diff --git a/Engine/Source/Editor/ContentBrowser/Private/SContentBrowser.cpp b/Engine/Source/Editor/ContentBrowser/Private/SContentBrowser.cpp
--- a/Engine/Source/Editor/ContentBrowser/Private/SContentBrowser.cpp
+++ b/Engine/Source/Editor/ContentBrowser/Private/SContentBrowser.cpp
@@ -747,7 +747,7 @@ TSharedRef<SWidget> SContentBrowser::CreateFavoritesView(const FContentBrowserCo
 		SAssignNew(FavoritesArea, SExpandableArea)
 		.BorderImage(FAppStyle::Get().GetBrush("Brushes.Header"))
 		.BodyBorderImage(FAppStyle::Get().GetBrush("Brushes.Recessed"))
-		.HeaderPadding(FMargin(5.0f, 3.0f))
+		.HeaderPadding(FMargin(5.0f, 7.0f))
 		.Visibility(this, &SContentBrowser::GetFavoriteFolderVisibility)
 		.AllowAnimatedTransition(false)
 		.HeaderContent()
