Fix browse button doing nothing if a content browser is not docked. It should open the drawer

diff --git a/Engine/Source/Editor/ContentBrowser/Private/ContentBrowserSingleton.cpp b/Engine/Source/Editor/ContentBrowser/Private/ContentBrowserSingleton.cpp
--- a/Engine/Source/Editor/ContentBrowser/Private/ContentBrowserSingleton.cpp
+++ b/Engine/Source/Editor/ContentBrowser/Private/ContentBrowserSingleton.cpp
@@ -615,12 +615,13 @@ void FContentBrowserSingleton::SharedCreateAssetDialogWindow(const TSharedRef<SA
 
 void FContentBrowserSingleton::ChooseNewPrimaryBrowser()
 {
+	// The drawer is the lowest priority. If a docked content browser can become the primary browser choose it first
 	// Find the first valid browser and trim any invalid browsers along the way
 	for (int32 BrowserIdx = 0; BrowserIdx < AllContentBrowsers.Num(); ++BrowserIdx)
 	{
-		if ( AllContentBrowsers[BrowserIdx].IsValid() )
+		if (TSharedPtr<SContentBrowser> TestBrowser = AllContentBrowsers[BrowserIdx].Pin())
 		{
-			if (AllContentBrowsers[BrowserIdx].Pin()->CanSetAsPrimaryContentBrowser())
+			if (TestBrowser->CanSetAsPrimaryContentBrowser() && TestBrowser != ContentBrowserDrawer)
 			{
 				SetPrimaryContentBrowser(AllContentBrowsers[BrowserIdx].Pin().ToSharedRef());
 				break;
@@ -633,6 +634,12 @@ void FContentBrowserSingleton::ChooseNewPrimaryBrowser()
 			BrowserIdx--;
 		}
 	}
+
+	// Set the content browser drawer as the primary browser
+	if (!PrimaryContentBrowser.IsValid())
+	{
+		PrimaryContentBrowser = ContentBrowserDrawer;
+	}
 }
 
 void FContentBrowserSingleton::FocusContentBrowser(const TSharedPtr<SContentBrowser>& BrowserToFocus)
@@ -758,7 +765,8 @@ TSharedRef<SWidget> FContentBrowserSingleton::CreateContentBrowser( const FName
 
 	AllContentBrowsers.Add( NewBrowser );
 
-	if( !PrimaryContentBrowser.IsValid() )
+	// The drawer is the lowest priority. If a new content browser is being added see if it is capable of becoming the primary browser
+	if( !PrimaryContentBrowser.IsValid() || PrimaryContentBrowser == ContentBrowserDrawer)
 	{
 		ChooseNewPrimaryBrowser();
 	}
@@ -778,11 +786,6 @@ TSharedRef<SDockTab> FContentBrowserSingleton::SpawnContentBrowserTab( const FSp
 
 	TSharedRef<SWidget> NewBrowser = CreateContentBrowser( SpawnTabArgs.GetTabId().TabType, NewTab, nullptr );
 
-	if ( !PrimaryContentBrowser.IsValid() )
-	{
-		ChooseNewPrimaryBrowser();
-	}
-
 	// Add wrapper for tutorial highlighting
 	TSharedRef<SBox> Wrapper =
 		SNew(SBox)
diff --git a/Engine/Source/Editor/StatusBar/Private/StatusBarSubsystem.cpp b/Engine/Source/Editor/StatusBar/Private/StatusBarSubsystem.cpp
--- a/Engine/Source/Editor/StatusBar/Private/StatusBarSubsystem.cpp
+++ b/Engine/Source/Editor/StatusBar/Private/StatusBarSubsystem.cpp
@@ -448,7 +448,7 @@ void UStatusBarSubsystem::CreateContentBrowserIfNeeded()
 		IContentBrowserSingleton& ContentBrowserSingleton = FModuleManager::Get().LoadModuleChecked<FContentBrowserModule>("ContentBrowser").Get();;
 
 		FContentBrowserConfig Config;
-		Config.bCanSetAsPrimaryBrowser = false;
+		Config.bCanSetAsPrimaryBrowser = true;
 
 		TFunction<TSharedPtr<SDockTab>()> GetTab(
 			[this]() -> TSharedPtr<SDockTab>
