Content Browser filters "Not Used in Any Level" & "Used in Any Level" were not checking for references by external actors

diff --git a/Engine/Source/Editor/ContentBrowser/Private/FrontendFilters.cpp b/Engine/Source/Editor/ContentBrowser/Private/FrontendFilters.cpp
--- a/Engine/Source/Editor/ContentBrowser/Private/FrontendFilters.cpp
+++ b/Engine/Source/Editor/ContentBrowser/Private/FrontendFilters.cpp
@@ -1316,9 +1316,11 @@ void FFrontendFilter_UsedInAnyLevel::ActiveStateChanged(bool bActive)
 
 	if (bActive)
 	{
-		// Find all the levels
+		// Find all the levels & external actors
 		FARFilter Filter;
 		Filter.ClassNames.Add(UWorld::StaticClass()->GetFName());
+		Filter.ClassNames.Add(AActor::StaticClass()->GetFName());
+		Filter.bRecursiveClasses = true;
 		FrontendFilterHelper::GetDependencies(Filter, *AssetRegistry, LevelsDependencies);
 	}
 }
@@ -1358,9 +1360,11 @@ void FFrontendFilter_NotUsedInAnyLevel::ActiveStateChanged(bool bActive)
 	
 	if (bActive)
 	{
-		// Find all the levels
+		// Find all the levels & external actors
 		FARFilter Filter;
 		Filter.ClassNames.Add(UWorld::StaticClass()->GetFName());
+		Filter.ClassNames.Add(AActor::StaticClass()->GetFName());
+		Filter.bRecursiveClasses = true;
 		FrontendFilterHelper::GetDependencies(Filter, *AssetRegistry, LevelsDependencies);
 	}
 }
