Fixed case sensitive sorting of folders in path view (not consistent with the folder view)

#rnx

diff --git a/Engine/Source/Editor/ContentBrowser/Private/SPathView.cpp b/Engine/Source/Editor/ContentBrowser/Private/SPathView.cpp
--- a/Engine/Source/Editor/ContentBrowser/Private/SPathView.cpp
+++ b/Engine/Source/Editor/ContentBrowser/Private/SPathView.cpp
@@ -1492,7 +1492,7 @@ void SPathView::DefaultSort(const FTreeItem* InTreeItem, TArray<TSharedPtr<FTree
 		else
 		{
 			// Two non special folders of the same priority, sort alphabetically
-			const int32 CompareResult = SortInfoA.FolderName.Compare(SortInfoB.FolderName);
+			const int32 CompareResult = SortInfoA.FolderName.Compare(SortInfoB.FolderName, ESearchCase::IgnoreCase);
 			if (CompareResult != 0)
 			{
 				return CompareResult < 0;
