Add common utility function to determine if we can explore folders

diff --git a/Engine/Source/Editor/ContentBrowser/Private/ContentBrowserUtils.cpp b/Engine/Source/Editor/ContentBrowser/Private/ContentBrowserUtils.cpp
--- a/Engine/Source/Editor/ContentBrowser/Private/ContentBrowserUtils.cpp
+++ b/Engine/Source/Editor/ContentBrowser/Private/ContentBrowserUtils.cpp
@@ -517,6 +517,24 @@ void ContentBrowserUtils::ExploreFolders(const TArray<FContentBrowserItem>& InIt
 	}
 }
 
+bool ContentBrowserUtils::CanExploreFolders(const TArray<FContentBrowserItem>& InItems)
+{
+	for (const FContentBrowserItem& SelectedItem : InItems)
+	{
+		FString ItemFilename;
+		if (SelectedItem.GetItemPhysicalPath(ItemFilename))
+		{
+			const bool bExists = SelectedItem.IsFile() ? FPaths::FileExists(ItemFilename) : FPaths::DirectoryExists(ItemFilename);
+			if (bExists)
+			{
+				return true;
+			}
+		}
+	}
+
+	return false;
+}
+
 template <typename OutputContainerType>
 void ConvertLegacySelectionToVirtualPathsImpl(TArrayView<const FAssetData> InAssets, TArrayView<const FString> InFolders, const bool InUseFolderPaths, OutputContainerType& OutVirtualPaths)
 {
diff --git a/Engine/Source/Editor/ContentBrowser/Private/ContentBrowserUtils.h b/Engine/Source/Editor/ContentBrowser/Private/ContentBrowserUtils.h
--- a/Engine/Source/Editor/ContentBrowser/Private/ContentBrowserUtils.h
+++ b/Engine/Source/Editor/ContentBrowser/Private/ContentBrowserUtils.h
@@ -70,6 +70,9 @@ namespace ContentBrowserUtils
 	/** Perform a batched "explore" operation on the specified file and/or folder paths */
 	void ExploreFolders(const TArray<FContentBrowserItem>& InItems, const TSharedRef<SWidget>& InParentContent);
 
+	/** Returns if can perform a batched "explore" operation on the specified file and/or folder paths */
+	bool CanExploreFolders(const TArray<FContentBrowserItem>& InItems);
+
 	/** Convert a legacy asset and path selection to their corresponding virtual paths for content browser data items */
 	void ConvertLegacySelectionToVirtualPaths(TArrayView<const FAssetData> InAssets, TArrayView<const FString> InFolders, const bool InUseFolderPaths, TArray<FName>& OutVirtualPaths);
 	void ConvertLegacySelectionToVirtualPaths(TArrayView<const FAssetData> InAssets, TArrayView<const FString> InFolders, const bool InUseFolderPaths, TSet<FName>& OutVirtualPaths);
