enable overflow ellipsis on tabs and content browser tiles

diff --git a/Engine/Source/Editor/ContentBrowser/Private/AssetViewWidgets.cpp b/Engine/Source/Editor/ContentBrowser/Private/AssetViewWidgets.cpp
--- a/Engine/Source/Editor/ContentBrowser/Private/AssetViewWidgets.cpp
+++ b/Engine/Source/Editor/ContentBrowser/Private/AssetViewWidgets.cpp
@@ -436,11 +436,6 @@ void SAssetViewItem::Tick( const FGeometry& AllottedGeometry, const double InCur
 	{
 		const float WrapWidth = GetNameTextWrapWidth();
 		InlineRenameWidget->SetWrapTextAt(WrapWidth);
-
-		/*if (ClassTextWidget.IsValid())
-		{
-			ClassTextWidget->SetWrapTextAt(WrapWidth);
-		}*/
 	}
 
 	UpdateDirtyState();
@@ -1749,6 +1744,7 @@ void SAssetTileItem::Construct( const FArguments& InArgs )
 									.IsReadOnly(this, &SAssetTileItem::IsNameReadOnly)
 									.Justification(IsFolder() ? ETextJustify::Center : ETextJustify::Left)
 									.LineBreakPolicy(FBreakIterator::CreateCamelCaseBreakIterator())
+									.OverflowPolicy(ETextOverflowPolicy::Ellipsis)
 								]
 							]
 							+ SVerticalBox::Slot()
diff --git a/Engine/Source/Runtime/Slate/Private/Widgets/Docking/SDockTab.cpp b/Engine/Source/Runtime/Slate/Private/Widgets/Docking/SDockTab.cpp
--- a/Engine/Source/Runtime/Slate/Private/Widgets/Docking/SDockTab.cpp
+++ b/Engine/Source/Runtime/Slate/Private/Widgets/Docking/SDockTab.cpp
@@ -651,6 +651,7 @@ void SDockTab::Construct( const FArguments& InArgs )
 						SAssignNew(LabelWidget, STextBlock)
 						.TextStyle(&GetCurrentStyle().TabTextStyle)
 						.Text(this, &SDockTab::GetTabLabel)
+						.OverflowPolicy(ETextOverflowPolicy::Ellipsis)
 					]
 
 					// Tab Label Suffix
