Sort asset names in a way that handles leading zeroes in trailing numbers

diff --git a/Engine/Source/Editor/ContentBrowser/Private/AssetViewSortManager.cpp b/Engine/Source/Editor/ContentBrowser/Private/AssetViewSortManager.cpp
--- a/Engine/Source/Editor/ContentBrowser/Private/AssetViewSortManager.cpp
+++ b/Engine/Source/Editor/ContentBrowser/Private/AssetViewSortManager.cpp
@@ -115,7 +115,7 @@ protected:
 		const FContentBrowserItemData* BItemData = B->GetItem().GetPrimaryInternalItem();
 		// TODO: Have an option to sort by display name? It's slower, but more correct for non-English languages
 		//const int32 Result = (AItemData && BItemData) ? AItemData->GetDisplayName().CompareTo(BItemData->GetDisplayName()) : 0;
-		const int32 Result = (AItemData && BItemData) ? AItemData->GetItemName().Compare(BItemData->GetItemName()) : 0;
+		const int32 Result = (AItemData && BItemData) ? FNameBuilder(AItemData->GetItemName()).ToView().Compare(FNameBuilder(BItemData->GetItemName()).ToView()) : 0;
 		if (Result < 0)
 		{
 			return bAscending;
@@ -169,7 +169,7 @@ protected:
 	{
 		const FContentBrowserItemData* AItemData = A->GetItem().GetPrimaryInternalItem();
 		const FContentBrowserItemData* BItemData = B->GetItem().GetPrimaryInternalItem();
-		const int32 Result = (AItemData && BItemData) ? AItemData->GetVirtualPath().Compare(BItemData->GetVirtualPath()) : 0;
+		const int32 Result = (AItemData && BItemData) ? FNameBuilder(AItemData->GetVirtualPath()).ToView().Compare(FNameBuilder(BItemData->GetVirtualPath()).ToView()) : 0;
 		if (Result < 0)
 		{
 			return bAscending;
