Re-add adjusting Content Browser thumbnail size via ctrl+mouse scroll wheel

This previously worked in UE4, then was removed when the thumbnail size selector was changed from a continuous slider to discrete predefined options.

diff --git a/Engine/Source/Editor/ContentBrowser/Private/SAssetView.cpp b/Engine/Source/Editor/ContentBrowser/Private/SAssetView.cpp
--- a/Engine/Source/Editor/ContentBrowser/Private/SAssetView.cpp
+++ b/Engine/Source/Editor/ContentBrowser/Private/SAssetView.cpp
@@ -1599,16 +1599,19 @@ FReply SAssetView::OnKeyDown( const FGeometry& MyGeometry, const FKeyEvent& InKe
 
 FReply SAssetView::OnMouseWheel( const FGeometry& MyGeometry, const FPointerEvent& MouseEvent )
 {
-/*
-	if( MouseEvent.IsControlDown() )
-	{
-		const float DesiredScale = FMath::Clamp<float>(GetThumbnailScale() + ( MouseEvent.GetWheelDelta() * 0.05f ), 0.0f, 1.0f);
-		if ( DesiredScale != GetThumbnailScale() )
+	// Make sure to not change the thumbnail scaling when we're in Columns view since thumbnail scaling isn't applicable there.
+	if( MouseEvent.IsControlDown() && IsThumbnailScalingAllowed() )
+	{
+		// Step up/down a level depending on the scroll wheel direction.
+		// Clamp value to enum min/max before updating.
+		const int32 Delta = MouseEvent.GetWheelDelta() > 0 ? 1 : -1;
+		const EThumbnailSize DesiredThumbnailSize = (EThumbnailSize)FMath::Clamp<int32>((int32)ThumbnailSize + Delta, 0, (int32)EThumbnailSize::MAX - 1);
+		if ( DesiredThumbnailSize != ThumbnailSize )
 		{
-			SetThumbnailScale( DesiredScale );
+			OnThumbnailSizeChanged(DesiredThumbnailSize);
 		}		
 		return FReply::Handled();
-	}*/
+	}
 	return FReply::Unhandled();
 }
 
