Sort asset names in a way that handles all trailing numbers, including fully numeric names

diff --git a/Engine/Source/Editor/ContentBrowser/Private/AssetViewSortManager.cpp b/Engine/Source/Editor/ContentBrowser/Private/AssetViewSortManager.cpp
--- a/Engine/Source/Editor/ContentBrowser/Private/AssetViewSortManager.cpp
+++ b/Engine/Source/Editor/ContentBrowser/Private/AssetViewSortManager.cpp
@@ -115,7 +115,7 @@ protected:
 		const FContentBrowserItemData* BItemData = B->GetItem().GetPrimaryInternalItem();
 		// TODO: Have an option to sort by display name? It's slower, but more correct for non-English languages
 		//const int32 Result = (AItemData && BItemData) ? AItemData->GetDisplayName().CompareTo(BItemData->GetDisplayName()) : 0;
-		const int32 Result = (AItemData && BItemData) ? FNameBuilder(AItemData->GetItemName()).ToView().Compare(FNameBuilder(BItemData->GetItemName()).ToView()) : 0;
+		const int32 Result = (AItemData && BItemData) ? FAssetViewSortManager::CompareWithNumericSuffix(FNameBuilder(AItemData->GetItemName()).ToView(), FNameBuilder(BItemData->GetItemName()).ToView()) : 0;
 		if (Result < 0)
 		{
 			return bAscending;
@@ -169,7 +169,7 @@ protected:
 	{
 		const FContentBrowserItemData* AItemData = A->GetItem().GetPrimaryInternalItem();
 		const FContentBrowserItemData* BItemData = B->GetItem().GetPrimaryInternalItem();
-		const int32 Result = (AItemData && BItemData) ? FNameBuilder(AItemData->GetVirtualPath()).ToView().Compare(FNameBuilder(BItemData->GetVirtualPath()).ToView()) : 0;
+		const int32 Result = (AItemData && BItemData) ? FAssetViewSortManager::CompareWithNumericSuffix(FNameBuilder(AItemData->GetVirtualPath()).ToView(), FNameBuilder(BItemData->GetVirtualPath()).ToView()) : 0;
 		if (Result < 0)
 		{
 			return bAscending;
@@ -576,3 +576,44 @@ bool FAssetViewSortManager::SetOrToggleSortColumn(const EColumnSortPriority::Typ
 		return false;
 	}
 }
+
+int32 FAssetViewSortManager::CompareWithNumericSuffix(FStringView A, FStringView B)
+{
+	auto SplitNumericSuffix = [](FStringView Full, FStringView& OutPrefix, int32& OutSuffix)
+	{
+		int32 SuffixSplit = Full.Len();
+		for (; SuffixSplit > 0 && FChar::IsDigit(Full[SuffixSplit - 1]); --SuffixSplit) {}
+
+		if (SuffixSplit < Full.Len())
+		{
+			OutPrefix = Full.Left(SuffixSplit);
+
+			TStringBuilder<32> SuffixStr;
+			SuffixStr.Append(Full.RightChop(SuffixSplit));
+			LexFromString(OutSuffix, SuffixStr.ToString());
+		}
+		else
+		{
+			OutPrefix = Full;
+			OutSuffix = INDEX_NONE;
+		}
+	};
+
+	FStringView APrefix;
+	int32 ASuffix = 0;
+	SplitNumericSuffix(A, APrefix, ASuffix);
+
+	FStringView BPrefix;
+	int32 BSuffix = 0;
+	SplitNumericSuffix(B, BPrefix, BSuffix);
+
+	// Are the prefixes identical?
+	const int32 PrefixResult = APrefix.Compare(BPrefix, ESearchCase::IgnoreCase);
+	if (PrefixResult != 0)
+	{
+		return PrefixResult;
+	}
+
+	// If so, compare the suffixes too
+	return ASuffix - BSuffix;
+}
diff --git a/Engine/Source/Editor/ContentBrowser/Private/AssetViewSortManager.h b/Engine/Source/Editor/ContentBrowser/Private/AssetViewSortManager.h
--- a/Engine/Source/Editor/ContentBrowser/Private/AssetViewSortManager.h
+++ b/Engine/Source/Editor/ContentBrowser/Private/AssetViewSortManager.h
@@ -3,6 +3,7 @@
 #pragma once
 
 #include "CoreMinimal.h"
+#include "Containers/StringView.h"
 #include "Widgets/Views/SHeaderRow.h"
 #include "ContentBrowserDelegates.h"
 
@@ -47,6 +48,10 @@ public:
 	static const FName ClassColumnId;
 	static const FName PathColumnId;
 
+public:
+	/** Compare the two strings, correctly ordering any numeric suffixes they may have */
+	static int32 CompareWithNumericSuffix(FStringView A, FStringView B);
+
 private:
 	/** The name of the column that is currently used for sorting. */
 	FName SortColumnId[EColumnSortPriority::Max];
diff --git a/Engine/Source/Editor/ContentBrowser/Private/SPathView.cpp b/Engine/Source/Editor/ContentBrowser/Private/SPathView.cpp
--- a/Engine/Source/Editor/ContentBrowser/Private/SPathView.cpp
+++ b/Engine/Source/Editor/ContentBrowser/Private/SPathView.cpp
@@ -10,6 +10,7 @@
 #include "Settings/ContentBrowserSettings.h"
 #include "IAssetTools.h"
 #include "AssetToolsModule.h"
+#include "AssetViewSortManager.h"
 #include "ContentBrowserSingleton.h"
 #include "ContentBrowserUtils.h"
 #include "ContentBrowserLog.h"
@@ -24,6 +25,7 @@
 #include "Widgets/Input/SSearchBox.h"
 #include "ContentBrowserModule.h"
 #include "Misc/BlacklistNames.h"
+#include "Misc/PathViews.h"
 
 #include "IContentBrowserDataModule.h"
 #include "ContentBrowserDataSource.h"
@@ -31,7 +33,6 @@
 
 #include "Application/SlateApplicationBase.h"
 #include "ToolMenus.h"
-#include <Misc/PathViews.h>
 
 #define LOCTEXT_NAMESPACE "ContentBrowser"
 
@@ -1492,7 +1493,7 @@ void SPathView::DefaultSort(const FTreeItem* InTreeItem, TArray<TSharedPtr<FTree
 		else
 		{
 			// Two non special folders of the same priority, sort alphabetically
-			const int32 CompareResult = SortInfoA.FolderName.Compare(SortInfoB.FolderName, ESearchCase::IgnoreCase);
+			const int32 CompareResult = FAssetViewSortManager::CompareWithNumericSuffix(SortInfoA.FolderName, SortInfoB.FolderName);
 			if (CompareResult != 0)
 			{
 				return CompareResult < 0;
