Alphabetize category submenus and clicked categories in the Content Browser filter dropdown

diff --git a/Engine/Source/Editor/ContentBrowser/Private/SFilterList.cpp b/Engine/Source/Editor/ContentBrowser/Private/SFilterList.cpp
--- a/Engine/Source/Editor/ContentBrowser/Private/SFilterList.cpp
+++ b/Engine/Source/Editor/ContentBrowser/Private/SFilterList.cpp
@@ -1413,6 +1413,11 @@ void SFilterList::PopulateAddFilterMenu(UToolMenu* Menu)
 		FToolMenuSection& Section = Menu->AddSection("ContentBrowserFilterAdvancedAsset", LOCTEXT("AdvancedAssetsMenuHeading", "Other Assets"));
 		if(MenuExpansion == EAssetTypeCategories::Basic)
 		{
+			// Sort by category name so that we add the submenus in alphabetical order
+			CategoryToMenuMap.ValueSort([](const FCategoryMenu& A, const FCategoryMenu& B) {
+				return A.Name.CompareTo(B.Name) < 0;
+			});
+
 			// For all the remaining categories, add them as submenus
 			for (const TPair<EAssetTypeCategories::Type, FCategoryMenu>& CategoryMenuPair : CategoryToMenuMap)
 			{
@@ -1537,6 +1542,13 @@ void SFilterList::FilterByTypeCategoryClicked(EAssetTypeCategories::Type Categor
 	TArray<TWeakPtr<IAssetTypeActions>> TypeActionsList;
 	GetTypeActionsForCategory(Category, TypeActionsList);
 
+	// Sort the list of type actions so that we add new filters in alphabetical order
+	TypeActionsList.Sort([](const TWeakPtr<IAssetTypeActions>& A, const TWeakPtr<IAssetTypeActions>& B) {
+		const FText NameA = A.IsValid() ? A.Pin()->GetName() : FText::GetEmpty();
+		const FText NameB = B.IsValid() ? B.Pin()->GetName() : FText::GetEmpty();
+		return NameA.CompareTo(NameB) < 0;
+	});
+
 	bool bFullCategoryInUse = IsAssetTypeCategoryInUse(Category);
 	bool ExecuteOnFilterChanged = false;
 
