Add initial widget library in closed tab.
---
Expose AssetView from ContentBrowser
Add Library view model, partial reuse from pallete view model (needs more reuse).
Add support for widget asset dragdrop via code classes & to widget hierarchy.

diff --git a/Engine/Source/Editor/ContentBrowser/Private/AssetViewTypes.h b/Engine/Source/Editor/ContentBrowser/Private/AssetViewTypes.h
deleted file mode 100644
diff --git a/Engine/Source/Editor/ContentBrowser/Private/AssetViewWidgets.h b/Engine/Source/Editor/ContentBrowser/Private/AssetViewWidgets.h
deleted file mode 100644
diff --git a/Engine/Source/Editor/ContentBrowser/Private/FrontendFilters.h b/Engine/Source/Editor/ContentBrowser/Private/FrontendFilters.h
deleted file mode 100644
diff --git a/Engine/Source/Editor/ContentBrowser/Private/HistoryManager.h b/Engine/Source/Editor/ContentBrowser/Private/HistoryManager.h
deleted file mode 100644
diff --git a/Engine/Source/Editor/ContentBrowser/Private/SAssetPicker.h b/Engine/Source/Editor/ContentBrowser/Private/SAssetPicker.h
--- a/Engine/Source/Editor/ContentBrowser/Private/SAssetPicker.h
+++ b/Engine/Source/Editor/ContentBrowser/Private/SAssetPicker.h
@@ -10,7 +10,7 @@
 #include "AssetData.h"
 #include "ARFilter.h"
 #include "IContentBrowserSingleton.h"
-#include "Editor/ContentBrowser/Private/SourcesData.h"
+#include "SourcesData.h"
 
 class FFrontendFilter_ShowOtherDevelopers;
 class FFrontendFilter_Text;
diff --git a/Engine/Source/Editor/ContentBrowser/Private/SAssetView.h b/Engine/Source/Editor/ContentBrowser/Private/SAssetView.h
deleted file mode 100644
diff --git a/Engine/Source/Editor/ContentBrowser/Private/SContentBrowser.h b/Engine/Source/Editor/ContentBrowser/Private/SContentBrowser.h
--- a/Engine/Source/Editor/ContentBrowser/Private/SContentBrowser.h
+++ b/Engine/Source/Editor/ContentBrowser/Private/SContentBrowser.h
@@ -14,7 +14,7 @@
 #include "ContentBrowserDataMenuContexts.h"
 #include "CollectionManagerTypes.h"
 #include "IContentBrowserSingleton.h"
-#include "Editor/ContentBrowser/Private/HistoryManager.h"
+#include "HistoryManager.h"
 #include "ContentBrowserDataSubsystem.h"
 #include "Textures/SlateIcon.h"
 
diff --git a/Engine/Source/Editor/ContentBrowser/Private/SourcesData.h b/Engine/Source/Editor/ContentBrowser/Private/SourcesData.h
deleted file mode 100644
diff --git a/Engine/Source/Editor/ContentBrowser/Public/AssetViewTypes.h b/Engine/Source/Editor/ContentBrowser/Public/AssetViewTypes.h
new file mode 100644
--- /dev/null
+++ b/Engine/Source/Editor/ContentBrowser/Public/AssetViewTypes.h
@@ -0,0 +1,93 @@
+// Copyright Epic Games, Inc. All Rights Reserved.
+
+#pragma once
+
+#include "CoreMinimal.h"
+#include "ContentBrowserItem.h"
+
+struct FAssetViewCustomColumn;
+
+/** An item (folder or file) displayed in the asset view */
+class FAssetViewItem
+{
+public:
+	FAssetViewItem() = default;
+
+	explicit FAssetViewItem(FContentBrowserItem&& InItem);
+	explicit FAssetViewItem(const FContentBrowserItem& InItem);
+
+	explicit FAssetViewItem(FContentBrowserItemData&& InItemData);
+	explicit FAssetViewItem(const FContentBrowserItemData& InItemData);
+
+	void AppendItemData(const FContentBrowserItem& InItem);
+
+	void AppendItemData(const FContentBrowserItemData& InItemData);
+
+	void RemoveItemData(const FContentBrowserItem& InItem);
+
+	void RemoveItemData(const FContentBrowserItemData& InItemData);
+
+	/** Clear cached custom column data */
+	void ClearCachedCustomColumns();
+
+	/** Updates cached custom column data (only does something for files) */
+	void CacheCustomColumns(TArrayView<const FAssetViewCustomColumn> CustomColumns, const bool bUpdateSortData, const bool bUpdateDisplayText, const bool bUpdateExisting);
+	
+	/** Get the display value of a custom column on this item */
+	bool GetCustomColumnDisplayValue(const FName ColumnName, FText& OutText) const;
+
+	/** Get the value (and optionally also the type) of a custom column on this item */
+	bool GetCustomColumnValue(const FName ColumnName, FString& OutString, UObject::FAssetRegistryTag::ETagType* OutType = nullptr) const;
+
+	/** Get the value (and optionally also the type) of a named tag on this item */
+	bool GetTagValue(const FName Tag, FString& OutString, UObject::FAssetRegistryTag::ETagType* OutType = nullptr) const;
+
+	/** Get the underlying Content Browser item */
+	const FContentBrowserItem& GetItem() const;
+
+	bool IsFolder() const;
+
+	bool IsFile() const;
+
+	bool IsTemporary() const;
+
+	/** Get the event fired when the data for this item changes */
+	FSimpleMulticastDelegate& OnItemDataChanged();
+
+	/** Get the event fired whenever a rename is requested */
+	FSimpleDelegate& OnRenameRequested();
+
+	/** Get the event fired whenever a rename is canceled */
+	FSimpleDelegate& OnRenameCanceled();
+
+	/** True if this item should enter inline renaming on the next scroll into view */
+	bool ShouldRenameWhenScrolledIntoView() const;
+
+	/** Set that this item should enter inline renaming on the next scroll into view */
+	void RenameWhenScrolledIntoView();
+
+	/** Clear that this item should enter inline renaming on the next scroll into view */
+	void ClearRenameWhenScrolledIntoView();
+
+private:
+	/** Underlying Content Browser item data */
+	FContentBrowserItem Item;
+
+	/** An event to fire when the data for this item changes */
+	FSimpleMulticastDelegate ItemDataChangedEvent;
+
+	/** Broadcasts whenever a rename is requested */
+	FSimpleDelegate RenameRequestedEvent;
+
+	/** Broadcasts whenever a rename is canceled */
+	FSimpleDelegate RenameCanceledEvent;
+
+	/** True if this item should enter inline renaming on the next scroll into view */
+	bool bRenameWhenScrolledIntoView = false;
+
+	/** Map of values/types for custom columns */
+	TMap<FName, TTuple<FString, UObject::FAssetRegistryTag::ETagType>> CachedCustomColumnData;
+	
+	/** Map of display text for custom columns */
+	TMap<FName, FText> CachedCustomColumnDisplayText;
+};
diff --git a/Engine/Source/Editor/ContentBrowser/Public/AssetViewWidgets.h b/Engine/Source/Editor/ContentBrowser/Public/AssetViewWidgets.h
new file mode 100644
--- /dev/null
+++ b/Engine/Source/Editor/ContentBrowser/Public/AssetViewWidgets.h
@@ -0,0 +1,770 @@
+// Copyright Epic Games, Inc. All Rights Reserved.
+
+#pragma once
+
+#include "CoreMinimal.h"
+#include "SlateFwd.h"
+#include "Misc/Attribute.h"
+#include "Layout/Geometry.h"
+#include "Input/Reply.h"
+#include "Layout/Visibility.h"
+#include "Widgets/DeclarativeSyntaxSupport.h"
+#include "Styling/SlateColor.h"
+#include "Widgets/SWidget.h"
+#include "Widgets/SCompoundWidget.h"
+#include "AssetData.h"
+#include "AssetThumbnail.h"
+#include "Fonts/SlateFontInfo.h"
+#include "ContentBrowserDelegates.h"
+#include "Widgets/Views/STableViewBase.h"
+#include "Types/SlateStructs.h"
+#include "Widgets/SBoxPanel.h"
+#include "Widgets/Views/STableRow.h"
+#include "AssetViewSortManager.h"
+#include "Widgets/Text/STextBlock.h"
+#include "EditorStyleSet.h"
+#include "Widgets/Views/SListView.h"
+#include "Widgets/Views/STileView.h"
+#include "Textures/SlateIcon.h"
+
+class SLayeredImage;
+class SAssetListItem;
+class SAssetTileItem;
+class FAssetViewItem;
+
+template <typename ItemType> class SListView;
+
+DECLARE_DELEGATE_ThreeParams( FOnRenameBegin, const TSharedPtr<FAssetViewItem>& /*AssetItem*/, const FString& /*OriginalName*/, const FSlateRect& /*MessageAnchor*/)
+DECLARE_DELEGATE_FourParams( FOnRenameCommit, const TSharedPtr<FAssetViewItem>& /*AssetItem*/, const FString& /*NewName*/, const FSlateRect& /*MessageAnchor*/, ETextCommit::Type /*CommitType*/ )
+DECLARE_DELEGATE_RetVal_FourParams( bool, FOnVerifyRenameCommit, const TSharedPtr<FAssetViewItem>& /*AssetItem*/, const FText& /*NewName*/, const FSlateRect& /*MessageAnchor*/, FText& /*OutErrorMessage*/)
+DECLARE_DELEGATE_OneParam( FOnItemDestroyed, const TSharedPtr<FAssetViewItem>& /*AssetItem*/);
+
+class SAssetListItem;
+class SAssetTileItem;
+
+enum class EThumbnailSize : uint8
+{
+	Tiny = 0,
+	Small,
+	Medium,
+	Large,
+	Huge,
+
+	// Not a size
+	MAX
+};
+
+
+
+namespace FAssetViewModeUtils 
+{
+	FReply OnViewModeKeyDown( const TSet< TSharedPtr<FAssetViewItem> >& SelectedItems, const FKeyEvent& InKeyEvent );
+}
+
+struct FAssetViewItemHelper
+{
+public:
+	static TSharedRef<SWidget> CreateListItemContents(SAssetListItem* const InListItem, const TSharedRef<SWidget>& InThumbnail, FName& OutItemShadowBorder);
+	static TSharedRef<SWidget> CreateTileItemContents(SAssetTileItem* const InTileItem, const TSharedRef<SWidget>& InThumbnail, FName& OutItemShadowBorder);
+
+private:
+	template <typename T>
+	static TSharedRef<SWidget> CreateListTileItemContents(T* const InTileOrListItem, const TSharedRef<SWidget>& InThumbnail, FName& OutItemShadowBorder);
+};
+
+/** The tile view mode of the asset view */
+class SAssetTileView : public STileView<TSharedPtr<FAssetViewItem>>
+{
+public:
+	virtual bool SupportsKeyboardFocus() const override { return true; }
+	virtual FReply OnKeyDown( const FGeometry& InGeometry, const FKeyEvent& InKeyEvent ) override;
+	virtual void Tick( const FGeometry& AllottedGeometry, const double InCurrentTime, const float InDeltaTime ) override;
+};
+
+/** The list view mode of the asset view */
+class SAssetListView : public SListView<TSharedPtr<FAssetViewItem>>
+{
+public:
+	virtual bool SupportsKeyboardFocus() const override { return true; }
+	virtual FReply OnKeyDown( const FGeometry& InGeometry, const FKeyEvent& InKeyEvent ) override;
+	virtual void Tick(const FGeometry& AllottedGeometry, const double InCurrentTime, const float InDeltaTime) override;
+};
+
+/** The columns view mode of the asset view */
+class SAssetColumnView : public SListView<TSharedPtr<FAssetViewItem>>
+{
+public:
+	virtual bool SupportsKeyboardFocus() const override { return true; }
+	virtual FReply OnKeyDown( const FGeometry& InGeometry, const FKeyEvent& InKeyEvent ) override;
+	virtual void Tick(const FGeometry& AllottedGeometry, const double InCurrentTime, const float InDeltaTime) override;
+};
+
+/** A base class for all asset view items */
+class SAssetViewItem : public SCompoundWidget
+{
+	friend class SAssetViewItemToolTip;
+
+public:
+
+	SLATE_BEGIN_ARGS( SAssetViewItem )
+		: _ShouldAllowToolTip(true)
+		, _ThumbnailEditMode(false)
+		{}
+
+		/** Data for the asset this item represents */
+	SLATE_ARGUMENT(TSharedPtr<FAssetViewItem>, AssetItem)
+
+		/** Delegate for when an asset name has entered a rename state */
+		SLATE_EVENT(FOnRenameBegin, OnRenameBegin)
+
+		/** Delegate for when an asset name has been entered for an item that is in a rename state */
+		SLATE_EVENT(FOnRenameCommit, OnRenameCommit)
+
+		/** Delegate for when an asset name has been entered for an item to verify the name before commit */
+		SLATE_EVENT(FOnVerifyRenameCommit, OnVerifyRenameCommit)
+
+		/** Called when any asset item is destroyed. Used in thumbnail management */
+		SLATE_EVENT(FOnItemDestroyed, OnItemDestroyed)
+
+		/** If false, the tooltip will not be displayed */
+		SLATE_ATTRIBUTE(bool, ShouldAllowToolTip)
+
+		/** If true, display the thumbnail edit mode UI */
+		SLATE_ATTRIBUTE(bool, ThumbnailEditMode)
+
+		/** The string in the title to highlight (used when searching by string) */
+		SLATE_ATTRIBUTE(FText, HighlightText)
+
+		/** Delegate to call (if bound) to check if it is valid to get a custom tooltip for this view item */
+		SLATE_EVENT(FOnIsAssetValidForCustomToolTip, OnIsAssetValidForCustomToolTip)
+
+		/** Delegate to call (if bound) to get a custom tooltip for this view item */
+		SLATE_EVENT(FOnGetCustomAssetToolTip, OnGetCustomAssetToolTip)
+
+		/** Delegate for when an item is about to show a tool tip */
+		SLATE_EVENT(FOnVisualizeAssetToolTip, OnVisualizeAssetToolTip)
+
+		/** Delegate for when an item's tooltip is about to close */
+		SLATE_EVENT(FOnAssetToolTipClosing, OnAssetToolTipClosing)
+
+		/** Delegate for getting the selection state of this item */ 
+		SLATE_ARGUMENT(FIsSelected, IsSelected)
+
+	SLATE_END_ARGS()
+
+	/** Virtual destructor */
+	virtual ~SAssetViewItem();
+
+	/** Performs common initialization logic for all asset view items */
+	void Construct( const FArguments& InArgs );
+
+	/** NOTE: Any functions overridden from the base widget classes *must* also be overridden by SAssetColumnViewRow and forwarded on to its internal item */
+	virtual void Tick( const FGeometry& AllottedGeometry, const double InCurrentTime, const float InDeltaTime ) override;
+	virtual TSharedPtr<IToolTip> GetToolTip() override;
+	virtual void OnDragEnter( const FGeometry& MyGeometry, const FDragDropEvent& DragDropEvent ) override;
+	virtual void OnDragLeave( const FDragDropEvent& DragDropEvent ) override;
+	virtual FReply OnDragOver( const FGeometry& MyGeometry, const FDragDropEvent& DragDropEvent ) override;
+	virtual FReply OnDrop( const FGeometry& MyGeometry, const FDragDropEvent& DragDropEvent ) override;
+	virtual FReply OnDrop( const FGeometry& MyGeometry, const FDragDropEvent& DragDropEvent, const TSharedRef<SWidget>& InParentWidget);
+	virtual bool OnVisualizeTooltip( const TSharedPtr<SWidget>& TooltipContent ) override;
+	virtual void OnToolTipClosing() override;
+
+	/** Returns the color this item should be tinted with */
+	virtual FSlateColor GetAssetColor() const;
+
+	/** Get the border image to display */
+	virtual const FSlateBrush* GetBorderImage() const;
+
+	/** Get the name text to be displayed for this item */
+	FText GetNameText() const;
+
+protected:
+	/** Check to see if the name should be read-only */
+	bool IsNameReadOnly() const;
+
+	/** Handles starting a name change */
+	virtual void HandleBeginNameChange( const FText& OriginalText );
+
+	/** Handles committing a name change */
+	virtual void HandleNameCommitted( const FText& NewText, ETextCommit::Type CommitInfo );
+
+	/** Handles verifying a name change */
+	virtual bool HandleVerifyNameChanged( const FText& NewText, FText& OutErrorMessage );
+
+	/** Handles committing a name change */
+	virtual void OnAssetDataChanged();
+
+	/** Notification for when the dirty flag changes */
+	virtual void DirtyStateChanged();
+
+	/** Gets the name of the class of this asset */
+	FText GetAssetClassText() const;
+
+	/** Gets the brush for the dirty indicator image */
+	const FSlateBrush* GetDirtyImage() const;
+
+	/** Generates the source control icon widget */
+	TSharedRef<SWidget> GenerateSourceControlIconWidget();
+
+	/** Generate a widget to inject extra external state indicator on the asset. */
+	TSharedRef<SWidget> GenerateExtraStateIconWidget(TAttribute<float> InMaxExtraStateIconWidth) const;
+
+	/** Generate a widget to inject extra external state indicator on the asset tooltip. */
+	TSharedRef<SWidget> GenerateExtraStateTooltipWidget() const;
+
+	/** Gets the visibility for the thumbnail edit mode UI */
+	EVisibility GetThumbnailEditModeUIVisibility() const;
+
+	/** Creates a tooltip widget for this item */
+	TSharedRef<SWidget> CreateToolTipWidget() const;
+
+	/** Gets the visibility of the checked out by other text block in the tooltip */
+	EVisibility GetCheckedOutByOtherTextVisibility() const;
+
+	/** Gets the text for the checked out by other text block in the tooltip */
+	FText GetCheckedOutByOtherText() const;
+
+	/** Helper function for CreateToolTipWidget. Gets the user description for the asset, if it exists. */
+	FText GetAssetUserDescription() const;
+
+	/** Helper function for CreateToolTipWidget. Adds a key value pair to the info box of the tooltip */
+	void AddToToolTipInfoBox(const TSharedRef<SVerticalBox>& InfoBox, const FText& Key, const FText& Value, bool bImportant) const;
+
+	/** Updates the bPackageDirty flag */
+	void UpdateDirtyState();
+
+	/** Returns true if the item is dirty. */
+	bool IsDirty() const;
+
+	/** Update the source control state of this item if required */
+	void UpdateSourceControlState(float InDeltaTime);
+	
+	/** Cache the display tags for this item */
+	void CacheDisplayTags();
+
+	/** Whether this item is a folder */
+	bool IsFolder() const;
+
+	/** Delegate handler for when the source control provider changes */
+	void HandleSourceControlProviderChanged(class ISourceControlProvider& OldProvider, class ISourceControlProvider& NewProvider);
+
+	/** Delegate handler for when source control state changes */
+	void HandleSourceControlStateChanged();
+
+	/** Returns the width at which the name label will wrap the name */
+	virtual float GetNameTextWrapWidth() const { return 0.0f; }
+
+protected:
+	/** Data for a cached display tag for this item (used in the tooltip, and also as the display string in column views) */
+	struct FTagDisplayItem
+	{
+		FTagDisplayItem(FName InTagKey, FText InDisplayKey, FText InDisplayValue, const bool InImportant)
+			: TagKey(InTagKey)
+			, DisplayKey(MoveTemp(InDisplayKey))
+			, DisplayValue(MoveTemp(InDisplayValue))
+			, bImportant(InImportant)
+		{
+		}
+
+		FName TagKey;
+		FText DisplayKey;
+		FText DisplayValue;
+		bool bImportant;
+	};
+
+	TSharedPtr< SInlineEditableTextBlock > InlineRenameWidget;
+
+	TSharedPtr<STextBlock> ClassTextWidget;
+
+	/** The data for this item */
+	TSharedPtr<FAssetViewItem> AssetItem;
+
+	/** The cached display tags for this item */
+	TArray<FTagDisplayItem> CachedDisplayTags;
+
+	/** Delegate for when an asset name has entered a rename state */
+	FOnRenameBegin OnRenameBegin;
+
+	/** Delegate for when an asset name has been entered for an item that is in a rename state */
+	FOnRenameCommit OnRenameCommit;
+
+	/** Delegate for when an asset name has been entered for an item to verify the name before commit */
+	FOnVerifyRenameCommit OnVerifyRenameCommit;
+
+	/** Called when any asset item is destroyed. Used in thumbnail management */
+	FOnItemDestroyed OnItemDestroyed;
+
+	/** Called to test if it is valid to make a custom tool tip for that asset */
+	FOnIsAssetValidForCustomToolTip OnIsAssetValidForCustomToolTip;
+
+	/** Called if bound to get a custom asset item tooltip */
+	FOnGetCustomAssetToolTip OnGetCustomAssetToolTip;
+
+	/** Called if bound when about to show a tooltip */
+	FOnVisualizeAssetToolTip OnVisualizeAssetToolTip;
+
+	/** Called if bound when a tooltip is closing */
+	FOnAssetToolTipClosing OnAssetToolTipClosing;
+	
+	/** Delegate for getting the selection state of this item */ 
+	FIsSelected IsSelected;
+
+	FGeometry LastGeometry;
+
+	/** If false, the tooltip will not be displayed */
+	TAttribute<bool> ShouldAllowToolTip;
+
+	/** If true, display the thumbnail edit mode UI */
+	TAttribute<bool> ThumbnailEditMode;
+
+	/** The substring to be highlighted in the name and tooltip path */
+	TAttribute<FText> HighlightText;
+
+	/** Cached brushes for the dirty state */
+	const FSlateBrush* AssetDirtyBrush;
+
+	/** Cached flag describing if the item is dirty */
+	bool bItemDirty;
+
+	/** Flag indicating whether we have requested initial source control state */
+	bool bSourceControlStateRequested;
+
+	/** Delay timer before we request a source control state update, to prevent spam */
+	float SourceControlStateDelay;
+
+	/** True when a drag is over this item with a drag operation that we know how to handle. The operation itself may not be valid to drop. */
+	bool bDraggedOver;
+
+	/** Widget for the source control state */
+	TSharedPtr<SLayeredImage> SCCStateWidget;
+
+	/** Delegate handle for the HandleSourceControlStateChanged function callback */
+	FDelegateHandle SourceControlStateChangedDelegateHandle;
+};
+
+
+/** An item in the asset list view */
+class SAssetListItem : public SAssetViewItem
+{
+	friend struct FAssetViewItemHelper;
+
+public:
+	SLATE_BEGIN_ARGS( SAssetListItem )
+		: _ThumbnailPadding(0)
+		, _ThumbnailLabel( EThumbnailLabel::ClassName )
+		, _ThumbnailHintColorAndOpacity( FLinearColor( 0.0f, 0.0f, 0.0f, 0.0f ) )
+		, _ItemHeight(16)
+		, _ShouldAllowToolTip(true)
+		, _ThumbnailEditMode(false)
+		, _AllowThumbnailHintLabel(false)
+		{}
+
+		/** The handle to the thumbnail this item should render */
+		SLATE_ARGUMENT( TSharedPtr<FAssetThumbnail>, AssetThumbnail)
+
+		/** Data for the asset this item represents */
+		SLATE_ARGUMENT( TSharedPtr<FAssetViewItem>, AssetItem )
+
+		/** How much padding to allow around the thumbnail */
+		SLATE_ARGUMENT( float, ThumbnailPadding )
+
+		/** The contents of the label displayed on the thumbnail */
+		SLATE_ARGUMENT( EThumbnailLabel::Type, ThumbnailLabel )
+
+		/**  */
+		SLATE_ATTRIBUTE( FLinearColor, ThumbnailHintColorAndOpacity )
+
+		/** The height of the list item */
+		SLATE_ATTRIBUTE( float, ItemHeight )
+
+		/** Delegate for when an asset name has entered a rename state */
+		SLATE_EVENT( FOnRenameBegin, OnRenameBegin )
+
+		/** Delegate for when an asset name has been entered for an item that is in a rename state */
+		SLATE_EVENT( FOnRenameCommit, OnRenameCommit )
+
+		/** Delegate for when an asset name has been entered for an item to verify the name before commit */
+		SLATE_EVENT( FOnVerifyRenameCommit, OnVerifyRenameCommit )
+
+		/** Called when any asset item is destroyed. Used in thumbnail management */
+		SLATE_EVENT( FOnItemDestroyed, OnItemDestroyed )
+
+		/** If false, the tooltip will not be displayed */
+		SLATE_ATTRIBUTE( bool, ShouldAllowToolTip )
+
+		/** The string in the title to highlight (used when searching by string) */
+		SLATE_ATTRIBUTE( FText, HighlightText )
+
+		/** If true, the thumbnail in this item can be edited */
+		SLATE_ATTRIBUTE( bool, ThumbnailEditMode )
+
+		/** Whether the thumbnail should ever show it's hint label */
+		SLATE_ARGUMENT( bool, AllowThumbnailHintLabel )
+
+		/** Whether the item is selected in the view */
+		SLATE_ARGUMENT(FIsSelected, IsSelected)
+
+		/** Whether the item is selected in the view without anything else being selected*/
+		SLATE_ARGUMENT(FIsSelected, IsSelectedExclusively)
+
+		/** Delegate to call (if bound) to check if it is valid to get a custom tooltip for this view item */
+		SLATE_EVENT(FOnIsAssetValidForCustomToolTip, OnIsAssetValidForCustomToolTip)
+
+		/** Delegate to request a custom tool tip if necessary */
+		SLATE_EVENT(FOnGetCustomAssetToolTip, OnGetCustomAssetToolTip)
+
+		/* Delegate to signal when the item is about to show a tooltip */
+		SLATE_EVENT(FOnVisualizeAssetToolTip, OnVisualizeAssetToolTip)
+
+		/** Delegate for when an item's tooltip is about to close */
+		SLATE_EVENT( FOnAssetToolTipClosing, OnAssetToolTipClosing )
+
+	SLATE_END_ARGS()
+
+	/** Destructor */
+	~SAssetListItem();
+
+	/** Constructs this widget with InArgs */
+	void Construct( const FArguments& InArgs );
+
+	/** Handles committing a name change */
+	virtual void OnAssetDataChanged() override;
+
+	/** Handles realtime thumbnails */
+	virtual void OnMouseEnter(const FGeometry& MyGeometry, const FPointerEvent& MouseEvent) override;
+
+	/** Handles realtime thumbnails */
+	virtual void OnMouseLeave(const FPointerEvent& MouseEvent) override;
+
+	/** Whether the widget should allow primitive tools to be displayed */
+	bool CanDisplayPrimitiveTools() const { return false; }
+
+private:
+	/** Get the expected width of an extra state icon. */
+	float GetExtraStateIconWidth() const;
+
+	/** Returns the max width size to be used by extra state icons. */
+	FOptionalSize GetExtraStateIconMaxWidth() const;
+
+	/** Returns the size of the state icon box widget (i.e dirty image, scc)*/
+	FOptionalSize GetStateIconImageSize() const;
+
+	/** Returns the size of the thumbnail widget */
+	FOptionalSize GetThumbnailBoxSize() const;
+
+private:
+	/** The handle to the thumbnail that this item is rendering */
+	TSharedPtr<FAssetThumbnail> AssetThumbnail;
+
+	/** The height allowed for this item */
+	TAttribute<float> ItemHeight;
+};
+
+/** An item in the asset tile view */
+class SAssetTileItem : public SAssetViewItem
+{
+	friend struct FAssetViewItemHelper;
+
+public:
+	SLATE_BEGIN_ARGS( SAssetTileItem )
+		: _ThumbnailPadding(0)
+		, _ThumbnailLabel( EThumbnailLabel::ClassName )
+		, _ThumbnailHintColorAndOpacity( FLinearColor( 0.0f, 0.0f, 0.0f, 0.0f ) )
+		, _AllowThumbnailHintLabel(true)
+		, _CurrentThumbnailSize(EThumbnailSize::Medium)
+		, _ItemWidth(16)
+		, _ShouldAllowToolTip(true)
+		, _ThumbnailEditMode(false)
+	{}
+
+		/** The handle to the thumbnail this item should render */
+		SLATE_ARGUMENT( TSharedPtr<FAssetThumbnail>, AssetThumbnail)
+
+		/** Data for the asset this item represents */
+		SLATE_ARGUMENT( TSharedPtr<FAssetViewItem>, AssetItem )
+
+		/** How much padding to allow around the thumbnail */
+		SLATE_ARGUMENT( float, ThumbnailPadding )
+
+		/** The contents of the label displayed on the thumbnail */
+		SLATE_ARGUMENT( EThumbnailLabel::Type, ThumbnailLabel )
+
+		/**  */
+		SLATE_ATTRIBUTE( FLinearColor, ThumbnailHintColorAndOpacity )
+
+		/** Whether the thumbnail should ever show it's hint label */
+		SLATE_ARGUMENT( bool, AllowThumbnailHintLabel )
+
+		/** Current size of the thumbnail that was generated */
+		SLATE_ATTRIBUTE(EThumbnailSize, CurrentThumbnailSize)
+
+		/** The width of the item */
+		SLATE_ATTRIBUTE( float, ItemWidth )
+
+		/** Delegate for when an asset name has entered a rename state */
+		SLATE_EVENT( FOnRenameBegin, OnRenameBegin )
+
+		/** Delegate for when an asset name has been entered for an item that is in a rename state */
+		SLATE_EVENT( FOnRenameCommit, OnRenameCommit )
+
+		/** Delegate for when an asset name has been entered for an item to verify the name before commit */
+		SLATE_EVENT( FOnVerifyRenameCommit, OnVerifyRenameCommit )
+
+		/** Called when any asset item is destroyed. Used in thumbnail management */
+		SLATE_EVENT( FOnItemDestroyed, OnItemDestroyed )
+
+		/** If false, the tooltip will not be displayed */
+		SLATE_ATTRIBUTE( bool, ShouldAllowToolTip )
+
+		/** The string in the title to highlight (used when searching by string) */
+		SLATE_ATTRIBUTE( FText, HighlightText )
+
+		/** If true, the thumbnail in this item can be edited */
+		SLATE_ATTRIBUTE( bool, ThumbnailEditMode )
+
+		/** Whether the item is selected in the view */
+		SLATE_ARGUMENT(FIsSelected, IsSelected)
+
+		/** Whether the item is selected in the view without anything else being selected*/
+		SLATE_ARGUMENT(FIsSelected, IsSelectedExclusively)
+
+		/** Delegate to call (if bound) to check if it is valid to get a custom tooltip for this view item */
+		SLATE_EVENT(FOnIsAssetValidForCustomToolTip, OnIsAssetValidForCustomToolTip)
+
+		/** Delegate to request a custom tool tip if necessary */
+		SLATE_EVENT(FOnGetCustomAssetToolTip, OnGetCustomAssetToolTip)
+
+		/* Delegate to signal when the item is about to show a tooltip */
+		SLATE_EVENT(FOnVisualizeAssetToolTip, OnVisualizeAssetToolTip)
+
+		/** Delegate for when an item's tooltip is about to close */
+		SLATE_EVENT( FOnAssetToolTipClosing, OnAssetToolTipClosing )
+
+	SLATE_END_ARGS()
+
+	/** Destructor */
+	~SAssetTileItem();
+
+	/** Constructs this widget with InArgs */
+	void Construct( const FArguments& InArgs );
+
+	/** Handles committing a name change */
+	virtual void OnAssetDataChanged() override;
+
+	/** Handles realtime thumbnails */
+	virtual void OnMouseEnter(const FGeometry& MyGeometry, const FPointerEvent& MouseEvent) override;
+
+	/** Handles realtime thumbnails */
+	virtual void OnMouseLeave(const FPointerEvent& MouseEvent) override;
+
+	/** Whether the widget should allow primitive tools to be displayed */
+	bool CanDisplayPrimitiveTools() const { return true; }
+
+	/** Get the border image to display */
+	virtual const FSlateBrush* GetBorderImage() const;
+protected:
+	/** SAssetViewItem interface */
+	virtual float GetNameTextWrapWidth() const override { return LastGeometry.GetLocalSize().X - 15.f; }
+
+	/** Get the expected width of an extra state icon. */
+	float GetExtraStateIconWidth() const;
+
+	/** Returns the max width size to be used by extra state icons. */
+	FOptionalSize GetExtraStateIconMaxWidth() const;
+
+	/** Returns the size of the state icon box widget (i.e dirty image, scc)*/
+	FOptionalSize GetStateIconImageSize() const;
+
+	/** Returns the size of the thumbnail widget */
+	FOptionalSize GetThumbnailBoxSize() const;
+
+	/** Gets the visibility of the asset class label in thumbnails */
+	EVisibility GetAssetClassLabelVisibility() const;
+
+	/** Returns the font to use for the thumbnail label */
+	FSlateFontInfo GetThumbnailFont() const;
+
+	/** Gets the max size of the name area for an asset */
+	FOptionalSize GetMaxAssetNameHeight() const;
+
+	const FSlateBrush* GetFolderBackgroundImage() const;
+	const FSlateBrush* GetFolderBackgroundShadowImage() const;
+
+	FOptionalSize GetNameAreaMaxDesiredHeight() const;
+
+	int32 GetGenericThumbnailSize() const;
+
+	static void InitializeAssetNameHeights();
+
+private:
+	/** The handle to the thumbnail that this item is rendering */
+	TSharedPtr<FAssetThumbnail> AssetThumbnail;
+
+	/** The width of the item. Used to enforce a square thumbnail. */
+	TAttribute<float> ItemWidth;
+
+	/** Max name height for each thumbnail size */
+	static float AssetNameHeights[(int32)EThumbnailSize::MAX];
+
+	/** The padding allotted for the thumbnail */
+	float ThumbnailPadding;
+
+	/** Current thumbnail size when this widget was generated */
+	TAttribute<EThumbnailSize> CurrentThumbnailSize;
+};
+
+/** An item in the asset column view */
+class SAssetColumnItem : public SAssetViewItem
+{
+public:
+	SLATE_BEGIN_ARGS( SAssetColumnItem ) {}
+
+		/** Data for the asset this item represents */
+		SLATE_ARGUMENT( TSharedPtr<FAssetViewItem>, AssetItem )
+
+		/** Delegate for when an asset name has entered a rename state */
+		SLATE_EVENT( FOnRenameBegin, OnRenameBegin )
+
+		/** Delegate for when an asset name has been entered for an item that is in a rename state */
+		SLATE_EVENT( FOnRenameCommit, OnRenameCommit )
+
+		/** Delegate for when an asset name has been entered for an item to verify the name before commit */
+		SLATE_EVENT( FOnVerifyRenameCommit, OnVerifyRenameCommit )
+
+		/** Called when any asset item is destroyed. Used in thumbnail management, though it may be used for more so It is in column items for consistency. */
+		SLATE_EVENT( FOnItemDestroyed, OnItemDestroyed )
+
+		/** The string in the title to highlight (used when searching by string) */
+		SLATE_ATTRIBUTE( FText, HighlightText )
+
+		/** Delegate to call (if bound) to check if it is valid to get a custom tooltip for this view item */
+		SLATE_EVENT(FOnIsAssetValidForCustomToolTip, OnIsAssetValidForCustomToolTip)
+
+		/** Delegate to request a custom tool tip if necessary */
+		SLATE_EVENT( FOnGetCustomAssetToolTip, OnGetCustomAssetToolTip)
+
+		/* Delegate to signal when the item is about to show a tooltip */
+		SLATE_EVENT( FOnVisualizeAssetToolTip, OnVisualizeAssetToolTip)
+
+		/** Delegate for when an item's tooltip is about to close */
+		SLATE_EVENT( FOnAssetToolTipClosing, OnAssetToolTipClosing )
+
+	SLATE_END_ARGS()
+
+	/** Constructs this widget with InArgs */
+	void Construct( const FArguments& InArgs );
+
+	/** Generates a widget for a particular column */
+	virtual TSharedRef<SWidget> GenerateWidgetForColumn( const FName& ColumnName, FIsSelected InIsSelected );
+
+	/** Handles committing a name change */
+	virtual void OnAssetDataChanged() override;
+
+private:
+	/** Gets the path to this asset */
+	FText GetAssetPathText() const;
+
+	/** Gets the value for the specified tag in this asset */
+	FText GetAssetTagText(FName Tag) const;
+
+private:
+	TAttribute<FText> HighlightText;
+
+	TSharedPtr<STextBlock> ClassText;
+	TSharedPtr<STextBlock> PathText;
+};
+
+class SAssetColumnViewRow : public SMultiColumnTableRow< TSharedPtr<FAssetViewItem> >
+{
+public:
+	SLATE_BEGIN_ARGS( SAssetColumnViewRow )
+		{}
+
+		SLATE_EVENT( FOnDragDetected, OnDragDetected )
+		SLATE_ARGUMENT( TSharedPtr<SAssetColumnItem>, AssetColumnItem )
+
+	SLATE_END_ARGS()
+
+	/** Constructs this widget with InArgs */
+	void Construct( const FArguments& InArgs, const TSharedRef<STableViewBase>& InOwnerTableView )
+	{
+		this->AssetColumnItem = InArgs._AssetColumnItem;
+		ensure(this->AssetColumnItem.IsValid());
+
+		SMultiColumnTableRow< TSharedPtr<FAssetViewItem> >::Construct( 
+			FSuperRowType::FArguments()
+				.Style(FEditorStyle::Get(), "ContentBrowser.AssetListView.ColumnListTableRow")
+				.OnDragDetected(InArgs._OnDragDetected), 
+			InOwnerTableView);
+		Content = this->AssetColumnItem;
+	}
+
+	virtual TSharedRef<SWidget> GenerateWidgetForColumn( const FName& ColumnName ) override
+	{
+		if ( this->AssetColumnItem.IsValid() )
+		{
+			return this->AssetColumnItem->GenerateWidgetForColumn(ColumnName, FIsSelected::CreateSP( this, &SAssetColumnViewRow::IsSelectedExclusively ));
+		}
+		else
+		{
+			return SNew(STextBlock) .Text( NSLOCTEXT("AssetView", "InvalidColumnId", "Invalid Column Item") );
+		}
+		
+	}
+
+	virtual FVector2D GetRowSizeForColumn(const FName& InColumnName) const override
+	{
+		const TSharedRef<SWidget>* ColumnWidget = GetWidgetFromColumnId(InColumnName);
+
+		if (ColumnWidget != nullptr)
+		{
+			return (*ColumnWidget)->GetDesiredSize();
+		}
+
+		return FVector2D::ZeroVector; 
+	}
+
+	virtual void Tick( const FGeometry& AllottedGeometry, const double InCurrentTime, const float InDeltaTime ) override
+	{
+		this->AssetColumnItem->Tick(AllottedGeometry, InCurrentTime, InDeltaTime);
+	}
+
+	virtual TSharedPtr<IToolTip> GetToolTip() override
+	{
+		return AssetColumnItem->GetToolTip();
+	}
+
+	virtual void OnDragEnter(const FGeometry& MyGeometry, const FDragDropEvent& DragDropEvent) override
+	{
+		AssetColumnItem->OnDragEnter(MyGeometry, DragDropEvent);
+	}
+
+	virtual void OnDragLeave(const FDragDropEvent& DragDropEvent) override
+	{
+		AssetColumnItem->OnDragLeave(DragDropEvent);
+	}
+
+	virtual FReply OnDragOver(const FGeometry& MyGeometry, const FDragDropEvent& DragDropEvent) override
+	{
+		return AssetColumnItem->OnDragOver(MyGeometry, DragDropEvent);
+	}
+
+	virtual FReply OnDrop(const FGeometry& MyGeometry, const FDragDropEvent& DragDropEvent) override
+	{
+		return AssetColumnItem->OnDrop(MyGeometry, DragDropEvent, AsShared());
+	}
+
+	virtual bool OnVisualizeTooltip(const TSharedPtr<SWidget>& TooltipContent) override
+	{
+		// We take the content from the asset column item during construction,
+		// so let the item handle the tooltip callback
+		return AssetColumnItem->OnVisualizeTooltip(TooltipContent);
+	}
+
+	virtual void OnToolTipClosing() override
+	{
+		AssetColumnItem->OnToolTipClosing();
+	}
+
+	TSharedPtr<SAssetColumnItem> AssetColumnItem;
+};
diff --git a/Engine/Source/Editor/ContentBrowser/Public/FrontendFilters.h b/Engine/Source/Editor/ContentBrowser/Public/FrontendFilters.h
new file mode 100644
--- /dev/null
+++ b/Engine/Source/Editor/ContentBrowser/Public/FrontendFilters.h
@@ -0,0 +1,393 @@
+// Copyright Epic Games, Inc. All Rights Reserved.
+
+#pragma once
+
+#include "CoreMinimal.h"
+#include "CollectionManagerTypes.h"
+#include "ISourceControlProvider.h"
+#include "IAssetTools.h"
+#include "Misc/TextFilterExpressionEvaluator.h"
+#include "FrontendFilterBase.h"
+#include "Containers/Ticker.h"
+
+#define LOCTEXT_NAMESPACE "ContentBrowser"
+
+class FMenuBuilder;
+struct FAssetCompileData;
+
+/** A filter for text search */
+class CONTENTBROWSER_API FFrontendFilter_Text : public FFrontendFilter
+{
+public:
+	FFrontendFilter_Text();
+	~FFrontendFilter_Text();
+
+	// FFrontendFilter implementation
+	virtual FString GetName() const override { return TEXT("TextFilter"); }
+	virtual FText GetDisplayName() const override { return LOCTEXT("FrontendFilter_Text", "Text"); }
+	virtual FText GetToolTipText() const override { return LOCTEXT("FrontendFilter_TextTooltip", "Show only assets that match the input text"); }
+
+	// IFilter implementation
+	virtual bool PassesFilter(FAssetFilterType InItem) const override;
+
+public:
+	/** Returns the unsanitized and unsplit filter terms */
+	FText GetRawFilterText() const;
+
+	/** Set the Text to be used as the Filter's restrictions */
+	void SetRawFilterText(const FText& InFilterText);
+
+	/** Get the last error returned from lexing or compiling the current filter text */
+	FText GetFilterErrorText() const;
+
+	/** If bIncludeClassName is true, the text filter will include an asset's class name in the search */
+	void SetIncludeClassName(const bool InIncludeClassName);
+
+	/** If bIncludeAssetPath is true, the text filter will match against full Asset path */
+	void SetIncludeAssetPath(const bool InIncludeAssetPath);
+
+	bool GetIncludeAssetPath() const;
+
+	/** If bIncludeCollectionNames is true, the text filter will match against collection names as well */
+	void SetIncludeCollectionNames(const bool InIncludeCollectionNames);
+
+	bool GetIncludeCollectionNames() const;
+private:
+	/** Handles an on collection created event */
+	void HandleCollectionCreated(const FCollectionNameType& Collection);
+
+	/** Handles an on collection destroyed event */
+	void HandleCollectionDestroyed(const FCollectionNameType& Collection);
+
+	/** Handles an on collection renamed event */
+	void HandleCollectionRenamed(const FCollectionNameType& OriginalCollection, const FCollectionNameType& NewCollection);
+
+	/** Handles an on collection updated event */
+	void HandleCollectionUpdated(const FCollectionNameType& Collection);
+
+	/** Rebuild the array of dynamic collections that are being referenced by the current query */
+	void RebuildReferencedDynamicCollections();
+
+	/** An array of dynamic collections that are being referenced by the current query. These should be tested against each asset when it's looking for collections that contain it */
+	TArray<FCollectionNameType> ReferencedDynamicCollections;
+
+	/** Transient context data, used when calling PassesFilter. Kept around to minimize re-allocations between multiple calls to PassesFilter */
+	TSharedRef<class FFrontendFilter_TextFilterExpressionContext> TextFilterExpressionContext;
+
+	/** Expression evaluator that can be used to perform complex text filter queries */
+	FTextFilterExpressionEvaluator TextFilterExpressionEvaluator;
+
+	/** Delegate handles */
+	FDelegateHandle OnCollectionCreatedHandle;
+	FDelegateHandle OnCollectionDestroyedHandle;
+	FDelegateHandle OnCollectionRenamedHandle;
+	FDelegateHandle OnCollectionUpdatedHandle;
+};
+
+/** A filter that displays only checked out assets */
+class CONTENTBROWSER_API FFrontendFilter_CheckedOut : public FFrontendFilter, public TSharedFromThis<FFrontendFilter_CheckedOut>
+{
+public:
+	FFrontendFilter_CheckedOut(TSharedPtr<FFrontendFilterCategory> InCategory);
+
+	// FFrontendFilter implementation
+	virtual FString GetName() const override { return TEXT("CheckedOut"); }
+	virtual FText GetDisplayName() const override { return LOCTEXT("FrontendFilter_CheckedOut", "Checked Out"); }
+	virtual FText GetToolTipText() const override { return LOCTEXT("FrontendFilter_CheckedOutTooltip", "Show only assets that you have checked out or pending for add."); }
+	virtual void ActiveStateChanged(bool bActive) override;
+	virtual void SetCurrentFilter(TArrayView<const FName> InSourcePaths, const FContentBrowserDataFilter& InBaseFilter) override;
+
+	// IFilter implementation
+	virtual bool PassesFilter(FAssetFilterType InItem) const override;
+
+private:
+	
+	/** Request the source control status for this filter */
+	void RequestStatus();
+
+	/** Callback when source control operation has completed */
+	void SourceControlOperationComplete(const FSourceControlOperationRef& InOperation, ECommandResult::Type InResult);
+
+	bool bSourceControlEnabled;
+};
+
+/** A filter that displays assets not tracked by source control */
+class CONTENTBROWSER_API FFrontendFilter_NotSourceControlled : public FFrontendFilter, public TSharedFromThis<FFrontendFilter_NotSourceControlled>
+{
+public:
+	FFrontendFilter_NotSourceControlled(TSharedPtr<FFrontendFilterCategory> InCategory);
+
+	// FFrontendFilter implementation
+	virtual FString GetName() const override { return TEXT("NotSourceControlled"); }
+	virtual FText GetDisplayName() const override { return LOCTEXT("FrontendFilter_NotSourceControlled", "Not Source Controlled"); }
+	virtual FText GetToolTipText() const override { return LOCTEXT("FrontendFilter_NotSourceControlledTooltip", "Show only assets that are not tracked by source control."); }
+	virtual void ActiveStateChanged(bool bActive) override;
+	virtual void SetCurrentFilter(TArrayView<const FName> InSourcePaths, const FContentBrowserDataFilter& InBaseFilter) override;
+
+	// IFilter implementation
+	virtual bool PassesFilter(FAssetFilterType InItem) const override;
+
+private:
+
+	/** Request the source control status for this filter */
+	void RequestStatus();
+
+	/** Callback when source control operation has completed */
+	void SourceControlOperationComplete(const FSourceControlOperationRef& InOperation, ECommandResult::Type InResult);
+
+	bool bSourceControlEnabled;
+	bool bIsRequestStatusRunning;
+	bool bInitialRequestCompleted;
+};
+
+/** A filter that displays only modified assets */
+class CONTENTBROWSER_API FFrontendFilter_Modified : public FFrontendFilter
+{
+public:
+	FFrontendFilter_Modified(TSharedPtr<FFrontendFilterCategory> InCategory);
+	~FFrontendFilter_Modified();
+
+	// FFrontendFilter implementation
+	virtual FString GetName() const override { return TEXT("Modified"); }
+	virtual FText GetDisplayName() const override { return LOCTEXT("FrontendFilter_Modified", "Modified"); }
+	virtual FText GetToolTipText() const override { return LOCTEXT("FrontendFilter_ModifiedTooltip", "Show only assets that have been modified and not yet saved."); }
+	virtual void ActiveStateChanged(bool bActive) override;
+
+	// IFilter implementation
+	virtual bool PassesFilter(FAssetFilterType InItem) const override;
+
+private:
+
+	/** Handler for when a package's dirty state has changed */
+	void OnPackageDirtyStateUpdated(UPackage* Package);
+
+	bool bIsCurrentlyActive;
+};
+
+/** A filter that displays blueprints that have replicated properties */
+class CONTENTBROWSER_API FFrontendFilter_ReplicatedBlueprint : public FFrontendFilter
+{
+public:
+	FFrontendFilter_ReplicatedBlueprint(TSharedPtr<FFrontendFilterCategory> InCategory) : FFrontendFilter(InCategory) {}
+
+	// FFrontendFilter implementation
+	virtual FString GetName() const override { return TEXT("ReplicatedBlueprint"); }
+	virtual FText GetDisplayName() const override { return LOCTEXT("FFrontendFilter_ReplicatedBlueprint", "Replicated Blueprints"); }
+	virtual FText GetToolTipText() const override { return LOCTEXT("FFrontendFilter_ReplicatedBlueprintToolTip", "Show only blueprints with replicated properties."); }
+
+	// IFilter implementation
+	virtual bool PassesFilter(FAssetFilterType InItem) const override;
+};
+
+/** A filter that compares the value of an asset registry tag to a target value */
+class CONTENTBROWSER_API FFrontendFilter_ArbitraryComparisonOperation : public FFrontendFilter
+{
+public:
+	FFrontendFilter_ArbitraryComparisonOperation(TSharedPtr<FFrontendFilterCategory> InCategory);
+
+	// FFrontendFilter implementation
+	virtual FString GetName() const override;
+	virtual FText GetDisplayName() const override;
+	virtual FText GetToolTipText() const override;
+	virtual void ModifyContextMenu(FMenuBuilder& MenuBuilder) override;
+	virtual void SaveSettings(const FString& IniFilename, const FString& IniSection, const FString& SettingsString) const override;
+	virtual void LoadSettings(const FString& IniFilename, const FString& IniSection, const FString& SettingsString) override;
+
+	// IFilter implementation
+	virtual bool PassesFilter(FAssetFilterType InItem) const override;
+
+protected:
+	static FString ConvertOperationToString(ETextFilterComparisonOperation Op);
+	
+	void SetComparisonOperation(ETextFilterComparisonOperation NewOp);
+	bool IsComparisonOperationEqualTo(ETextFilterComparisonOperation TestOp) const;
+
+	FText GetKeyValueAsText() const;
+	FText GetTargetValueAsText() const;
+	void OnKeyValueTextCommitted(const FText& InText, ETextCommit::Type InCommitType);
+	void OnTargetValueTextCommitted(const FText& InText, ETextCommit::Type InCommitType);
+
+public:
+	FName TagName;
+	FString TargetTagValue;
+	ETextFilterComparisonOperation ComparisonOp;
+};
+
+/** An inverse filter that allows display of content in developer folders that are not the current user's */
+class CONTENTBROWSER_API FFrontendFilter_ShowOtherDevelopers : public FFrontendFilter
+{
+public:
+	/** Constructor */
+	FFrontendFilter_ShowOtherDevelopers(TSharedPtr<FFrontendFilterCategory> InCategory);
+
+	// FFrontendFilter implementation
+	virtual FString GetName() const override { return TEXT("ShowOtherDevelopers"); }
+	virtual FText GetDisplayName() const override { return LOCTEXT("FrontendFilter_ShowOtherDevelopers", "Other Developers"); }
+	virtual FText GetToolTipText() const override { return LOCTEXT("FrontendFilter_ShowOtherDevelopersTooltip", "Allow display of assets in developer folders that aren't yours."); }
+	virtual bool IsInverseFilter() const override { return true; }
+	virtual void SetCurrentFilter(TArrayView<const FName> InSourcePaths, const FContentBrowserDataFilter& InBaseFilter) override;
+
+	// IFilter implementation
+	virtual bool PassesFilter(FAssetFilterType InItem) const override;
+
+public:
+	/** Sets if we should filter out assets from other developers */
+	void SetShowOtherDeveloperAssets(bool bValue);
+
+	/** Gets if we should filter out assets from other developers */
+	bool GetShowOtherDeveloperAssets() const;
+
+private:
+	FString BaseDeveloperPath;
+	TArray<ANSICHAR> BaseDeveloperPathAnsi;
+	FString UserDeveloperPath;
+	bool bIsOnlyOneDeveloperPathSelected;
+	bool bShowOtherDeveloperAssets;
+};
+
+/** An inverse filter that allows display of object redirectors */
+class CONTENTBROWSER_API FFrontendFilter_ShowRedirectors : public FFrontendFilter
+{
+public:
+	/** Constructor */
+	FFrontendFilter_ShowRedirectors(TSharedPtr<FFrontendFilterCategory> InCategory);
+
+	// FFrontendFilter implementation
+	virtual FString GetName() const override { return TEXT("ShowRedirectors"); }
+	virtual FText GetDisplayName() const override { return LOCTEXT("FrontendFilter_ShowRedirectors", "Show Redirectors"); }
+	virtual FText GetToolTipText() const override { return LOCTEXT("FrontendFilter_ShowRedirectorsToolTip", "Allow display of Redirectors."); }
+	virtual bool IsInverseFilter() const override { return true; }
+	virtual void SetCurrentFilter(TArrayView<const FName> InSourcePaths, const FContentBrowserDataFilter& InBaseFilter) override;
+
+	// IFilter implementation
+	virtual bool PassesFilter(FAssetFilterType InItem) const override;
+
+private:
+	bool bAreRedirectorsInBaseFilter;
+	FName RedirectorClassName;
+};
+
+/** A filter that only displays assets used by loaded levels */
+class CONTENTBROWSER_API FFrontendFilter_InUseByLoadedLevels : public FFrontendFilter
+{
+public:
+	/** Constructor/Destructor */
+	FFrontendFilter_InUseByLoadedLevels(TSharedPtr<FFrontendFilterCategory> InCategory);
+	~FFrontendFilter_InUseByLoadedLevels() override;
+
+	// FFrontendFilter implementation
+	virtual FString GetName() const override { return TEXT("InUseByLoadedLevels"); }
+	virtual FText GetDisplayName() const override { return LOCTEXT("FrontendFilter_InUseByLoadedLevels", "In Use By Level"); }
+	virtual FText GetToolTipText() const override { return LOCTEXT("FrontendFilter_InUseByLoadedLevelsToolTip", "Show only assets that are currently in use by any loaded level."); }
+	virtual void ActiveStateChanged(bool bActive) override;
+
+	// IFilter implementation
+	virtual bool PassesFilter(FAssetFilterType InItem) const override;
+
+	/** Handler for when maps change in the editor */
+	void OnEditorMapChange( uint32 MapChangeFlags );
+
+	/** Handler for when an asset is renamed */
+	void OnAssetPostRename(const TArray<FAssetRenameData>& AssetsAndNames);
+
+	/** Handler for when assets are finished compiling */
+	void OnAssetPostCompile(const TArray<FAssetCompileData>& CompiledAssets);
+
+private:
+	void Refresh();
+	void RegisterDelayedRefresh(float DelayInSeconds);
+	void UnregisterDelayedRefresh();
+	FTSTicker::FDelegateHandle DelayedRefreshHandle;
+	bool bIsDirty = false;
+	bool bIsCurrentlyActive = false;
+};
+
+/** A filter that only displays assets used by any level */
+class CONTENTBROWSER_API FFrontendFilter_UsedInAnyLevel: public FFrontendFilter
+{
+public:
+	/** Constructor/Destructor */
+	FFrontendFilter_UsedInAnyLevel(TSharedPtr<FFrontendFilterCategory> InCategory);
+	virtual ~FFrontendFilter_UsedInAnyLevel();
+
+	// FFrontendFilter implementation
+	virtual FString GetName() const override { return TEXT("UsedInAnyLevel"); }
+	virtual FText GetDisplayName() const override { return LOCTEXT("FFrontendFilter_UsedInAnyLevel", "Used In Any Level"); }
+	virtual FText GetToolTipText() const override { return LOCTEXT("FFrontendFilter_UsedInAnyLevelTooltip", "Show only assets that are used in any level."); }
+	virtual void ActiveStateChanged(bool bActive) override;
+
+	// IFilter implementation
+	virtual bool PassesFilter(FAssetFilterType InItem) const override;
+
+private:
+	class IAssetRegistry* AssetRegistry;
+	TSet<FName> LevelsDependencies;
+};
+
+/** A filter that only displays assets used by any level */
+class CONTENTBROWSER_API FFrontendFilter_NotUsedInAnyLevel : public FFrontendFilter
+{
+public:
+	/** Constructor/Destructor */
+	FFrontendFilter_NotUsedInAnyLevel(TSharedPtr<FFrontendFilterCategory> InCategory);
+	virtual ~FFrontendFilter_NotUsedInAnyLevel();
+
+	// FFrontendFilter implementation
+	virtual FString GetName() const override { return TEXT("NotUsedInAnyLevel"); }
+	virtual FText GetDisplayName() const override { return LOCTEXT("FFrontendFilter_NotUsedInAnyLevel", "Not Used In Any Level"); }
+	virtual FText GetToolTipText() const override { return LOCTEXT("FFrontendFilter_NotUsedInAnyLevelTooltip", "Show only assets that are not used in any level."); }
+	virtual void ActiveStateChanged(bool bActive) override;
+
+	// IFilter implementation
+	virtual bool PassesFilter(FAssetFilterType InItem) const override;
+
+private:
+	class IAssetRegistry* AssetRegistry;
+	TSet<FName> LevelsDependencies;
+};
+
+/** A filter that displays recently opened assets */
+class CONTENTBROWSER_API FFrontendFilter_Recent : public FFrontendFilter
+{
+public:
+	FFrontendFilter_Recent(TSharedPtr<FFrontendFilterCategory> InCategory);
+	~FFrontendFilter_Recent();
+
+	// FFrontendFilter implementation
+	virtual FString GetName() const override { return TEXT("RecentlyOpened"); }
+	virtual FText GetDisplayName() const override { return LOCTEXT("FrontendFilter_Recent", "Recently Opened"); }
+	virtual FText GetToolTipText() const override { return LOCTEXT("FrontendFilter_RecentTooltip", "Show only recently opened assets."); }
+	virtual void ActiveStateChanged(bool bActive) override;
+	virtual void SetCurrentFilter(TArrayView<const FName> InSourcePaths, const FContentBrowserDataFilter& InBaseFilter) override;
+
+	// IFilter implementation
+	virtual bool PassesFilter(FAssetFilterType InItem) const override;
+
+	void ResetFilter(FName InName);
+
+private:
+	void RefreshRecentPackagePaths();
+
+	TSet<FName> RecentPackagePaths;
+	bool bIsCurrentlyActive;
+};
+
+/** A filter that displays only assets that are not read only */
+class CONTENTBROWSER_API FFrontendFilter_Writable : public FFrontendFilter
+{
+public:
+	FFrontendFilter_Writable(TSharedPtr<FFrontendFilterCategory> InCategory);
+	~FFrontendFilter_Writable();
+
+	// FFrontendFilter implementation
+	virtual FString GetName() const override { return TEXT("Writable"); }
+	virtual FText GetDisplayName() const override { return LOCTEXT("FrontendFilter_Writable", "Writable"); }
+	virtual FText GetToolTipText() const override { return LOCTEXT("FrontendFilter_WritableTooltip", "Show only assets that are not read only."); }
+
+	// IFilter implementation
+	virtual bool PassesFilter(FAssetFilterType InItem) const override;
+
+private:
+};
+
+#undef LOCTEXT_NAMESPACE
diff --git a/Engine/Source/Editor/ContentBrowser/Public/HistoryManager.h b/Engine/Source/Editor/ContentBrowser/Public/HistoryManager.h
new file mode 100644
--- /dev/null
+++ b/Engine/Source/Editor/ContentBrowser/Public/HistoryManager.h
@@ -0,0 +1,140 @@
+// Copyright Epic Games, Inc. All Rights Reserved.
+
+#pragma once
+
+#include "CoreMinimal.h"
+#include "SourcesData.h"
+
+class FMenuBuilder;
+
+struct FSelectionData
+{
+	/** Virtual paths (for both folders and files) of the selected items */
+	TSet<FName> SelectedVirtualPaths;
+
+	int32 Num() const
+	{
+		return SelectedVirtualPaths.Num();
+	}
+
+	void Reset()
+	{
+		SelectedVirtualPaths.Reset();
+	}
+
+	void Empty()
+	{
+		SelectedVirtualPaths.Empty();
+	}
+};
+
+/** The history data object, storing all important history data */
+struct FHistoryData
+{
+	/** History Description */
+	FText HistoryDesc;
+
+	/** The base set of filters on the asset view which includes selected paths and collections */
+	FSourcesData SourcesData;
+
+	/** The selection data from before the sources changed */
+	FSelectionData SelectionData;
+};
+
+/** The delegate for when history data should be applied */
+DECLARE_DELEGATE_OneParam(FOnApplyHistoryData, const FHistoryData& /*Data*/);
+
+/** The delegate for when history data should be updated */
+DECLARE_DELEGATE_OneParam(FOnUpdateHistoryData, FHistoryData& /*Data*/);
+
+/** The class responsible for managing all content browser history */
+class FHistoryManager
+{
+public:
+	/** Constructor */
+	FHistoryManager();
+
+	/** Set the delegate for applying history data */
+	void SetOnApplyHistoryData(const FOnApplyHistoryData& InOnApplyHistoryData);
+
+	/** Set the delegate for updating history data */
+	void SetOnUpdateHistoryData(const FOnUpdateHistoryData& InOnUpdateHistoryData);
+
+	/** Goes back one history snapshot and returns the history data at that snapshot */
+	bool GoBack();
+
+	/** Goes forward one history snapshot and returns the history data at that snapshot */
+	bool GoForward();
+
+	/** Stores new history data.  Called when creating a history snapshot */
+	void AddHistoryData();
+
+	/** Triggers an update for the current history data. This is typically done right before changing the history. */
+	void UpdateHistoryData();
+
+	/** Determines if a user can go forward in history */
+	bool CanGoForward() const;
+
+	/** Determines if a user can go back in history */
+	bool CanGoBack() const;
+
+	/** Gets the description of the previous history entry */
+	FText GetBackDesc() const;
+
+	/** Gets the description of the next history entry */
+	FText GetForwardDesc() const;
+
+	/**
+	 * Populates a list of menu items that can be added to a context menu
+	 * to allow a user to jump to different history snapshots instead of using the back and forward buttons
+	 *
+	 * @param bGetPrior If true gets history snapshots prior to the current history index(for navigating back). If false get history snapshots after the current history index (for navigating forward).
+	 * @param MenuBuilder The menubuilder to populate with menu items
+	 */
+	void GetAvailableHistoryMenuItems(bool bGetPrior, FMenuBuilder& MenuBuilder);
+
+	/** Removes all history data as determined by the passed in predicate */
+	template <class PREDICATE_CLASS>
+	void RemoveHistoryData(const PREDICATE_CLASS& Predicate)
+	{
+		for (int32 HistoryIndex = 0; HistoryIndex < HistoryData.Num(); HistoryIndex++)
+		{
+			if (Predicate(HistoryData[HistoryIndex]))
+			{
+				HistoryData.RemoveAt(HistoryIndex);
+				if (CurrentHistoryIndex >= HistoryIndex)
+				{
+					// Ensure if possible that current history index continues to point to the same item if something is removed before it.
+					CurrentHistoryIndex = FMath::Max(0, CurrentHistoryIndex - 1);
+				}
+				HistoryIndex--;
+			}
+		}
+	}
+
+private:
+	/** Notifies the owner to update to the state described by the current history data */
+	void ApplyCurrentHistoryData();
+
+	/** Notifies the owner to update the current history data */
+	void UpdateCurrentHistoryData();
+
+	/** Handler for when a history item is chosen in the GetAvailableHistroyMenuItems list */
+	void ExecuteJumpToHistory(int32 HistoryIndex);
+
+private:
+	/** The delegate for when history data should be applied */
+	FOnApplyHistoryData OnApplyHistoryData;
+
+	/** The delegate for when history data should be updated */
+	FOnUpdateHistoryData OnUpdateHistoryData;
+
+	/** A list of history snapshots */
+	TArray<FHistoryData> HistoryData;
+
+	/** The current history index the user is at (changes when the user goes back,forward, or history snapshots are taken) */
+	int32 CurrentHistoryIndex;
+
+	/** Max number of history items that can be stored.  Once the max is reached, the oldest history item is removed */
+	int32 MaxHistoryEntries;
+};
diff --git a/Engine/Source/Editor/ContentBrowser/Public/SAssetView.h b/Engine/Source/Editor/ContentBrowser/Public/SAssetView.h
new file mode 100644
--- /dev/null
+++ b/Engine/Source/Editor/ContentBrowser/Public/SAssetView.h
@@ -0,0 +1,1109 @@
+// Copyright Epic Games, Inc. All Rights Reserved.
+
+#pragma once
+
+#include "CoreMinimal.h"
+#include "Misc/Attribute.h"
+#include "Input/Reply.h"
+#include "Layout/Visibility.h"
+#include "Widgets/DeclarativeSyntaxSupport.h"
+#include "Styling/SlateColor.h"
+#include "Widgets/SWidget.h"
+#include "Widgets/SCompoundWidget.h"
+#include "AssetData.h"
+#include "ARFilter.h"
+#include "AssetThumbnail.h"
+#include "IContentBrowserSingleton.h"
+#include "ContentBrowserDataFilter.h"
+#include "SourcesData.h"
+#include "Animation/CurveSequence.h"
+#include "Widgets/Views/STableViewBase.h"
+#include "Widgets/Views/STableRow.h"
+#include "Editor/ContentBrowser/Private/AssetViewSortManager.h"
+#include "AssetViewTypes.h"
+#include "HistoryManager.h"
+#include "Misc/BlacklistNames.h"
+
+class FMenuBuilder;
+class FBlacklistPaths;
+class FWeakWidgetPath;
+class FWidgetPath;
+class SAssetColumnView;
+class SAssetListView;
+class SAssetTileView;
+class SComboButton;
+class UFactory;
+struct FPropertyChangedEvent;
+struct FToolMenuContext;
+class SContentBrowser;
+
+/** Fires whenever the asset view is asked to start to create a temporary item */
+DECLARE_DELEGATE_OneParam(FOnAssetViewNewItemRequested, const FContentBrowserItem& /*NewItem*/);
+
+/** Fires whenever one of the "Search" options changes, useful for modifying search criteria to match */
+DECLARE_DELEGATE(FOnSearchOptionChanged);
+
+/** Fires whenever asset view options menu is being opened, gives chance for external code to set additional context */
+DECLARE_DELEGATE_OneParam(FOnExtendAssetViewOptionsMenuContext, FToolMenuContext&);
+
+enum class EThumbnailSize : uint8;
+
+/**
+ * A widget to display a list of filtered assets
+ */
+class CONTENTBROWSER_API SAssetView : public SCompoundWidget
+{
+public:
+	SLATE_BEGIN_ARGS( SAssetView )
+		: _InitialCategoryFilter(EContentBrowserItemCategoryFilter::IncludeAll)
+		, _ThumbnailLabel( EThumbnailLabel::ClassName )
+		, _AllowThumbnailHintLabel(true)
+		, _InitialViewType(EAssetViewType::Tile)
+		, _ShowBottomToolbar(true)
+		, _ShowViewOptions(true)
+		, _AllowThumbnailEditMode(false)
+		, _CanShowClasses(true)
+		, _CanShowFolders(false)
+		, _CanShowReadOnlyFolders(true)
+		, _FilterRecursivelyWithBackendFilter(true)
+		, _CanShowRealTimeThumbnails(false)
+		, _CanShowDevelopersFolder(false)
+		, _CanShowFavorites(false)
+		, _CanDockCollections(false)
+		, _PreloadAssetsForContextMenu(true)
+		, _SelectionMode( ESelectionMode::Multi )
+		, _AllowDragging(true)
+		, _AllowFocusOnSync(true)
+		, _FillEmptySpaceInTileView(true)
+		, _ShowPathInColumnView(false)
+		, _ShowTypeInColumnView(true)
+		, _SortByPathInColumnView(false)
+		, _ForceShowEngineContent(false)
+		, _ForceShowPluginContent(false)
+		{}
+
+		/** Called to check if an asset should be filtered out by external code */
+		SLATE_EVENT( FOnShouldFilterAsset, OnShouldFilterAsset )
+
+		/** Called when the asset view is asked to start to create a temporary item */
+		SLATE_EVENT( FOnAssetViewNewItemRequested, OnNewItemRequested )
+
+		/** Called to when an item is selected */
+		SLATE_EVENT( FOnContentBrowserItemSelectionChanged, OnItemSelectionChanged )
+
+		/** Called when the user double clicks, presses enter, or presses space on an Content Browser item */
+		SLATE_EVENT( FOnContentBrowserItemsActivated, OnItemsActivated )
+
+		/** Delegate to invoke when a context menu for an item is opening. */
+		SLATE_EVENT( FOnGetContentBrowserItemContextMenu, OnGetItemContextMenu )
+
+		/** Called when the user has committed a rename of one or more items */
+		SLATE_EVENT( FOnContentBrowserItemRenameCommitted, OnItemRenameCommitted )
+
+		/** Delegate to call (if bound) to check if it is valid to get a custom tooltip for this asset item */
+		SLATE_EVENT(FOnIsAssetValidForCustomToolTip, OnIsAssetValidForCustomToolTip)
+
+		/** Called to get a custom asset item tool tip (if necessary) */
+		SLATE_EVENT( FOnGetCustomAssetToolTip, OnGetCustomAssetToolTip )
+
+		/** Called when an asset item is about to show a tooltip */
+		SLATE_EVENT( FOnVisualizeAssetToolTip, OnVisualizeAssetToolTip )
+
+		/** Called when an asset item's tooltip is closing */
+		SLATE_EVENT(FOnAssetToolTipClosing, OnAssetToolTipClosing)
+
+		/** Called when opening view options menu */
+		SLATE_EVENT(FOnExtendAssetViewOptionsMenuContext, OnExtendAssetViewOptionsMenuContext)
+
+		/** Initial set of item categories that this view should show - may be adjusted further by things like CanShowClasses or legacy delegate bindings */
+		SLATE_ARGUMENT( EContentBrowserItemCategoryFilter, InitialCategoryFilter )
+
+		/** The warning text to display when there are no assets to show */
+		SLATE_ATTRIBUTE( FText, AssetShowWarningText )
+
+		/** Attribute to determine what text should be highlighted */
+		SLATE_ATTRIBUTE( FText, HighlightedText )
+
+		/** What the label on the asset thumbnails should be */
+		SLATE_ARGUMENT( EThumbnailLabel::Type, ThumbnailLabel )
+
+		/** Whether to ever show the hint label on thumbnails */
+		SLATE_ARGUMENT( bool, AllowThumbnailHintLabel )
+
+		/** The filter collection used to further filter down assets returned from the backend */
+		SLATE_ARGUMENT( TSharedPtr<FAssetFilterCollectionType>, FrontendFilters )
+
+		/** Show path view filters submenu in view options menu */
+		SLATE_ARGUMENT( bool, bShowPathViewFilters )
+
+		/** The initial base sources filter */
+		SLATE_ARGUMENT( FSourcesData, InitialSourcesData )
+
+		/** The initial backend filter */
+		SLATE_ARGUMENT( FARFilter, InitialBackendFilter )
+
+		/** The asset that should be initially selected */
+		SLATE_ARGUMENT( FAssetData, InitialAssetSelection )
+
+		/** The initial view type */
+		SLATE_ARGUMENT( EAssetViewType::Type, InitialViewType )
+		/** Should the toolbar indicating number of selected assets, mode switch buttons, etc... be shown? */
+		SLATE_ARGUMENT( bool, ShowBottomToolbar )
+
+		/** Should view options be shown. Note: If ShowBottomToolbar is false, then view options are not shown regardless of this setting */
+		SLATE_ARGUMENT( bool, ShowViewOptions)
+
+		/** True if the asset view may edit thumbnails */
+		SLATE_ARGUMENT( bool, AllowThumbnailEditMode )
+
+		/** Indicates if this view is allowed to show classes */
+		SLATE_ARGUMENT( bool, CanShowClasses )
+
+		/** Indicates if the 'Show Folders' option should be enabled or disabled */
+		SLATE_ARGUMENT( bool, CanShowFolders )
+
+		/** Indicates if this view is allowed to show folders that cannot be written to */
+		SLATE_ARGUMENT( bool, CanShowReadOnlyFolders )
+
+		/** If true, recursive filtering will be caused by applying a backend filter */
+		SLATE_ARGUMENT( bool, FilterRecursivelyWithBackendFilter )
+
+		/** Indicates if the 'Real-Time Thumbnails' option should be enabled or disabled */
+		SLATE_ARGUMENT( bool, CanShowRealTimeThumbnails )
+
+		/** Indicates if the 'Show Developers' option should be enabled or disabled */
+		SLATE_ARGUMENT( bool, CanShowDevelopersFolder )
+
+		/** Indicates if the 'Show Favorites' option should be enabled or disabled */
+		SLATE_ARGUMENT(bool, CanShowFavorites)
+
+		/** Indicates if the 'Dock Collections' option should be enabled or disabled */
+		SLATE_ARGUMENT(bool, CanDockCollections)
+
+		/** Indicates if the context menu is going to load the assets, and if so to preload before the context menu is shown, and warn about the pending load. */
+		SLATE_ARGUMENT( bool, PreloadAssetsForContextMenu )
+
+		/** The selection mode the asset view should use */
+		SLATE_ARGUMENT( ESelectionMode::Type, SelectionMode )
+
+		/** Whether to allow dragging of items */
+		SLATE_ARGUMENT( bool, AllowDragging )
+
+		/** Whether this asset view should allow focus on sync or not */
+		SLATE_ARGUMENT( bool, AllowFocusOnSync )
+
+		/** Whether this asset view should allow the thumbnails to consume empty space after the user scale is applied */
+		SLATE_ARGUMENT( bool, FillEmptySpaceInTileView )
+
+		/** Should show Path in column view if true */
+		SLATE_ARGUMENT(bool, ShowPathInColumnView)
+
+		/** Should show Type in column view if true */
+		SLATE_ARGUMENT(bool, ShowTypeInColumnView)
+
+		/** Sort by path in the column view. Only works if the initial view type is Column */
+		SLATE_ARGUMENT(bool, SortByPathInColumnView)
+
+		/** Should always show engine content */
+		SLATE_ARGUMENT(bool, ForceShowEngineContent)
+
+		/** Should always show plugin content */
+		SLATE_ARGUMENT(bool, ForceShowPluginContent)
+
+		/** Called to check if an asset tag should be display in details view. */
+		SLATE_EVENT( FOnShouldDisplayAssetTag, OnAssetTagWantsToBeDisplayed )
+
+		/** Called to add extra asset data to the asset view, to display virtual assets. These get treated similar to Class assets */
+		SLATE_EVENT( FOnGetCustomSourceAssets, OnGetCustomSourceAssets )
+
+		/** Columns to hide by default */
+		SLATE_ARGUMENT( TArray<FString>, HiddenColumnNames )
+			
+		/** Custom columns that can be use specific */
+		SLATE_ARGUMENT( TArray<FAssetViewCustomColumn>, CustomColumns )
+
+		/** Custom columns that can be use specific */
+		SLATE_EVENT( FOnSearchOptionChanged, OnSearchOptionsChanged)
+
+		/** The content browser that owns this view if any */
+		SLATE_ARGUMENT(TSharedPtr<SContentBrowser>, OwningContentBrowser)
+	SLATE_END_ARGS()
+
+	~SAssetView();
+
+	/** Constructs this widget with InArgs */
+	void Construct( const FArguments& InArgs );
+
+	/** Changes the base sources for this view */
+	void SetSourcesData(const FSourcesData& InSourcesData);
+
+	/** Returns the sources filter applied to this asset view */
+	const FSourcesData& GetSourcesData() const;
+
+	/** Returns true if a real asset path is selected (i.e \Engine\* or \Game\*) */
+	bool IsAssetPathSelected() const;
+
+	/** Notifies the asset view that the filter-list filter has changed */
+	void SetBackendFilter(const FARFilter& InBackendFilter);
+
+	/** Handler for when a data source requests folder item creation */
+	void NewFolderItemRequested(const FContentBrowserItemTemporaryContext& NewItemContext);
+
+	/** Handler for when a data source requests file item creation */
+	void NewFileItemRequested(const FContentBrowserItemDataTemporaryContext& NewItemContext);
+
+	/** Creates a new asset item designed to allocate a new object once it is named. Uses the supplied factory to create the asset */
+	void CreateNewAsset(const FString& DefaultAssetName, const FString& PackagePath, UClass* AssetClass, UFactory* Factory);
+
+	/** Sets up an inline rename for the specified item */
+	void RenameItem(const FContentBrowserItem& ItemToRename);
+
+	/** Selects the specified items. */
+	void SyncToItems( TArrayView<const FContentBrowserItem> ItemsToSync, const bool bFocusOnSync = true );
+
+	/** Selects the specified virtual paths. */
+	void SyncToVirtualPaths( TArrayView<const FName> VirtualPathsToSync, const bool bFocusOnSync = true );
+
+	/** Selects the specified assets and paths. */
+	void SyncToLegacy( TArrayView<const FAssetData> AssetDataList, TArrayView<const FString> FolderList, const bool bFocusOnSync = true );
+
+	/** Sets the state of the asset view to the one described by the history data */
+	void ApplyHistoryData( const FHistoryData& History );
+
+	/** Returns all the items currently selected in the view */
+	TArray<TSharedPtr<FAssetViewItem>> GetSelectedViewItems() const;
+
+	/** Returns all the items currently selected in the view */
+	TArray<FContentBrowserItem> GetSelectedItems() const;
+
+	/** Returns all the folder items currently selected in the view */
+	TArray<FContentBrowserItem> GetSelectedFolderItems() const;
+
+	/** Returns all the file items currently selected in the view */
+	TArray<FContentBrowserItem> GetSelectedFileItems() const;
+
+	/** Returns all the asset data objects in items currently selected in the view */
+	TArray<FAssetData> GetSelectedAssets() const;
+
+	/** Returns all the folders currently selected in the view */
+	TArray<FString> GetSelectedFolders() const;
+
+	/** Requests that the asset view refreshes all it's source items. This is slow and should only be used if the source items change. */
+	void RequestSlowFullListRefresh();
+
+	/** Requests that the asset view refreshes only items that are filtered through frontend sources. This should be used when possible. */
+	void RequestQuickFrontendListRefresh();
+
+	/** Saves any settings to config that should be persistent between editor sessions */
+	void SaveSettings(const FString& IniFilename, const FString& IniSection, const FString& SettingsString) const;
+
+	/** Loads any settings to config that should be persistent between editor sessions */
+	void LoadSettings(const FString& IniFilename, const FString& IniSection, const FString& SettingsString);
+
+	/** Adjusts the selected asset by the selection delta, which should be +1 or -1) */
+	void AdjustActiveSelection(int32 SelectionDelta);
+
+	/** Returns true if an asset is currently in the process of being renamed */
+	bool IsRenamingAsset() const;
+
+	// SWidget inherited
+	virtual void Tick( const FGeometry& AllottedGeometry, const double InCurrentTime, const float InDeltaTime ) override;
+	virtual void OnDragLeave( const FDragDropEvent& DragDropEvent ) override;
+	virtual FReply OnDragOver( const FGeometry& MyGeometry, const FDragDropEvent& DragDropEvent ) override;
+	virtual FReply OnDrop( const FGeometry& MyGeometry, const FDragDropEvent& DragDropEvent ) override;
+	virtual FReply OnKeyChar( const FGeometry& MyGeometry,const FCharacterEvent& InCharacterEvent ) override;
+	virtual FReply OnKeyDown( const FGeometry& MyGeometry, const FKeyEvent& InKeyEvent ) override;
+	virtual FReply OnMouseWheel( const FGeometry& MyGeometry, const FPointerEvent& MouseEvent ) override;
+	virtual void OnFocusChanging( const FWeakWidgetPath& PreviousFocusPath, const FWidgetPath& NewWidgetPath, const FFocusEvent& InFocusEvent ) override;
+
+	/** Opens the selected assets or folders, depending on the selection */
+	void OnOpenAssetsOrFolders();
+
+	/** Loads the selected assets and previews them if possible */
+	void OnPreviewAssets();
+
+	/** Clears the selection of all the lists in the view */
+	void ClearSelection(bool bForceSilent = false);
+
+	/** Returns true if the asset view is in thumbnail editing mode */
+	bool IsThumbnailEditMode() const;
+
+	/** Delegate called when an editor setting is changed */
+	void HandleSettingChanged(FName PropertyName);
+
+	/** Set whether the user is currently searching or not */
+	void SetUserSearching(bool bInSearching);
+
+	/**
+	 * Forces the plugin content folder to be shown.
+	 *
+	 * @param bEnginePlugin		If true, also forces the engine folder to be shown.
+	 */
+	void ForceShowPluginFolder( bool bEnginePlugin );
+
+	/** @return true when we are including class names in search criteria */
+	bool IsIncludingClassNames() const;
+
+	/** @return true when we are including the entire asset path in search criteria */
+	bool IsIncludingAssetPaths() const;
+
+	/** @return true when we are including collection names in search criteria */
+	bool IsIncludingCollectionNames() const;
+
+	/** Handler for when the view combo button is clicked */
+	TSharedRef<SWidget> GetViewButtonContent();
+
+	/** Gets the text for the asset count label */
+	FText GetAssetCountText() const;
+
+private:
+
+	/** Sets the pending selection to the current selection (used when changing views or refreshing the view). */
+	void SyncToSelection( const bool bFocusOnSync = true );
+
+	/** @return the thumbnail scale setting path to use when looking up the setting in an ini. */
+	FString GetThumbnailScaleSettingPath(const FString& SettingsString) const;
+
+	/** @return the view type setting path to use when looking up the setting in an ini. */
+	FString GetCurrentViewTypeSettingPath(const FString& SettingsString) const;
+
+	/** Calculates a new filler scale used to adjust the thumbnails to fill empty space. */
+	void CalculateFillScale( const FGeometry& AllottedGeometry );
+
+	/** Calculates the latest color and opacity for the hint on thumbnails */
+	void CalculateThumbnailHintColorAndOpacity();
+
+	/** True if we have items pending that ProcessItemsPendingFilter still needs to process */
+	bool HasItemsPendingFilter() const;
+
+	/** Handles amortizing the additional filters */
+	void ProcessItemsPendingFilter(const double TickStartTime);
+
+	/** Creates a new tile view */
+	TSharedRef<class SAssetTileView> CreateTileView();
+
+	/** Creates a new list view */
+	TSharedRef<class SAssetListView> CreateListView();
+
+	/** Creates a new column view */
+	TSharedRef<class SAssetColumnView> CreateColumnView();
+
+	/** Returns true if the specified search token is allowed */
+	bool IsValidSearchToken(const FString& Token) const;
+
+	/** Regenerates the AssetItems list from the AssetRegistry */
+	void RefreshSourceItems();
+
+	/** Regenerates the FilteredAssetItems list from the AssetItems list */
+	void RefreshFilteredItems();
+
+	/** Sets the asset type that represents the majority of the assets in view */
+	void SetMajorityAssetType(FName NewMajorityAssetType);
+
+	/** Handler for when an asset is added to a collection */
+	void OnAssetsAddedToCollection( const FCollectionNameType& Collection, const TArray< FName >& ObjectPaths );
+
+	/** Handler for when an asset is removed from a collection */
+	void OnAssetsRemovedFromCollection( const FCollectionNameType& Collection, const TArray< FName >& ObjectPaths );
+
+	/** Handler for when a collection is renamed */
+	void OnCollectionRenamed( const FCollectionNameType& OriginalCollection, const FCollectionNameType& NewCollection );
+
+	/** Handler for when a collection is updated */
+	void OnCollectionUpdated( const FCollectionNameType& Collection );
+
+	/** Handler for when any frontend filters have been changed */
+	void OnFrontendFiltersChanged();
+
+	/** Returns true if there is any frontend filter active */
+	bool IsFrontendFilterActive() const;
+
+	/** Returns true if the specified Content Browser item passes all applied frontend (non asset registry) filters */
+	bool PassesCurrentFrontendFilter(const FContentBrowserItem& Item) const;
+
+	/** Returns true if the current filters deem that the asset view should be filtered recursively (overriding folder view) */
+	bool ShouldFilterRecursively() const;
+
+	/** Sorts the contents of the asset view alphabetically */
+	void SortList(bool bSyncToSelection = true);
+
+	/** Returns the thumbnails hint color and opacity */
+	FLinearColor GetThumbnailHintColorAndOpacity() const;
+
+	/** Register menu for when the view combo button is clicked */
+	static void RegisterGetViewButtonMenu();
+
+	/** Fill in menu content for when the view combo button is clicked */
+	void PopulateViewButtonMenu(class UToolMenu* Menu);
+
+	/** Toggle whether folders should be shown or not */
+	void ToggleShowFolders();
+
+	/** Whether or not it's possible to show folders */
+	bool IsToggleShowFoldersAllowed() const;
+
+	/** @return true when we are showing folders */
+	bool IsShowingFolders() const;
+
+	/** @return true when we are showing read-only folders */
+	bool IsShowingReadOnlyFolders() const;
+
+	/** Toggle whether empty folders should be shown or not */
+	void ToggleShowEmptyFolders();
+
+	/** Whether or not it's possible to show empty folders */
+	bool IsToggleShowEmptyFoldersAllowed() const;
+
+	/** @return true when we are showing empty folders */
+	bool IsShowingEmptyFolders() const;
+
+	/** Toggle whether localized content should be shown or not */
+	void ToggleShowLocalizedContent();
+
+	/** Whether or not it's possible to show localized content */
+	bool IsToggleShowLocalizedContentAllowed() const;
+
+	/** @return true when we are showing folders */
+	bool IsShowingLocalizedContent() const;
+
+	/** Toggle whether to show real-time thumbnails */
+	void ToggleRealTimeThumbnails();
+
+	/** Whether it is possible to show real-time thumbnails */
+	bool CanShowRealTimeThumbnails() const;
+
+	/** @return true if we are showing real-time thumbnails */
+	bool IsShowingRealTimeThumbnails() const;
+
+	/** Toggle whether plugin content should be shown or not */
+	void ToggleShowPluginContent();
+
+	/** @return true when we are showing plugin content */
+	bool IsShowingPluginContent() const;
+
+	/** Toggle whether engine content should be shown or not */
+	void ToggleShowEngineContent();
+
+	/** @return true when we are showing engine content */
+	bool IsShowingEngineContent() const;
+
+	/** Toggle whether developers content should be shown or not */
+	void ToggleShowDevelopersContent();
+
+	/** Whether or not it's possible to toggle developers content */
+	bool IsToggleShowDevelopersContentAllowed() const;
+
+	/** Whether or not it's possible to toggle engine content */
+	bool IsToggleShowEngineContentAllowed() const;
+
+	/** Whether or not it's possible to toggle plugin content */
+	bool IsToggleShowPluginContentAllowed() const;
+	
+	/** @return true when we are showing the developers content */
+	bool IsShowingDevelopersContent() const;
+
+	/** Toggle whether favorites should be shown or not */
+	void ToggleShowFavorites();
+
+	/** Whether or not it's possible to toggle favorites */
+	bool IsToggleShowFavoritesAllowed() const;
+
+	/** @return true when we are showing favorites */
+	bool IsShowingFavorites() const;
+
+	/** Toggle whether the collections view should be docked under the paths view */
+	void ToggleDockCollections();
+
+	/** Whether or not it's possible to dock the collections view */
+	bool IsToggleDockCollectionsAllowed() const;
+
+	/** @return true when the collections view is docked */
+	bool HasDockedCollections() const;
+
+	/** Toggle whether C++ content should be shown or not */
+	void ToggleShowCppContent();
+
+	/** Whether or not it's possible to show C++ content */
+	bool IsToggleShowCppContentAllowed() const;
+
+	/** @return true when we are showing C++ content */
+	bool IsShowingCppContent() const;
+
+	/** Toggle whether class names should be included in search criteria */
+	void ToggleIncludeClassNames();
+
+	/** Whether or not it's possible to include class names in search criteria */
+	bool IsToggleIncludeClassNamesAllowed() const;
+
+	/** Toggle whether the entire asset path should be included in search criteria */
+	void ToggleIncludeAssetPaths();
+
+	/** Whether or not it's possible to include the entire asset path in search criteria */
+	bool IsToggleIncludeAssetPathsAllowed() const;
+
+	/** Toggle whether collection names should be included in search criteria */
+	void ToggleIncludeCollectionNames();
+
+	/** Whether or not it's possible to include collection names in search criteria */
+	bool IsToggleIncludeCollectionNamesAllowed() const;
+
+	/** @return true when we are filtering recursively when we have an asset path */
+	bool IsFilteringRecursively() const;
+
+	/** Whether or not it's possible to toggle how to filtering recursively */
+	bool IsToggleFilteringRecursivelyAllowed() const;
+
+	/** Toggle whether we're filtering recursively */
+	void ToggleFilteringRecursively();
+
+	/** Toggle whether we are showing the all virtual folder */
+	void ToggleShowAllFolder();
+
+	/** @return true when we are showing the all virtual folder */
+	bool IsShowingAllFolder() const;
+
+	/** Toggle whether we are re-organizing folders into virtual locations */
+	void ToggleOrganizeFolders();
+
+	/** @return true when we are organizing folders into virtual locations */
+	bool IsOrganizingFolders() const;
+
+	/** Sets the view type and updates lists accordingly */
+	void SetCurrentViewType(EAssetViewType::Type NewType);
+
+	/** Sets the view type and forcibly dismisses all currently open context menus */
+	void SetCurrentViewTypeFromMenu(EAssetViewType::Type NewType);
+
+	/** Clears the reference to the current view and creates a new one, based on CurrentViewType */
+	void CreateCurrentView();
+
+	/** Gets the current view type (list or tile) */
+	EAssetViewType::Type GetCurrentViewType() const;
+
+	TSharedRef<SWidget> CreateShadowOverlay( TSharedRef<STableViewBase> Table );
+
+	/** Returns true if ViewType is the current view type */
+	bool IsCurrentViewType(EAssetViewType::Type ViewType) const;
+
+	/** Set the keyboard focus to the correct list view that should be active */
+	void FocusList() const;
+
+	/** Refreshes the list view to display any changes made to the non-filtered assets */
+	void RefreshList();
+
+	/** Sets the sole selection for all lists in the view */
+	void SetSelection(const TSharedPtr<FAssetViewItem>& Item);
+
+	/** Sets selection for an item in all lists in the view */
+	void SetItemSelection(const TSharedPtr<FAssetViewItem>& Item, bool bSelected, const ESelectInfo::Type SelectInfo = ESelectInfo::Direct);
+
+	/** Scrolls the selected item into view for all lists in the view */
+	void RequestScrollIntoView(const TSharedPtr<FAssetViewItem>& Item);
+
+	/** Handler for list view widget creation */
+	TSharedRef<ITableRow> MakeListViewWidget(TSharedPtr<FAssetViewItem> AssetItem, const TSharedRef<STableViewBase>& OwnerTable);
+
+	/** Handler for tile view widget creation */
+	TSharedRef<ITableRow> MakeTileViewWidget(TSharedPtr<FAssetViewItem> AssetItem, const TSharedRef<STableViewBase>& OwnerTable);
+
+	/** Handler for column view widget creation */
+	TSharedRef<ITableRow> MakeColumnViewWidget(TSharedPtr<FAssetViewItem> AssetItem, const TSharedRef<STableViewBase>& OwnerTable);
+
+	/** Handler for when any asset item widget gets destroyed */
+	void AssetItemWidgetDestroyed(const TSharedPtr<FAssetViewItem>& Item);
+	
+	/** Creates new thumbnails that are near the view area and deletes old thumbnails that are no longer relevant. */
+	void UpdateThumbnails();
+
+	/**  Helper function for UpdateThumbnails. Adds the specified item to the new thumbnail relevancy map and creates any thumbnails for new items. Returns the thumbnail. */
+	TSharedPtr<class FAssetThumbnail> AddItemToNewThumbnailRelevancyMap(const TSharedPtr<FAssetViewItem>& Item, TMap< TSharedPtr<FAssetViewItem>, TSharedPtr<class FAssetThumbnail> >& NewRelevantThumbnails);
+
+	/** Handler for tree view selection changes */
+	void AssetSelectionChanged( TSharedPtr<FAssetViewItem > AssetItem, ESelectInfo::Type SelectInfo );
+
+	/** Handler for when an item has scrolled into view after having been requested to do so */
+	void ItemScrolledIntoView(TSharedPtr<FAssetViewItem> AssetItem, const TSharedPtr<ITableRow>& Widget);
+
+	/** Handler for context menus */
+	TSharedPtr<SWidget> OnGetContextMenuContent();
+
+	/** Handler called when an asset context menu is about to open */
+	bool CanOpenContextMenu() const;
+
+	/** Handler for double clicking an item */
+	void OnListMouseButtonDoubleClick(TSharedPtr<FAssetViewItem> AssetItem);
+
+	/** Handle dragging an asset */
+	FReply OnDraggingAssetItem(const FGeometry& MyGeometry, const FPointerEvent& MouseEvent);
+
+	/** Checks if the asset name being committed is valid */
+	bool AssetVerifyRenameCommit(const TSharedPtr<FAssetViewItem>& Item, const FText& NewName, const FSlateRect& MessageAnchor, FText& OutErrorMessage);
+
+	/** An asset item has started to be renamed */
+	void AssetRenameBegin(const TSharedPtr<FAssetViewItem>& Item, const FString& NewName, const FSlateRect& MessageAnchor);
+
+	/** An asset item that was prompting the user for a new name was committed. Return false to indicate that the user should enter a new name */
+	void AssetRenameCommit(const TSharedPtr<FAssetViewItem>& Item, const FString& NewName, const FSlateRect& MessageAnchor, const ETextCommit::Type CommitType);
+
+	/** Gets the color and opacity for all names of assets in the asset view */
+	FLinearColor GetAssetNameColorAndOpacity() const;
+
+	/** Returns true if tooltips should be allowed right now. Tooltips are typically disabled while right click scrolling. */
+	bool ShouldAllowToolTips() const;
+
+	/** Returns true if the asset view is currently allowing the user to edit thumbnails */
+	bool IsThumbnailEditModeAllowed() const;
+
+	/** The "Done Editing" button was pressed in the thumbnail edit mode strip */
+	FReply EndThumbnailEditModeClicked();
+
+	/** Gets the visibility of the Thumbnail Edit Mode label */
+	EVisibility GetEditModeLabelVisibility() const;
+
+	/** Gets the visibility of the list view */
+	EVisibility GetListViewVisibility() const;
+
+	/** Gets the visibility of the tile view */
+	EVisibility GetTileViewVisibility() const;
+
+	/** Gets the visibility of the column view */
+	EVisibility GetColumnViewVisibility() const;
+
+	/** Toggles thumbnail editing mode */
+	void ToggleThumbnailEditMode();
+
+	/** Called when  thumbnail size is changed */
+	void OnThumbnailSizeChanged(EThumbnailSize NewThumbnailSize);
+
+	/** Is the current thumbnail size menu option checked */
+	bool IsThumbnailSizeChecked(EThumbnailSize InThumbnailSize) const;
+
+	/** Are thumbnails allowed to be scaled by the user */
+	bool IsThumbnailScalingAllowed() const;
+
+	/** Gets the current thumbnail scale */
+	float GetThumbnailScale() const;
+
+	/** Gets the current thumbnail size enum */
+	EThumbnailSize GetThumbnailSize() const { return ThumbnailSize; }
+
+	/** Gets the scaled item height for the list view */
+	float GetListViewItemHeight() const;
+	
+	/** Gets the final scaled item height for the tile view */
+	float GetTileViewItemHeight() const;
+
+	/** Gets the scaled item height for the tile view before the filler scale is applied */
+	float GetTileViewItemBaseHeight() const;
+
+	/** Gets the final scaled item width for the tile view */
+	float GetTileViewItemWidth() const;
+
+	/** Gets the scaled item width for the tile view before the filler scale is applied */
+	float GetTileViewItemBaseWidth() const;
+
+	/** Gets the sort mode for the supplied ColumnId */
+	EColumnSortMode::Type GetColumnSortMode(const FName ColumnId) const;
+
+	/** Gets the sort order for the supplied ColumnId */
+	EColumnSortPriority::Type GetColumnSortPriority(const FName ColumnId) const;
+
+	/** Handler for when a column header is clicked */
+	void OnSortColumnHeader(const EColumnSortPriority::Type SortPriority, const FName& ColumnId, const EColumnSortMode::Type NewSortMode);
+
+	/** @return The state of the is working progress bar */
+	TOptional< float > GetIsWorkingProgressBarState() const;
+
+	/** Is the no assets to show warning visible? */
+	EVisibility IsAssetShowWarningTextVisible() const;
+
+	/** Gets the text for displaying no assets to show warning */
+	FText GetAssetShowWarningText() const;
+
+	/** Whether we have a single source collection selected */
+	bool HasSingleCollectionSource() const;
+
+	/** Starts the async flow of creating a file or folder item from deferred data */
+	void BeginCreateDeferredItem();
+
+	/** Finishes the async flow of creating a file or folder item from deferred data */
+	FContentBrowserItem EndCreateDeferredItem(const TSharedPtr<FAssetViewItem>& InItem, const FString& InName, const bool bFinalize, FText& OutErrorText);
+
+	/** @return The current quick-jump term */
+	FText GetQuickJumpTerm() const;
+
+	/** @return Whether the quick-jump term is currently visible? */
+	EVisibility IsQuickJumpVisible() const;
+
+	/** @return The color that should be used for the quick-jump term */
+	FSlateColor GetQuickJumpColor() const;
+
+	/** Reset the quick-jump to its empty state */
+	void ResetQuickJump();
+
+	/**
+	 * Called from OnKeyChar and OnKeyDown to handle quick-jump key presses
+	 * @param InCharacter		The character that was typed
+	 * @param bIsControlDown	Was the control key pressed?
+	 * @param bIsAltDown		Was the alt key pressed?
+	 * @param bTestOnly			True if we only want to test whether the key press would be handled, but not actually update the quick-jump term
+	 * @return FReply::Handled or FReply::Unhandled
+	 */
+	FReply HandleQuickJumpKeyDown(const TCHAR InCharacter, const bool bIsControlDown, const bool bIsAltDown, const bool bTestOnly);
+
+	/**
+	 * Perform a quick-jump to the next available asset in FilteredAssetItems that matches the current term
+	 * @param bWasJumping		True if we were performing an ongoing quick-jump last Tick
+	 * @return True if the quick-jump found a valid match, false otherwise
+	 */
+	bool PerformQuickJump(const bool bWasJumping);
+	
+	/** Generates the column filtering menu */
+	void FillToggleColumnsMenu(FMenuBuilder& MenuBuilder);
+
+	/** Resets the column filtering state to make them all visible */
+	void ResetColumns();
+
+	/** Export columns to CSV */
+	void ExportColumns();
+
+	/** Toggle the column at ColumnIndex */
+	void ToggleColumn(const FString ColumnName);
+	/** Sets the column visibility by removing/inserting the column*/
+	void SetColumnVisibility(const FString ColumnName, const bool bShow);
+
+	/** Whether or not a column can be toggled, has to be valid column and mandatory minimum number of columns = 1*/
+	bool CanToggleColumn(const FString ColumnName) const;
+
+	/** Whether or not a column is visible to show it's state in the filtering menu */
+	bool IsColumnVisible(const FString ColumnName) const;
+
+	/** Creates the row header context menu allowing for hiding individually clicked columns*/
+	TSharedRef<SWidget> CreateRowHeaderMenuContent(const FString ColumnName);
+
+	/** Will compute the max row size from all its children for the specified column id*/
+	FVector2D GetMaxRowSizeForColumn(const FName& ColumnId);
+
+	/** Append the current effective backend filter (intersection of BackendFilter and SupportedFilter) to the given filter. */
+	void AppendBackendFilter(FARFilter& FilterToAppendTo) const;
+
+	FContentBrowserDataFilter CreateBackendDataFilter() const;
+
+	/** Handles updating the view when content items are changed */
+	void HandleItemDataUpdated(TArrayView<const FContentBrowserItemDataUpdate> InUpdatedItems);
+
+	/** Notification for when the content browser has completed it's initial search */
+	void HandleItemDataDiscoveryComplete();
+
+private:
+	friend class FAssetViewFrontendFilterHelper;
+
+	/** The available items from querying the backend data sources */
+	TMap<FContentBrowserItemKey, TSharedPtr<FAssetViewItem>> AvailableBackendItems;
+
+	/**
+	 * The items from AvailableBackendItems that are pending a run through any additional filtering before they can be shown in the filtered view list.
+	 * @note This filtering will run without amortization via ProcessItemsPendingFilter, so only use it for items that *must* be processed this frame.
+	 */
+	TSet<TSharedPtr<FAssetViewItem>> ItemsPendingPriorityFilter;
+
+	/**
+	 * The items from AvailableBackendItems that are pending a run through any additional frontend filters before they can be shown in the filtered view list.
+	 * @note This filtering will run amortized on the game thread via ProcessItemsPendingFilter.
+	 */
+	TSet<TSharedPtr<FAssetViewItem>> ItemsPendingFrontendFilter;
+
+	/** The items that are being shown in the filtered view list */
+	TArray<TSharedPtr<FAssetViewItem>> FilteredAssetItems;
+
+	/* Map of an item name to the current count of FilteredAssetItems that are of that type */ 
+	TMap<FName, int32> FilteredAssetItemTypeCounts;
+
+	/** The list view that is displaying the assets */
+	EAssetViewType::Type CurrentViewType;
+	TSharedPtr<class SAssetListView> ListView;
+	TSharedPtr<class SAssetTileView> TileView;
+	TSharedPtr<class SAssetColumnView> ColumnView;
+	TSharedPtr<SBox> ViewContainer;
+
+	/** The content browser that created this asset view if any */
+	TWeakPtr<SContentBrowser> OwningContentBrowser;
+
+	/** The current base source filter for the view */
+	FSourcesData SourcesData;
+	FARFilter BackendFilter;
+	TSharedPtr<FBlacklistNames> AssetClassBlacklist;
+	TSharedPtr<FBlacklistPaths> FolderBlacklist;
+	TSharedPtr<FBlacklistPaths> WritableFolderBlacklist;
+	TSharedPtr<FAssetFilterCollectionType> FrontendFilters;
+
+	/** Show path view filters submenu in view options menu  */
+	bool bShowPathViewFilters;
+
+	/** If true, the source items will be refreshed next frame. Very slow. */
+	bool bSlowFullListRefreshRequested;
+
+	/** If true, the frontend items will be refreshed next frame. Much faster. */
+	bool bQuickFrontendListRefreshRequested;
+
+	/** The list of items to sync next frame */
+	FSelectionData PendingSyncItems;
+
+	/** Should we take focus when the PendingSyncAssets are processed? */
+	bool bPendingFocusOnSync;
+
+	/** Was recursive filtering enabled the last time a full slow refresh performed? */
+	bool bWereItemsRecursivelyFiltered;
+
+	/** Called to check if an asset should be filtered out by external code */
+	FOnShouldFilterAsset OnShouldFilterAsset;
+
+	/** Called when the asset view is asked to start to create a temporary item */
+	FOnAssetViewNewItemRequested OnNewItemRequested;
+
+	/** Called when an item was selected in the list. Provides more context than OnAssetSelected. */
+	FOnContentBrowserItemSelectionChanged OnItemSelectionChanged;
+
+	/** Called when the user double clicks, presses enter, or presses space on a Content Browser item */
+	FOnContentBrowserItemsActivated OnItemsActivated;
+
+	/** Delegate to invoke when generating the context menu for an item */
+	FOnGetContentBrowserItemContextMenu OnGetItemContextMenu;
+
+	/** Called when the user has committed a rename of one or more items */
+	FOnContentBrowserItemRenameCommitted OnItemRenameCommitted;
+
+	/** Called to check if an asset tag should be display in details view. */
+	FOnShouldDisplayAssetTag OnAssetTagWantsToBeDisplayed;
+
+	/** Called to see if it is valid to get a custom asset tool tip */
+	FOnIsAssetValidForCustomToolTip OnIsAssetValidForCustomToolTip;
+
+	/** Called to get a custom asset item tooltip (If necessary) */
+	FOnGetCustomAssetToolTip OnGetCustomAssetToolTip;
+
+	/** Called when a custom asset item is about to show a tooltip */
+	FOnVisualizeAssetToolTip OnVisualizeAssetToolTip;
+
+	/** Called when a custom asset item's tooltip is closing */
+	FOnAssetToolTipClosing OnAssetToolTipClosing;
+
+	/** Called to add extra asset data to the asset view, to display virtual assets. These get treated similar to Class assets */
+	FOnGetCustomSourceAssets OnGetCustomSourceAssets;
+
+	/** Called when a search option changes to notify that results should be rebuilt */
+	FOnSearchOptionChanged OnSearchOptionsChanged;
+
+	/** Called when opening view options menu */
+	FOnExtendAssetViewOptionsMenuContext OnExtendAssetViewOptionsMenuContext;
+
+	/** When true, filtered list items will be sorted next tick. Provided another sort hasn't happened recently or we are renaming an asset */
+	bool bPendingSortFilteredItems;
+	double CurrentTime;
+	double LastSortTime;
+	double SortDelaySeconds;
+
+	/** Weak ptr to the asset that is waiting to be renamed when scrolled into view, and the window is active */
+	TWeakPtr<FAssetViewItem> AwaitingRename;
+
+	/** Set when the user is in the process of naming an asset */
+	TWeakPtr<FAssetViewItem> RenamingAsset;
+
+	/** Pool for maintaining and rendering thumbnails */
+	TSharedPtr<class FAssetThumbnailPool> AssetThumbnailPool;
+
+	/** A map of FAssetViewAsset to the thumbnail that represents it. Only items that are currently visible or within half of the FilteredAssetItems array index distance described by NumOffscreenThumbnails are in this list */
+	TMap< TSharedPtr<FAssetViewItem>, TSharedPtr<class FAssetThumbnail> > RelevantThumbnails;
+
+	/** The set of FAssetItems that currently have widgets displaying them. */
+	TArray< TSharedPtr<FAssetViewItem> > VisibleItems;
+
+	/** The number of thumbnails to keep for asset items that are not currently visible. Half of the thumbnails will be before the earliest item and half will be after the latest. */
+	int32 NumOffscreenThumbnails;
+
+	/** The current size of relevant thumbnails */
+	int32 CurrentThumbnailSize;
+
+	/** Flag to defer thumbnail updates until the next frame */
+	bool bPendingUpdateThumbnails;
+
+	/** The size of thumbnails */
+	int32 ListViewThumbnailResolution;
+	int32 ListViewThumbnailSize;
+	int32 ListViewThumbnailPadding;
+	int32 TileViewThumbnailResolution;
+	int32 TileViewThumbnailSize;
+	int32 TileViewThumbnailPadding;
+	int32 TileViewNameHeight;
+
+	/** The max and min thumbnail scales as a fraction of the rendered size */
+	float MinThumbnailScale;
+	float MaxThumbnailScale;
+
+	/** Scalar applied to thumbnail sizes so that users thumbnails are scaled based on users display area size*/
+	float ThumbnailScaleRangeScalar;
+
+	/** Current thumbnail size */
+	EThumbnailSize ThumbnailSize;
+
+	/** Flag indicating if we will be filling the empty space in the tile view. */
+	bool bFillEmptySpaceInTileView;
+
+	/** The amount to scale each thumbnail so that the empty space is filled. */
+	float FillScale;
+
+	/** When in columns view, this is the name of the asset type which is most commonly found in the recent results */
+	FName MajorityAssetType;
+
+	/** The manager responsible for sorting assets in the view */
+	FAssetViewSortManager SortManager;
+
+	/** When true, selection change notifications will not be sent */
+	bool bBulkSelecting;
+
+	/** When true, the user may edit thumbnails */
+	bool bAllowThumbnailEditMode;
+
+	/** True when the asset view is currently allowing the user to edit thumbnails */
+	bool bThumbnailEditMode;
+
+	/** Indicates if this view is allowed to show classes */
+	bool bCanShowClasses;
+
+	/** Indicates if the 'Show Folders' option should be enabled or disabled */
+	bool bCanShowFolders;
+
+	/** Indicates if this view is allowed to show folders that cannot be written to */
+	bool bCanShowReadOnlyFolders;
+
+	/** If true, recursive filtering will be caused by applying a backend filter */
+	bool bFilterRecursivelyWithBackendFilter;
+
+	/** Indicates if the 'Real-Time Thumbnails' option should be enabled or disabled */
+	bool bCanShowRealTimeThumbnails;
+
+	/** Indicates if the 'Show Developers' option should be enabled or disabled */
+	bool bCanShowDevelopersFolder;
+
+	/** Indicates if the 'Show Favorites' option should be enabled or disabled */
+	bool bCanShowFavorites;
+
+	/** Indicates if the 'Dock Collections' option should be enabled or disabled */
+	bool bCanDockCollections;
+
+	/** Indicates if the context menu is going to load the assets, and if so to preload before the context menu is shown, and warn about the pending load. */
+	bool bPreloadAssetsForContextMenu;
+
+	/** If true, it will show path column in the asset view */
+	bool bShowPathInColumnView;
+
+	/** If true, it will show type in the asset view */
+	bool bShowTypeInColumnView;
+
+	/** If true, it sorts by path and then name */
+	bool bSortByPathInColumnView;
+
+	/** If true, engine content is always shown */
+	bool bForceShowEngineContent;
+
+	/** If true, plugin content is always shown */
+	bool bForceShowPluginContent;
+
+	/** The current selection mode used by the asset view */
+	ESelectionMode::Type SelectionMode;
+
+	/** The max number of results to process per tick */
+	float MaxSecondsPerFrame;
+
+	/** When delegate amortization began */
+	double AmortizeStartTime;
+
+	/** The total time spent amortizing the delegate filter */
+	double TotalAmortizeTime;
+
+	/** Whether the asset view is currently working on something and should display a cue to the user */
+	bool bIsWorking;
+
+	/** The text to highlight on the assets */
+	TAttribute< FText > HighlightedText;
+
+	/** What the label on the thumbnails should be */
+	EThumbnailLabel::Type ThumbnailLabel;
+
+	/** Whether to ever show the hint label on thumbnails */
+	bool AllowThumbnailHintLabel;
+
+	/** The sequence used to generate the opacity of the thumbnail hint */
+	FCurveSequence ThumbnailHintFadeInSequence;
+
+	/** The current thumbnail hint color and opacity*/
+	FLinearColor ThumbnailHintColorAndOpacity;
+
+	/** The text to show when there are no assets to show */
+	TAttribute< FText > AssetShowWarningText;
+
+	/** Initial set of item categories that this view should show - may be adjusted further by things like CanShowClasses or legacy delegate bindings */
+	EContentBrowserItemCategoryFilter InitialCategoryFilter;
+
+	/** Whether to allow dragging of items */
+	bool bAllowDragging;
+
+	/** Whether this asset view should allow focus on sync or not */
+	bool bAllowFocusOnSync;
+
+	/** Flag set if the user is currently searching */
+	bool bUserSearching;
+	
+	/** Whether or not to notify about newly selected items on on the next asset sync */
+	bool bShouldNotifyNextAssetSync;
+
+	/** A struct to hold data for the deferred creation of a file or folder item */
+	struct FCreateDeferredItemData
+	{
+		FContentBrowserItemTemporaryContext ItemContext;
+
+		bool bWasAddedToView = false;
+	};
+
+	/** File or folder item pending deferred creation */
+	TUniquePtr<FCreateDeferredItemData> DeferredItemToCreate;
+
+	/** Struct holding the data for the asset quick-jump */
+	struct FQuickJumpData
+	{
+		FQuickJumpData()
+			: bIsJumping(false)
+			, bHasChangedSinceLastTick(false)
+			, bHasValidMatch(false)
+			, LastJumpTime(0)
+		{
+		}
+
+		/** True if we're currently performing an ongoing quick-jump */
+		bool bIsJumping;
+
+		/** True if the jump data has changed since the last Tick */
+		bool bHasChangedSinceLastTick;
+
+		/** True if the jump term found a valid match */
+		bool bHasValidMatch;
+
+		/** Time (taken from Tick) that we last performed a quick-jump */
+		double LastJumpTime;
+
+		/** The string we should be be looking for */
+		FString JumpTerm;
+	};
+
+	/** Data for the asset quick-jump */
+	FQuickJumpData QuickJumpData;
+	
+	/** Column filtering state */
+	TArray<FString> DefaultHiddenColumnNames;
+	TArray<FString> HiddenColumnNames;
+	int32 NumVisibleColumns;
+
+	TArray<FAssetViewCustomColumn> CustomColumns;
+public:
+	bool ShouldColumnGenerateWidget(const FString ColumnName) const;
+};
diff --git a/Engine/Source/Editor/ContentBrowser/Public/SourcesData.h b/Engine/Source/Editor/ContentBrowser/Public/SourcesData.h
new file mode 100644
--- /dev/null
+++ b/Engine/Source/Editor/ContentBrowser/Public/SourcesData.h
@@ -0,0 +1,67 @@
+// Copyright Epic Games, Inc. All Rights Reserved.
+
+#pragma once
+
+#include "CoreMinimal.h"
+#include "CollectionManagerTypes.h"
+#include "ICollectionManager.h"
+#include "CollectionManagerModule.h"
+
+struct FSourcesData
+{
+	TArray<FName> VirtualPaths;
+	TArray<FCollectionNameType> Collections;
+
+	FSourcesData() = default;
+
+	explicit FSourcesData(FName InVirtualPath)
+	{
+		VirtualPaths.Add(InVirtualPath);
+	}
+
+	explicit FSourcesData(FCollectionNameType InCollection)
+	{
+		Collections.Add(InCollection);
+	}
+
+	FSourcesData(TArray<FName> InVirtualPaths, TArray<FCollectionNameType> InCollections)
+		: VirtualPaths(MoveTemp(InVirtualPaths))
+		, Collections(MoveTemp(InCollections))
+	{
+	}
+
+	FSourcesData(const FSourcesData& Other) = default;
+	FSourcesData(FSourcesData&& Other) = default;
+
+	FSourcesData& operator=(const FSourcesData& Other) = default;
+	FSourcesData& operator=(FSourcesData&& Other) = default;
+
+	FORCEINLINE bool IsEmpty() const
+	{
+		return VirtualPaths.Num() == 0 && Collections.Num() == 0;
+	}
+
+	FORCEINLINE bool HasVirtualPaths() const
+	{
+		return VirtualPaths.Num() > 0;
+	}
+
+	FORCEINLINE bool HasCollections() const
+	{
+		return Collections.Num() > 0;
+	}
+
+	bool IsDynamicCollection() const
+	{
+		if ( Collections.Num() == 1 && FCollectionManagerModule::IsModuleAvailable() )
+		{
+			// Collection manager module should already be loaded since it may cause a hitch on the first search
+			FCollectionManagerModule& CollectionManagerModule = FCollectionManagerModule::GetModule();
+
+			ECollectionStorageMode::Type StorageMode = ECollectionStorageMode::Static;
+			return (CollectionManagerModule.Get().GetCollectionStorageMode(Collections[0].Name, Collections[0].Type, StorageMode) && StorageMode == ECollectionStorageMode::Dynamic);
+		}
+
+		return false;
+	}
+};
diff --git a/Engine/Source/Editor/UMGEditor/Private/BlueprintModes/WidgetDesignerApplicationMode.cpp b/Engine/Source/Editor/UMGEditor/Private/BlueprintModes/WidgetDesignerApplicationMode.cpp
--- a/Engine/Source/Editor/UMGEditor/Private/BlueprintModes/WidgetDesignerApplicationMode.cpp
+++ b/Engine/Source/Editor/UMGEditor/Private/BlueprintModes/WidgetDesignerApplicationMode.cpp
@@ -11,6 +11,7 @@
 #include "ToolMenus.h"
 
 #include "TabFactory/PaletteTabSummoner.h"
+#include "TabFactory/LibraryTabSummoner.h"
 #include "TabFactory/HierarchyTabSummoner.h"
 #include "TabFactory/BindWidgetTabSummoner.h"
 #include "TabFactory/DesignerTabSummoner.h"
@@ -31,7 +32,7 @@ FWidgetDesignerApplicationMode::FWidgetDesignerApplicationMode(TSharedPtr<FWidge
 	// Override the default created category here since "Designer Editor" sounds awkward
 	WorkspaceMenuCategory = FWorkspaceItem::NewGroup(LOCTEXT("WorkspaceMenu_WidgetDesigner", "Widget Designer"));
 
-	TabLayout = FTabManager::NewLayout( "WidgetBlueprintEditor_Designer_Layout_v4_2" )
+	TabLayout = FTabManager::NewLayout("WidgetBlueprintEditor_Designer_Layout_v4_3")
 	->AddArea
 	(
 		FTabManager::NewPrimaryArea()
@@ -49,26 +50,18 @@ FWidgetDesignerApplicationMode::FWidgetDesignerApplicationMode(TSharedPtr<FWidge
 				->Split
 				(
 					FTabManager::NewStack()
-					->SetSizeCoefficient( 0.5f )
-					->AddTab( FPaletteTabSummoner::TabID, ETabState::OpenedTab )
+					->SetSizeCoefficient(0.5f)
+					->SetForegroundTab(FPaletteTabSummoner::TabID)
+					->AddTab(FPaletteTabSummoner::TabID, ETabState::OpenedTab)
+					->AddTab(FLibraryTabSummoner::TabID, ETabState::ClosedTab)
 				)
 				->Split
 				(
-					FTabManager::NewSplitter()
-					->SetSizeCoefficient(0.15f)
-					->SetOrientation(Orient_Horizontal)
-					->Split
-					(
-						FTabManager::NewStack()
-						->SetSizeCoefficient(0.5f)
-						->AddTab(FHierarchyTabSummoner::TabID, ETabState::OpenedTab)
-					)
-					->Split
-					(
-						FTabManager::NewStack()
-						->SetSizeCoefficient(0.5f)
-						->AddTab(FBindWidgetTabSummoner::TabID, ETabState::OpenedTab)
-					)
+					FTabManager::NewStack()
+					->SetSizeCoefficient(0.5f)
+					->SetForegroundTab(FHierarchyTabSummoner::TabID)
+					->AddTab(FHierarchyTabSummoner::TabID, ETabState::OpenedTab)
+					->AddTab(FBindWidgetTabSummoner::TabID, ETabState::OpenedTab)
 				)
 			)
 			->Split
@@ -119,6 +112,7 @@ FWidgetDesignerApplicationMode::FWidgetDesignerApplicationMode(TSharedPtr<FWidge
 	TabFactories.RegisterFactory(MakeShareable(new FHierarchyTabSummoner(InWidgetEditor)));
 	TabFactories.RegisterFactory(MakeShareable(new FBindWidgetTabSummoner(InWidgetEditor)));
 	TabFactories.RegisterFactory(MakeShareable(new FPaletteTabSummoner(InWidgetEditor)));
+	TabFactories.RegisterFactory(MakeShareable(new FLibraryTabSummoner(InWidgetEditor)));
 	TabFactories.RegisterFactory(MakeShareable(new FSequencerTabSummoner(InWidgetEditor)));
 	TabFactories.RegisterFactory(MakeShareable(new FAnimationTabSummoner(InWidgetEditor)));
 	TabFactories.RegisterFactory(MakeShareable(new FCompilerResultsSummoner(InWidgetEditor)));
diff --git a/Engine/Source/Editor/UMGEditor/Private/Designer/SDesignerView.cpp b/Engine/Source/Editor/UMGEditor/Private/Designer/SDesignerView.cpp
--- a/Engine/Source/Editor/UMGEditor/Private/Designer/SDesignerView.cpp
+++ b/Engine/Source/Editor/UMGEditor/Private/Designer/SDesignerView.cpp
@@ -2640,60 +2640,18 @@ void SDesignerView::DetermineDragDropPreviewWidgets(TArray<UWidget*>& OutWidgets
 		return;
 	}
 
-	TSharedPtr<FWidgetTemplateDragDropOp> TemplateDragDropOp = DragDropEvent.GetOperationAs<FWidgetTemplateDragDropOp>();
-	TSharedPtr<FAssetDragDropOp> AssetDragDropOp = DragDropEvent.GetOperationAs<FAssetDragDropOp>();
+	TSharedPtr<FDragDropOperation> DragDropOp = DragDropEvent.GetOperation();
+	UWidget* Widget = FWidgetBlueprintEditorUtils::GetWidgetTemplateFromDragDrop(Blueprint, RootWidgetTree, DragDropOp);
 
-	if (TemplateDragDropOp.IsValid())
+	if (Widget)
 	{
-		UWidget* Widget = TemplateDragDropOp->Template->Create(RootWidgetTree);
-
-		if (Widget)
-		{
-			if ( Cast<UUserWidget>(Widget) == nullptr || Blueprint->IsWidgetFreeFromCircularReferences(Cast<UUserWidget>(Widget)) )
-			{
-				OutWidgets.Add(Widget);
-			}
-		}
-	}
-	else if (AssetDragDropOp.IsValid())
-	{
-		for (const FAssetData& AssetData : AssetDragDropOp->GetAssets())
-		{
-			UWidget* Widget = nullptr;
-			UClass* AssetClass = FindObjectChecked<UClass>(ANY_PACKAGE, *AssetData.AssetClass.ToString());
-
-			if (FWidgetTemplateBlueprintClass::Supports(AssetClass))
-			{
-				// Allows a UMG Widget Blueprint to be dragged from the Content Browser to another Widget Blueprint...as long as we're not trying to place a
-				// blueprint inside itself.
-				FString BlueprintPath = Blueprint->GetPathName();
-				if (BlueprintPath != AssetData.ObjectPath.ToString())
-				{
-					Widget = FWidgetTemplateBlueprintClass(AssetData).Create(RootWidgetTree);
-
-					// Check to make sure that this widget can be added to the current blueprint
-					if ( Cast<UUserWidget>(Widget) != nullptr && !Blueprint->IsWidgetFreeFromCircularReferences(Cast<UUserWidget>(Widget)) )
-					{
-						Widget = nullptr;
-					}
-				}
-			}
-			else if (FWidgetTemplateImageClass::Supports(AssetClass))
-			{
-				Widget = FWidgetTemplateImageClass(AssetData).Create(RootWidgetTree);
-			}
-
-			if (Widget)
-			{
-				OutWidgets.Add(Widget);
-			}
-		}
+		OutWidgets.Add(Widget);
 	}
 
 	// Mark the widgets for design-time rendering
-	for (UWidget* Widget : OutWidgets)
+	for (UWidget* OutWidget : OutWidgets)
 	{
-		Widget->SetDesignerFlags(BlueprintEditor.Pin()->GetCurrentDesignerFlags());
+		OutWidget->SetDesignerFlags(BlueprintEditor.Pin()->GetCurrentDesignerFlags());
 	}
 }
 
diff --git a/Engine/Source/Editor/UMGEditor/Private/Hierarchy/SHierarchyViewItem.cpp b/Engine/Source/Editor/UMGEditor/Private/Hierarchy/SHierarchyViewItem.cpp
--- a/Engine/Source/Editor/UMGEditor/Private/Hierarchy/SHierarchyViewItem.cpp
+++ b/Engine/Source/Editor/UMGEditor/Private/Hierarchy/SHierarchyViewItem.cpp
@@ -21,8 +21,10 @@
 #include "Kismet2/BlueprintEditorUtils.h"
 
 #include "DragAndDrop/DecoratedDragDropOp.h"
-#include "WidgetTemplate.h"
+#include "DragAndDrop/AssetDragDropOp.h"
+#include "DragAndDrop/ClassDragDropOp.h"
 #include "DragDrop/WidgetTemplateDragDropOp.h"
+#include "WidgetTemplate.h"
 
 
 
@@ -33,7 +35,6 @@
 #include "WidgetBlueprintEditorUtils.h"
 #include "ScopedTransaction.h"
 #include "Styling/SlateIconFinder.h"
-#include "DragAndDrop/AssetDragDropOp.h"
 #include "WidgetTemplateBlueprintClass.h"
 #include "WidgetTemplateImageClass.h"
 
@@ -188,10 +189,15 @@ TOptional<EItemDropZone> ProcessHierarchyDragDrop(const FDragDropEvent& DragDrop
 	check( Blueprint != nullptr && Blueprint->WidgetTree != nullptr );
 
 	// Is this a drag/drop op to create a new widget in the tree?
-	TSharedPtr<FWidgetTemplateDragDropOp> TemplateDragDropOp = DragDropEvent.GetOperationAs<FWidgetTemplateDragDropOp>();
-	if ( TemplateDragDropOp.IsValid() )
+	TSharedPtr<FDragDropOperation> DragDropOp = DragDropEvent.GetOperation();
+	if (DragDropOp.IsValid() )
 	{
-		TemplateDragDropOp->ResetToDefaultToolTip();
+		TSharedPtr<FDecoratedDragDropOp> DecoratedDragDropOp = nullptr;
+		if (DragDropOp->IsOfType<FDecoratedDragDropOp>())
+		{
+			DecoratedDragDropOp = StaticCastSharedPtr<FDecoratedDragDropOp>(DragDropOp);
+			DecoratedDragDropOp->ResetToDefaultToolTip();
+		}
 
 		// Are we adding to the root?
 		if ( !TargetItem.IsValid() && Blueprint->WidgetTree->RootWidget == nullptr )
@@ -204,11 +210,14 @@ TOptional<EItemDropZone> ProcessHierarchyDragDrop(const FDragDropEvent& DragDrop
 				Blueprint->WidgetTree->SetFlags(RF_Transactional);
 				Blueprint->WidgetTree->Modify();
 
-				Blueprint->WidgetTree->RootWidget = TemplateDragDropOp->Template->Create(Blueprint->WidgetTree);
+				Blueprint->WidgetTree->RootWidget = FWidgetBlueprintEditorUtils::GetWidgetTemplateFromDragDrop(Blueprint, Blueprint->WidgetTree, DragDropOp);
 				FBlueprintEditorUtils::MarkBlueprintAsStructurallyModified(Blueprint);
 			}
 
-			TemplateDragDropOp->CurrentIconBrush = FEditorStyle::GetBrush(TEXT("Graph.ConnectorFeedback.OK"));
+			if (DecoratedDragDropOp.IsValid())
+			{
+				DecoratedDragDropOp->CurrentIconBrush = FEditorStyle::GetBrush(TEXT("Graph.ConnectorFeedback.OK"));
+			}
 			return EItemDropZone::OntoItem;
 		}
 		// Are we adding to a panel?
@@ -216,7 +225,10 @@ TOptional<EItemDropZone> ProcessHierarchyDragDrop(const FDragDropEvent& DragDrop
 		{
 			if (!Parent->CanAddMoreChildren())
 			{
-				TemplateDragDropOp->CurrentHoverText = LOCTEXT("NoAdditionalChildren", "Widget can't accept additional children.");
+				if (DecoratedDragDropOp.IsValid())
+				{
+					DecoratedDragDropOp->CurrentHoverText = LOCTEXT("NoAdditionalChildren", "Widget can't accept additional children.");
+				}
 			}
 			else
 			{
@@ -230,7 +242,7 @@ TOptional<EItemDropZone> ProcessHierarchyDragDrop(const FDragDropEvent& DragDrop
 					
 					Parent->SetFlags(RF_Transactional);
 					Parent->Modify();
-					UWidget* Widget = TemplateDragDropOp->Template->Create(Blueprint->WidgetTree);
+					UWidget* Widget = FWidgetBlueprintEditorUtils::GetWidgetTemplateFromDragDrop(Blueprint, Blueprint->WidgetTree, DragDropOp);
 
 					UPanelSlot* NewSlot = nullptr;
 					if (Index.IsSet())
@@ -246,16 +258,25 @@ TOptional<EItemDropZone> ProcessHierarchyDragDrop(const FDragDropEvent& DragDrop
 					FBlueprintEditorUtils::MarkBlueprintAsStructurallyModified(Blueprint);
 				}
 
-				TemplateDragDropOp->CurrentIconBrush = FEditorStyle::GetBrush(TEXT("Graph.ConnectorFeedback.OK"));
+				if (DecoratedDragDropOp.IsValid())
+				{
+					DecoratedDragDropOp->CurrentIconBrush = FEditorStyle::GetBrush(TEXT("Graph.ConnectorFeedback.OK"));
+				}
 				return EItemDropZone::OntoItem;
 			}
 		}
 		else
 		{
-			TemplateDragDropOp->CurrentHoverText = LOCTEXT("CantHaveChildren", "Widget can't have children.");
+			if (DecoratedDragDropOp.IsValid())
+			{
+				DecoratedDragDropOp->CurrentHoverText = LOCTEXT("CantHaveChildren", "Widget can't have children.");
+			}
 		}
 
-		TemplateDragDropOp->CurrentIconBrush = FEditorStyle::GetBrush(TEXT("Graph.ConnectorFeedback.Error"));
+		if (DecoratedDragDropOp.IsValid())
+		{
+			DecoratedDragDropOp->CurrentIconBrush = FEditorStyle::GetBrush(TEXT("Graph.ConnectorFeedback.Error"));
+		}
 		return TOptional<EItemDropZone>();
 	}
 
@@ -530,6 +551,28 @@ void FHierarchyModel::OnNameTextCommited(const FText& InText, ETextCommit::Type
 
 }
 
+bool FHierarchyModel::HasCircularReferences(UWidgetBlueprint* Blueprint, UWidget* Widget, TSharedPtr<FDragDropOperation>& DragDropOp)
+{
+	bool bHasCircularReferences = false;
+	if (Widget)
+	{
+		if (!Blueprint->IsWidgetFreeFromCircularReferences(Cast<UUserWidget>(Widget)))
+		{
+			if (DragDropOp->IsOfType<FDecoratedDragDropOp>())
+			{
+				TSharedPtr<FDecoratedDragDropOp> DecoratedDragDropOp = StaticCastSharedPtr<FDecoratedDragDropOp>(DragDropOp);
+				DecoratedDragDropOp->CurrentIconBrush = FEditorStyle::GetBrush(TEXT("Graph.ConnectorFeedback.Error"));
+				DecoratedDragDropOp->CurrentHoverText = LOCTEXT("CircularReference", "This would cause a circular reference.");
+			}
+
+			bHasCircularReferences = true;
+		}
+
+		RemovePreviewWidget(Blueprint, Widget);
+	}
+
+	return bHasCircularReferences;
+}
 
 void FHierarchyModel::DetermineDragDropPreviewWidgets(TArray<UWidget*>& OutWidgets, const FDragDropEvent& DragDropEvent)
 {
@@ -541,60 +584,18 @@ void FHierarchyModel::DetermineDragDropPreviewWidgets(TArray<UWidget*>& OutWidge
 		return;
 	}
 
-	TSharedPtr<FWidgetTemplateDragDropOp> TemplateDragDropOp = DragDropEvent.GetOperationAs<FWidgetTemplateDragDropOp>();
-	TSharedPtr<FAssetDragDropOp> AssetDragDropOp = DragDropEvent.GetOperationAs<FAssetDragDropOp>();
+	TSharedPtr<FDragDropOperation> DragDropOp = DragDropEvent.GetOperation();
+	UWidget* Widget = FWidgetBlueprintEditorUtils::GetWidgetTemplateFromDragDrop(Blueprint, Blueprint->WidgetTree, DragDropOp);
 
-	if (TemplateDragDropOp.IsValid())
+	if (Widget)
 	{
-		UWidget* Widget = TemplateDragDropOp->Template->Create(Blueprint->WidgetTree);
-
-		if (Widget)
-		{
-			if (Cast<UUserWidget>(Widget) == nullptr || Blueprint->IsWidgetFreeFromCircularReferences(Cast<UUserWidget>(Widget)))
-			{
-				OutWidgets.Add(Widget);
-			}
-		}
-	}
-	else if (AssetDragDropOp.IsValid())
-	{
-		for (const FAssetData& AssetData : AssetDragDropOp->GetAssets())
-		{
-			UWidget* Widget = nullptr;
-			UClass* AssetClass = FindObjectChecked<UClass>(ANY_PACKAGE, *AssetData.AssetClass.ToString());
-
-			if (FWidgetTemplateBlueprintClass::Supports(AssetClass))
-			{
-				// Allows a UMG Widget Blueprint to be dragged from the Content Browser to another Widget Blueprint...as long as we're not trying to place a
-				// blueprint inside itself.
-				FString BlueprintPath = Blueprint->GetPathName();
-				if (BlueprintPath != AssetData.ObjectPath.ToString())
-				{
-					Widget = FWidgetTemplateBlueprintClass(AssetData).Create(Blueprint->WidgetTree);
-
-					// Check to make sure that this widget can be added to the current blueprint
-					if (Cast<UUserWidget>(Widget) != nullptr && !Blueprint->IsWidgetFreeFromCircularReferences(Cast<UUserWidget>(Widget)))
-					{
-						Widget = nullptr;
-					}
-				}
-			}
-			else if (FWidgetTemplateImageClass::Supports(AssetClass))
-			{
-				Widget = FWidgetTemplateImageClass(AssetData).Create(Blueprint->WidgetTree);
-			}
-
-			if (Widget)
-			{
-				OutWidgets.Add(Widget);
-			}
-		}
+		OutWidgets.Add(Widget);
 	}
 
 	// Mark the widgets for design-time rendering
-	for (UWidget* Widget : OutWidgets)
+	for (UWidget* OutWidget : OutWidgets)
 	{
-		Widget->SetDesignerFlags(BlueprintEditor.Pin()->GetCurrentDesignerFlags());
+		OutWidget->SetDesignerFlags(BlueprintEditor.Pin()->GetCurrentDesignerFlags());
 	}
 }
 
@@ -745,25 +746,15 @@ bool FHierarchyRoot::DoesWidgetOverrideFlowDirection() const
 TOptional<EItemDropZone> FHierarchyRoot::HandleCanAcceptDrop(const FDragDropEvent& DragDropEvent, EItemDropZone DropZone)
 {
 	bool bIsFreeFromCircularReferences = true;
-	TSharedPtr<FWidgetTemplateDragDropOp> TemplateDragDropOp = DragDropEvent.GetOperationAs<FWidgetTemplateDragDropOp>();
-	if (TemplateDragDropOp.IsValid())
+
+	TSharedPtr<FDragDropOperation> DragDropOp = DragDropEvent.GetOperation();
+	if (DragDropOp.IsValid())
 	{
 		UWidgetBlueprint* Blueprint = BlueprintEditor.Pin()->GetWidgetBlueprintObj();
 		if(Blueprint)
 		{
-			UWidget* Widget = TemplateDragDropOp->Template->Create(Blueprint->WidgetTree);
-	
-			if (Widget)
-			{
-				if (!Blueprint->IsWidgetFreeFromCircularReferences(Cast<UUserWidget>(Widget)))
-				{
-					TemplateDragDropOp->CurrentIconBrush = FEditorStyle::GetBrush(TEXT("Graph.ConnectorFeedback.Error"));
-					TemplateDragDropOp->CurrentHoverText = LOCTEXT("CircularReference", "This would cause a circular reference.");
-					bIsFreeFromCircularReferences = false;
-				}
-
-				FHierarchyModel::RemovePreviewWidget(Blueprint, Widget);
-			}
+			UWidget* Widget = FWidgetBlueprintEditorUtils::GetWidgetTemplateFromDragDrop(Blueprint, Blueprint->WidgetTree, DragDropOp);
+			bIsFreeFromCircularReferences = !HasCircularReferences(Blueprint, Widget, DragDropOp);
 		}
 	}
 
@@ -895,36 +886,36 @@ TOptional<EItemDropZone> FNamedSlotModel::HandleCanAcceptDrop(const FDragDropEve
 {
 	UWidgetBlueprint* Blueprint = BlueprintEditor.Pin()->GetWidgetBlueprintObj();
 
-	TSharedPtr<FWidgetTemplateDragDropOp> TemplateDragDropOp = DragDropEvent.GetOperationAs<FWidgetTemplateDragDropOp>();
-	if (TemplateDragDropOp.IsValid())
+	TSharedPtr<FDragDropOperation> DragDropOp = DragDropEvent.GetOperation();
+	if (DragDropOp.IsValid())
 	{
-		TemplateDragDropOp->ResetToDefaultToolTip();
+		TSharedPtr<FDecoratedDragDropOp> DecoratedDragDropOp = nullptr;
+		if (DragDropOp->IsOfType<FDecoratedDragDropOp>())
+		{
+			DecoratedDragDropOp = StaticCastSharedPtr<FDecoratedDragDropOp>(DragDropOp);
+			DecoratedDragDropOp->ResetToDefaultToolTip();
+		}
 
-		if ( INamedSlotInterface* NamedSlotHost = Cast<INamedSlotInterface>(Item.GetTemplate()) )
+		if (INamedSlotInterface* NamedSlotHost = Cast<INamedSlotInterface>(Item.GetTemplate()))
 		{
 			// Only assign content to the named slot if it is null.
-			if ( NamedSlotHost->GetContentForSlot(SlotName) != nullptr )
+			if (NamedSlotHost->GetContentForSlot(SlotName) != nullptr)
 			{
-				TemplateDragDropOp->CurrentIconBrush = FEditorStyle::GetBrush(TEXT("Graph.ConnectorFeedback.Error"));
-				TemplateDragDropOp->CurrentHoverText = LOCTEXT("NamedSlotAlreadyFull", "Named Slot already has a child.");
+				if (DecoratedDragDropOp.IsValid())
+				{
+					DecoratedDragDropOp->CurrentIconBrush = FEditorStyle::GetBrush(TEXT("Graph.ConnectorFeedback.Error"));
+					DecoratedDragDropOp->CurrentHoverText = LOCTEXT("NamedSlotAlreadyFull", "Named Slot already has a child.");
+				}
 				return TOptional<EItemDropZone>();
 			}
 
-			UWidget* Widget = TemplateDragDropOp->Template->Create(Blueprint->WidgetTree);
-			bool bIsFreeFromCircularReferences = true;
-			if (Widget)
-			{
-				if (!Blueprint->IsWidgetFreeFromCircularReferences(Cast<UUserWidget>(Widget)))
-				{
-					TemplateDragDropOp->CurrentIconBrush = FEditorStyle::GetBrush(TEXT("Graph.ConnectorFeedback.Error"));
-					TemplateDragDropOp->CurrentHoverText = LOCTEXT("CircularReference", "This would cause a circular reference.");
-					bIsFreeFromCircularReferences = false;
-				}
+			UWidget* Widget = FWidgetBlueprintEditorUtils::GetWidgetTemplateFromDragDrop(Blueprint, Blueprint->WidgetTree, DragDropOp);
+			bool bIsFreeFromCircularReferences = !HasCircularReferences(Blueprint, Widget, DragDropOp);
 
-				FHierarchyModel::RemovePreviewWidget(Blueprint, Widget);
+			if (DecoratedDragDropOp.IsValid())
+			{
+				DecoratedDragDropOp->CurrentIconBrush = FEditorStyle::GetBrush(TEXT("Graph.ConnectorFeedback.OK"));
 			}
-
-			TemplateDragDropOp->CurrentIconBrush = FEditorStyle::GetBrush(TEXT("Graph.ConnectorFeedback.OK"));
 			return bIsFreeFromCircularReferences ? EItemDropZone::OntoItem : TOptional<EItemDropZone>();
 		}
 	}
@@ -983,8 +974,8 @@ FReply FNamedSlotModel::HandleAcceptDrop(FDragDropEvent const& DragDropEvent, EI
 		return FReply::Unhandled();
 	}
 
-	TSharedPtr<FWidgetTemplateDragDropOp> TemplateDragDropOp = DragDropEvent.GetOperationAs<FWidgetTemplateDragDropOp>();
-	if (TemplateDragDropOp.IsValid())
+	TSharedPtr<FDragDropOperation> DragDropOp = DragDropEvent.GetOperation();
+	if (DragDropOp.IsValid())
 	{
 		FScopedTransaction Transaction(LOCTEXT("AddWidgetFromTemplate", "Add Widget"));
 
@@ -992,7 +983,7 @@ FReply FNamedSlotModel::HandleAcceptDrop(FDragDropEvent const& DragDropEvent, EI
 		Blueprint->WidgetTree->SetFlags(RF_Transactional);
 		Blueprint->WidgetTree->Modify();
 
-		UWidget* DroppingWidget = TemplateDragDropOp->Template->Create(Blueprint->WidgetTree);
+		UWidget* DroppingWidget = FWidgetBlueprintEditorUtils::GetWidgetTemplateFromDragDrop(Blueprint, Blueprint->WidgetTree, DragDropOp);
 
 		DoDrop(SlotHostWidget, DroppingWidget);
 
@@ -1153,25 +1144,15 @@ FSlateFontInfo FHierarchyWidget::GetFont() const
 TOptional<EItemDropZone> FHierarchyWidget::HandleCanAcceptDrop(const FDragDropEvent& DragDropEvent, EItemDropZone DropZone)
 {
 	bool bIsFreeFromCircularReferences = true;
-	if (TSharedPtr<FWidgetTemplateDragDropOp> TemplateDragDropOp = DragDropEvent.GetOperationAs<FWidgetTemplateDragDropOp>())
+
+	TSharedPtr<FDragDropOperation> DragDropOp = DragDropEvent.GetOperation();
+	if (DragDropOp.IsValid())
 	{
 		UWidgetBlueprint* Blueprint = BlueprintEditor.Pin()->GetWidgetBlueprintObj();
-		if(Blueprint)
+		if (Blueprint)
 		{
-			UWidget* Widget = TemplateDragDropOp->Template->Create(Blueprint->WidgetTree);
-
-			if (Widget)
-			{
-				if (!Blueprint->IsWidgetFreeFromCircularReferences(Cast<UUserWidget>(Widget)))
-				{
-					TemplateDragDropOp->CurrentIconBrush = FEditorStyle::GetBrush(TEXT("Graph.ConnectorFeedback.Error"));
-					TemplateDragDropOp->CurrentHoverText = LOCTEXT("CircularReference", "This would cause a circular reference.");
-					bIsFreeFromCircularReferences = false;
-				}
-
-				RemovePreviewWidget(Blueprint, Widget);
-
-			}
+			UWidget* Widget = FWidgetBlueprintEditorUtils::GetWidgetTemplateFromDragDrop(Blueprint, Blueprint->WidgetTree, DragDropOp);
+			bIsFreeFromCircularReferences = !HasCircularReferences(Blueprint, Widget, DragDropOp);
 		}
 	}
 	bool bIsDrop = false;
diff --git a/Engine/Source/Editor/UMGEditor/Private/Hierarchy/SHierarchyViewItem.h b/Engine/Source/Editor/UMGEditor/Private/Hierarchy/SHierarchyViewItem.h
--- a/Engine/Source/Editor/UMGEditor/Private/Hierarchy/SHierarchyViewItem.h
+++ b/Engine/Source/Editor/UMGEditor/Private/Hierarchy/SHierarchyViewItem.h
@@ -106,6 +106,7 @@ protected:
 	virtual void GetChildren(TArray< TSharedPtr<FHierarchyModel> >& Children) = 0;
 	virtual void UpdateSelection() = 0;
 	virtual FWidgetReference AsDraggedWidgetReference() const { return FWidgetReference(); }
+	virtual bool HasCircularReferences(class UWidgetBlueprint* Blueprint, class UWidget* Widget, TSharedPtr<class FDragDropOperation>& DragDropOp);
 	void DetermineDragDropPreviewWidgets(TArray<class UWidget*>& OutWidgets, const FDragDropEvent& DragDropEvent);
 	void RemovePreviewWidget(class UWidgetBlueprint* Blueprint, class UWidget* Widget);
 
diff --git a/Engine/Source/Editor/UMGEditor/Private/Library/SLibraryView.cpp b/Engine/Source/Editor/UMGEditor/Private/Library/SLibraryView.cpp
new file mode 100644
--- /dev/null
+++ b/Engine/Source/Editor/UMGEditor/Private/Library/SLibraryView.cpp
@@ -0,0 +1,338 @@
+// Copyright Epic Games, Inc. All Rights Reserved.
+
+#include "Library/SLibraryView.h"
+#include "Library/SLibraryViewModel.h"
+#include "Misc/ConfigCacheIni.h"
+#include "UObject/UObjectHash.h"
+#include "UObject/UObjectIterator.h"
+#include "WidgetBlueprint.h"
+#include "Editor.h"
+#include "Widgets/Images/SImage.h"
+#include "Widgets/Layout/SScrollBorder.h"
+
+#if WITH_EDITOR
+	#include "EditorStyleSet.h"
+#endif // WITH_EDITOR
+
+
+
+#include "DragDrop/WidgetTemplateDragDropOp.h"
+
+#include "Templates/WidgetTemplateClass.h"
+#include "Templates/WidgetTemplateBlueprintClass.h"
+
+#include "AssetRegistryModule.h"
+#include "Widgets/Input/SSearchBox.h"
+#include "Widgets/Input/SCheckBox.h"
+
+#include "Settings/ContentBrowserSettings.h"
+#include "WidgetBlueprintEditorUtils.h"
+
+
+
+#include "UMGEditorProjectSettings.h"
+
+#define LOCTEXT_NAMESPACE "UMG"
+
+FText SLibraryViewItem::GetFavoriteToggleToolTipText() const
+{
+	if (GetFavoritedState() == ECheckBoxState::Checked)
+	{
+		return LOCTEXT("Unfavorite", "Click to remove this widget from your favorites.");
+	}
+	return LOCTEXT("Favorite", "Click to add this widget to your favorites.");
+}
+
+ECheckBoxState SLibraryViewItem::GetFavoritedState() const
+{
+	if (WidgetViewModel->IsFavorite())
+	{
+		return ECheckBoxState::Checked;
+	}
+	else
+	{
+		return ECheckBoxState::Unchecked;
+	}
+}
+
+void SLibraryViewItem::OnFavoriteToggled(ECheckBoxState InNewState)
+{
+	if (InNewState == ECheckBoxState::Checked)
+	{
+		//Add to favorites
+		WidgetViewModel->AddToFavorites();
+	}
+	else
+	{
+		//Remove from favorites
+		WidgetViewModel->RemoveFromFavorites();
+	}
+}
+
+EVisibility SLibraryViewItem::GetFavoritedStateVisibility() const
+{
+	return GetFavoritedState() == ECheckBoxState::Checked || IsHovered() ? EVisibility::Visible : EVisibility::Hidden;
+}
+
+void SLibraryViewItem::Construct(const FArguments& InArgs, TSharedPtr<FWidgetTemplateViewModel> InWidgetViewModel)
+{
+	WidgetViewModel = InWidgetViewModel;
+
+	ChildSlot
+		[
+			SNew(SHorizontalBox)
+			.ToolTip(WidgetViewModel->Template->GetToolTip())
+			+ SHorizontalBox::Slot()
+			.AutoWidth()
+			.VAlign(VAlign_Center)
+			[	
+				SNew(SCheckBox)
+				.ToolTipText(this, &SLibraryViewItem::GetFavoriteToggleToolTipText)
+				.IsChecked(this, &SLibraryViewItem::GetFavoritedState)
+				.OnCheckStateChanged(this, &SLibraryViewItem::OnFavoriteToggled)
+				.Style(FEditorStyle::Get(), "UMGEditor.Library.FavoriteToggleStyle")
+				.Visibility(this, &SLibraryViewItem::GetFavoritedStateVisibility)
+			]
+			+ SHorizontalBox::Slot()
+			.AutoWidth()
+			.VAlign(VAlign_Center)
+			[
+				SNew(SImage)
+				.ColorAndOpacity(FSlateColor::UseForeground())
+				.Image(WidgetViewModel->Template->GetIcon())
+			]
+
+			+ SHorizontalBox::Slot()
+			.AutoWidth()
+			.Padding(2, 0, 0, 0)
+			.VAlign(VAlign_Center)
+			[
+				SNew(STextBlock)
+				.Text(InWidgetViewModel->GetName())
+				.HighlightText(InArgs._HighlightText)
+			]
+		];
+}
+
+FReply SLibraryViewItem::OnMouseButtonDoubleClick(const FGeometry& InMyGeometry, const FPointerEvent& InMouseEvent)
+{
+	return WidgetViewModel->Template->OnDoubleClicked();
+};
+
+void SLibraryView::Construct(const FArguments& InArgs, TSharedPtr<FWidgetBlueprintEditor> InBlueprintEditor)
+{
+	BlueprintEditor = InBlueprintEditor;
+
+	UBlueprint* BP = InBlueprintEditor->GetBlueprintObj();
+	LibraryViewModel = InBlueprintEditor->GetLibraryViewModel();
+
+	// Register to the update of the viewmodel.
+	LibraryViewModel->OnUpdating.AddRaw(this, &SLibraryView::OnViewModelUpdating);
+	LibraryViewModel->OnUpdated.AddRaw(this, &SLibraryView::OnViewModelUpdated);
+
+	WidgetFilter = MakeShareable(new WidgetViewModelTextFilter(
+		WidgetViewModelTextFilter::FItemToStringArray::CreateSP(this, &SLibraryView::GetWidgetFilterStrings)));
+
+	FilterHandler = MakeShareable(new LibraryFilterHandler());
+	FilterHandler->SetFilter(WidgetFilter.Get());
+	FilterHandler->SetRootItems(&(LibraryViewModel->GetWidgetViewModels()), &TreeWidgetViewModels);
+	FilterHandler->SetGetChildrenDelegate(LibraryFilterHandler::FOnGetChildren::CreateRaw(this, &SLibraryView::OnGetChildren));
+
+	SAssignNew(WidgetTemplatesView, STreeView< TSharedPtr<FWidgetViewModel> >)
+		.ItemHeight(1.0f)
+		.SelectionMode(ESelectionMode::Single)
+		.OnGenerateRow(this, &SLibraryView::OnGenerateWidgetTemplateItem)
+		.OnGetChildren(FilterHandler.ToSharedRef(), &LibraryFilterHandler::OnGetFilteredChildren)
+		.OnSelectionChanged(this, &SLibraryView::WidgetLibrary_OnSelectionChanged)
+		.TreeItemsSource(&TreeWidgetViewModels);
+		
+
+	FilterHandler->SetTreeView(WidgetTemplatesView.Get());
+
+	ChildSlot
+	[
+		SNew(SVerticalBox)
+
+		+ SVerticalBox::Slot()
+		.Padding(4)
+		.AutoHeight()
+		[
+			SAssignNew(SearchBoxPtr, SSearchBox)
+			.HintText(LOCTEXT("SearchTemplates", "Search Library"))
+			.OnTextChanged(this, &SLibraryView::OnSearchChanged)
+		]
+
+		+ SVerticalBox::Slot()
+		.FillHeight(1.0f)
+		[
+			SNew(SScrollBorder, WidgetTemplatesView.ToSharedRef())
+			[
+				SNew(SBorder)
+				.BorderImage(FAppStyle::Get().GetBrush("Brushes.Recessed"))
+				.Padding(0)
+				[
+					WidgetTemplatesView.ToSharedRef()
+				]
+			]
+		]
+	];
+
+	bRefreshRequested = true;
+
+	LibraryViewModel->Update();
+	LoadItemExpansion();
+}
+
+SLibraryView::~SLibraryView()
+{
+	// Unregister to the update of the viewmodel.
+	LibraryViewModel->OnUpdating.RemoveAll(this);
+	LibraryViewModel->OnUpdated.RemoveAll(this);
+
+	// If the filter is enabled, disable it before saving the expanded items since
+	// filtering expands all items by default.
+	if (FilterHandler->GetIsEnabled())
+	{
+		FilterHandler->SetIsEnabled(false);
+		FilterHandler->RefreshAndFilterTree();
+	}
+
+	SaveItemExpansion();
+}
+
+void SLibraryView::OnSearchChanged(const FText& InFilterText)
+{
+	bRefreshRequested = true;
+	FilterHandler->SetIsEnabled(!InFilterText.IsEmpty());
+	WidgetFilter->SetRawFilterText(InFilterText);
+	SearchBoxPtr->SetError(WidgetFilter->GetFilterErrorText());
+	LibraryViewModel->SetSearchText(InFilterText);
+}
+
+void SLibraryView::WidgetLibrary_OnSelectionChanged(TSharedPtr<FWidgetViewModel> SelectedItem, ESelectInfo::Type SelectInfo)
+{
+	if (!SelectedItem.IsValid()) 
+	{
+		return;
+	}
+
+	// Reset the selected
+	BlueprintEditor.Pin()->SetSelectedTemplate(nullptr);
+	BlueprintEditor.Pin()->SetSelectedUserWidget(FAssetData());
+
+	// If it's not a template, return
+	if (!SelectedItem->IsTemplate())
+	{
+		return;
+	}
+
+	TSharedPtr<FWidgetTemplateViewModel> SelectedTemplate = StaticCastSharedPtr<FWidgetTemplateViewModel>(SelectedItem);
+	if (SelectedTemplate.IsValid())
+	{
+		TSharedPtr<FWidgetTemplateClass> TemplateClass = StaticCastSharedPtr<FWidgetTemplateClass>(SelectedTemplate->Template);
+		if (TemplateClass.IsValid())
+		{
+			if (TemplateClass->GetWidgetClass().IsValid())
+			{
+				BlueprintEditor.Pin()->SetSelectedTemplate(TemplateClass->GetWidgetClass());
+			}
+			else
+			{
+				TSharedPtr<FWidgetTemplateBlueprintClass> UserCreatedTemplate = StaticCastSharedPtr<FWidgetTemplateBlueprintClass>(TemplateClass);
+				if (UserCreatedTemplate.IsValid())
+				{
+					// Then pass in the asset data of selected widget
+					FAssetData UserCreatedWidget = UserCreatedTemplate->GetWidgetAssetData();
+					BlueprintEditor.Pin()->SetSelectedUserWidget(UserCreatedWidget);
+				}
+			}
+		}
+	}
+}
+
+TSharedPtr<FWidgetTemplate> SLibraryView::GetSelectedTemplateWidget() const
+{
+	TArray<TSharedPtr<FWidgetViewModel>> SelectedTemplates = WidgetTemplatesView.Get()->GetSelectedItems();
+	if (SelectedTemplates.Num() == 1)
+	{
+		TSharedPtr<FWidgetTemplateViewModel> TemplateViewModel = StaticCastSharedPtr<FWidgetTemplateViewModel>(SelectedTemplates[0]);
+		if (TemplateViewModel.IsValid())
+		{
+			return TemplateViewModel->Template;
+		}
+	}
+	return nullptr;
+}
+
+void SLibraryView::LoadItemExpansion()
+{
+	// Restore the expansion state of the widget groups.
+	for ( TSharedPtr<FWidgetViewModel>& ViewModel : LibraryViewModel->GetWidgetViewModels())
+	{
+		bool IsExpanded;
+		if ( GConfig->GetBool(TEXT("WidgetTemplatesExpanded"), *ViewModel->GetName().ToString(), IsExpanded, GEditorPerProjectIni) && IsExpanded )
+		{
+			WidgetTemplatesView->SetItemExpansion(ViewModel, true);
+		}
+	}
+}
+
+void SLibraryView::SaveItemExpansion()
+{
+	// Restore the expansion state of the widget groups.
+	for ( TSharedPtr<FWidgetViewModel>& ViewModel : LibraryViewModel->GetWidgetViewModels() )
+	{
+		const bool IsExpanded = WidgetTemplatesView->IsItemExpanded(ViewModel);
+		GConfig->SetBool(TEXT("WidgetTemplatesExpanded"), *ViewModel->GetName().ToString(), IsExpanded, GEditorPerProjectIni);
+	}
+}
+
+void SLibraryView::OnGetChildren(TSharedPtr<FWidgetViewModel> Item, TArray< TSharedPtr<FWidgetViewModel> >& Children)
+{
+	return Item->GetChildren(Children);
+}
+
+TSharedRef<ITableRow> SLibraryView::OnGenerateWidgetTemplateItem(TSharedPtr<FWidgetViewModel> Item, const TSharedRef<STableViewBase>& OwnerTable)
+{
+	return Item->BuildRow(OwnerTable);
+}
+
+void SLibraryView::OnViewModelUpdating()
+{
+	// Save the old expanded items temporarily
+	WidgetTemplatesView->GetExpandedItems(ExpandedItems);
+}
+
+void SLibraryView::OnViewModelUpdated()
+{
+	bRefreshRequested = true;
+
+	// Restore the expansion state
+	for (TSharedPtr<FWidgetViewModel>& ExpandedItem : ExpandedItems)
+	{
+		for (TSharedPtr<FWidgetViewModel>& ViewModel : LibraryViewModel->GetWidgetViewModels())
+		{
+			if (ViewModel->GetName().EqualTo(ExpandedItem->GetName()) || ViewModel->ShouldForceExpansion())
+			{
+				WidgetTemplatesView->SetItemExpansion(ViewModel, true);
+			}
+		}
+	}
+	ExpandedItems.Reset();
+}
+
+void SLibraryView::Tick(const FGeometry& AllottedGeometry, const double InCurrentTime, const float InDeltaTime)
+{
+	if (bRefreshRequested)
+	{
+		bRefreshRequested = false;
+		FilterHandler->RefreshAndFilterTree();
+	}
+}
+
+void SLibraryView::GetWidgetFilterStrings(TSharedPtr<FWidgetViewModel> WidgetViewModel, TArray<FString>& OutStrings)
+{
+	WidgetViewModel->GetFilterStrings(OutStrings);
+}
+
+#undef LOCTEXT_NAMESPACE
diff --git a/Engine/Source/Editor/UMGEditor/Private/Library/SLibraryView.h b/Engine/Source/Editor/UMGEditor/Private/Library/SLibraryView.h
new file mode 100644
--- /dev/null
+++ b/Engine/Source/Editor/UMGEditor/Private/Library/SLibraryView.h
@@ -0,0 +1,143 @@
+// Copyright Epic Games, Inc. All Rights Reserved.
+
+#pragma once
+
+#include "CoreMinimal.h"
+#include "Widgets/DeclarativeSyntaxSupport.h"
+#include "Widgets/SCompoundWidget.h"
+#include "WidgetBlueprintEditor.h"
+#include "Misc/TextFilter.h"
+#include "Widgets/Views/STreeView.h"
+#include "Framework/Views/TreeFilterHandler.h"
+
+class FWidgetTemplate;
+class UWidgetBlueprint;
+class SLibraryView;
+class FLibraryViewModel;
+class FWidgetViewModel;
+class FWidgetTemplateViewModel;
+
+/** Widget used to show a single row of the Library and Library favorite panel. */
+class SLibraryViewItem : public SCompoundWidget
+{
+public:
+
+	SLATE_BEGIN_ARGS(SLibraryViewItem) {}
+
+	/** The current text to highlight */
+	SLATE_ATTRIBUTE(FText, HighlightText)
+
+	SLATE_END_ARGS()
+
+	/**
+	* Constructs this widget
+	*
+	* @param InArgs    Declaration from which to construct the widget
+	*/
+	void Construct(const FArguments& InArgs, TSharedPtr<FWidgetTemplateViewModel> InWidgetViewModel);
+
+private:
+	/**
+	 * Retrieves tooltip that describes the current favorited state 
+	 *
+	 * @return Text describing what this toggle will do when you click on it.
+	 */
+	FText GetFavoriteToggleToolTipText() const;
+
+	/**
+	 * Checks on the associated action's favorite state, and returns a
+	 * corresponding checkbox state to match.
+	 *
+	 * @return ECheckBoxState::Checked if the associated action is already favorited, ECheckBoxState::Unchecked if not.
+	 */
+	ECheckBoxState GetFavoritedState() const;
+
+	/**
+	 * Triggers when the user clicks this toggle, adds or removes the associated
+	 * action to the user's favorites.
+	 *
+	 * @param  InNewState	The new state that the user set the checkbox to.
+	 */
+	void OnFavoriteToggled(ECheckBoxState InNewState);
+
+	EVisibility GetFavoritedStateVisibility() const;
+
+	virtual FReply OnMouseButtonDoubleClick(const FGeometry& InMyGeometry, const FPointerEvent& InMouseEvent) override;
+
+private:
+	TSharedPtr<FWidgetTemplateViewModel> WidgetViewModel;
+};
+
+/** */
+class SLibraryView : public SCompoundWidget
+{
+public:
+	typedef TTextFilter<TSharedPtr<FWidgetViewModel>> WidgetViewModelTextFilter;
+
+public:
+	SLATE_BEGIN_ARGS( SLibraryView ){}
+	SLATE_END_ARGS()
+
+	void Construct(const FArguments& InArgs, TSharedPtr<FWidgetBlueprintEditor> InBlueprintEditor);
+	virtual ~SLibraryView();
+	
+	virtual void Tick( const FGeometry& AllottedGeometry, const double InCurrentTime, const float InDeltaTime ) override;
+
+	/** Gets the text currently displayed in the search box. */
+	FText GetSearchText() const;
+
+	/** On Selected template widget changed in Library view */
+	void WidgetLibrary_OnSelectionChanged(TSharedPtr<FWidgetViewModel> SelectedItem, ESelectInfo::Type SelectInfo);
+
+	/** Gets the selected template widget in Library */
+	TSharedPtr<FWidgetTemplate> GetSelectedTemplateWidget() const;
+
+private:
+
+	void OnGetChildren(TSharedPtr<FWidgetViewModel> Item, TArray< TSharedPtr<FWidgetViewModel> >& Children);
+	TSharedRef<ITableRow> OnGenerateWidgetTemplateItem(TSharedPtr<FWidgetViewModel> Item, const TSharedRef<STableViewBase>& OwnerTable);
+
+	/** Called when the filter text is changed. */
+	void OnSearchChanged(const FText& InFilterText);
+
+	void OnViewModelUpdating();
+	void OnViewModelUpdated();
+
+private:
+	/** Load the expansion state for the TreeView */
+	void LoadItemExpansion();
+
+	/** Save the expansion state for the TreeView */
+	void SaveItemExpansion();
+
+	/** Gets an array of strings used for filtering/searching the specified widget. */
+	void GetWidgetFilterStrings(TSharedPtr<FWidgetViewModel> WidgetViewModel, TArray<FString>& OutStrings);
+
+	TWeakPtr<class FWidgetBlueprintEditor> BlueprintEditor;
+	TSharedPtr<FLibraryViewModel> LibraryViewModel;
+
+	/** Handles filtering the Library based on an IFilter. */
+	typedef TreeFilterHandler<TSharedPtr<FWidgetViewModel>> LibraryFilterHandler;
+	TSharedPtr<LibraryFilterHandler> FilterHandler;
+
+	typedef TArray<TSharedPtr<FWidgetViewModel>> ViewModelsArray;
+	/** The root view models which are actually displayed by the TreeView which will be managed by the TreeFilterHandler. */
+	ViewModelsArray TreeWidgetViewModels;
+
+	TSharedPtr< STreeView< TSharedPtr<FWidgetViewModel> > > WidgetTemplatesView;
+
+	/** The search box used to update the filter text */
+	TSharedPtr<class SSearchBox> SearchBoxPtr;
+
+	/** The filter instance which is used by the TreeFilterHandler to filter the TreeView. */
+	TSharedPtr<WidgetViewModelTextFilter> WidgetFilter;
+
+	/** Expended Items in the Tree view */
+	TSet<TSharedPtr<FWidgetViewModel>> ExpandedItems;
+
+	/** Set to true to force a refresh of the treeview */
+	bool bRefreshRequested;
+
+	/** Are editor widgets supported. */
+	bool bAllowEditorWidget;
+};
diff --git a/Engine/Source/Editor/UMGEditor/Private/Library/SLibraryViewModel.cpp b/Engine/Source/Editor/UMGEditor/Private/Library/SLibraryViewModel.cpp
new file mode 100644
--- /dev/null
+++ b/Engine/Source/Editor/UMGEditor/Private/Library/SLibraryViewModel.cpp
@@ -0,0 +1,509 @@
+// Copyright Epic Games, Inc. All Rights Reserved.
+
+#include "Library/SLibraryViewModel.h"
+#include "Library/SLibraryView.h"
+#include "Widgets/Views/STableViewBase.h"
+#include "Widgets/Views/STableRow.h"
+#include "WidgetBlueprint.h"
+#include "Editor.h"
+#include "FrontendFilters.h"
+#include "SAssetView.h"
+
+#if WITH_EDITOR
+	#include "EditorStyleSet.h"
+#endif // WITH_EDITOR
+
+#include "DragDrop/WidgetTemplateDragDropOp.h"
+
+#include "Templates/WidgetTemplateClass.h"
+#include "Templates/WidgetTemplateBlueprintClass.h"
+
+#include "AssetRegistryModule.h"
+#include "WidgetBlueprintEditorUtils.h"
+
+#include "Settings/ContentBrowserSettings.h"
+#include "Settings/WidgetDesignerSettings.h"
+#include "UMGEditorProjectSettings.h"
+#include "WidgetPaletteFavorites.h"
+
+#define LOCTEXT_NAMESPACE "UMG"
+
+FLibraryViewModel::FLibraryViewModel(TSharedPtr<FWidgetBlueprintEditor> InBlueprintEditor)
+	: bRebuildRequested(true)
+{
+	BlueprintEditor = InBlueprintEditor;
+
+	FavoriteHeader = MakeShareable(new FWidgetHeaderViewModel());
+	FavoriteHeader->GroupName = LOCTEXT("Favorites", "Favorites");
+}
+
+void FLibraryViewModel::RegisterToEvents()
+{
+	// Register for events that can trigger a Library rebuild
+	GEditor->OnBlueprintReinstanced().AddRaw(this, &FLibraryViewModel::OnBlueprintReinstanced);
+	FEditorDelegates::OnAssetsDeleted.AddSP(this, &FLibraryViewModel::HandleOnAssetsDeleted);
+	FCoreUObjectDelegates::ReloadCompleteDelegate.AddSP(this, &FLibraryViewModel::OnReloadComplete);
+
+	// register for any objects replaced
+	FCoreUObjectDelegates::OnObjectsReplaced.AddRaw(this, &FLibraryViewModel::OnObjectsReplaced);
+
+	// Register for favorite list update to handle the case where a favorite is added in another window of the UMG Designer
+	UWidgetPaletteFavorites* Favorites = GetDefault<UWidgetDesignerSettings>()->Favorites;
+	Favorites->OnFavoritesUpdated.AddSP(this, &FLibraryViewModel::OnFavoritesUpdated);
+}
+
+FLibraryViewModel::~FLibraryViewModel()
+{
+	GEditor->OnBlueprintReinstanced().RemoveAll(this);
+	FEditorDelegates::OnAssetsDeleted.RemoveAll(this);
+	FCoreUObjectDelegates::ReloadCompleteDelegate.RemoveAll(this);
+	FCoreUObjectDelegates::OnObjectsReplaced.RemoveAll(this);
+
+	UWidgetPaletteFavorites* Favorites = GetDefault<UWidgetDesignerSettings>()->Favorites;
+	Favorites->OnFavoritesUpdated.RemoveAll(this);
+}
+
+void FLibraryViewModel::AddToFavorites(const FWidgetTemplateViewModel* WidgetTemplateViewModel)
+{
+	UWidgetPaletteFavorites* Favorites = GetDefault<UWidgetDesignerSettings>()->Favorites;
+	Favorites->Add(WidgetTemplateViewModel->GetName().ToString());
+}
+
+void FLibraryViewModel::RemoveFromFavorites(const FWidgetTemplateViewModel* WidgetTemplateViewModel)
+{
+	UWidgetPaletteFavorites* Favorites = GetDefault<UWidgetDesignerSettings>()->Favorites;
+	Favorites->Remove(WidgetTemplateViewModel->GetName().ToString());
+}
+
+void FLibraryViewModel::Update()
+{
+	if (bRebuildRequested)
+	{
+		OnUpdating.Broadcast();
+		BuildWidgetList();
+		bRebuildRequested = false;
+		OnUpdated.Broadcast();
+	}
+}
+
+
+UWidgetBlueprint* FLibraryViewModel::GetBlueprint() const
+{
+	if (BlueprintEditor.IsValid())
+	{
+		UBlueprint* BP = BlueprintEditor.Pin()->GetBlueprintObj();
+		return Cast<UWidgetBlueprint>(BP);
+	}
+
+	return NULL;
+}
+
+void FLibraryViewModel::BuildWidgetList()
+{
+	// Clear the current list of view models and categories
+	WidgetViewModels.Reset();
+	WidgetTemplateCategories.Reset();
+
+	// Generate a list of templates
+	BuildClassWidgetList();
+
+	// Clear the Favorite section
+	bool bHasFavorites = FavoriteHeader->Children.Num() != 0;
+	FavoriteHeader->Children.Reset();
+	
+	// Copy of the list of favorites to be able to do some cleanup in the real list
+	UWidgetPaletteFavorites* FavoritesLibrary = GetDefault<UWidgetDesignerSettings>()->Favorites;
+	TArray<FString> FavoritesList = FavoritesLibrary->GetFavorites();
+
+	// For each entry in the category create a view model for the widget template
+	for ( auto& Entry : WidgetTemplateCategories )
+	{
+		TSharedPtr<FWidgetHeaderViewModel> Header = MakeShareable(new FWidgetHeaderViewModel());
+		Header->GroupName = FText::FromString(Entry.Key);
+
+
+		TSharedPtr<FWidgetTemplateListViewModel> TemplateViewModel = MakeShareable(new FWidgetTemplateListViewModel());
+		TemplateViewModel->Templates = Entry.Value;
+		Header->Children.Add(TemplateViewModel);
+
+		//for ( auto& Template : Entry.Value )
+		//{
+		//	TSharedPtr<FWidgetTemplateViewModel> TemplateViewModel = MakeShareable(new FWidgetTemplateViewModel());
+		//	TemplateViewModel->Template = Template;
+		//	TemplateViewModel->FavortiesViewModel = this;
+		//	Header->Children.Add(TemplateViewModel);
+
+		//	// If it's a favorite, we also add it to the Favorite section
+		//	int32 index = FavoritesList.Find(Template->Name.ToString());
+		//	if (index != INDEX_NONE)
+		//	{
+		//		TemplateViewModel->SetFavorite();
+
+		//		// We have to create a second copy of the ViewModel for the treeview has it doesn't support to have the same element twice.
+		//		TSharedPtr<FWidgetTemplateViewModel> FavoriteTemplateViewModel = MakeShareable(new FWidgetTemplateViewModel());
+		//		FavoriteTemplateViewModel->Template = Template;
+		//		FavoriteTemplateViewModel->FavortiesViewModel = this;
+		//		FavoriteTemplateViewModel->SetFavorite();
+
+		//		FavoriteHeader->Children.Add(FavoriteTemplateViewModel);
+
+		//		// Remove the favorite from the temporary list
+		//		FavoritesList.RemoveAt(index);
+		//	}
+
+		//}
+
+		Header->Children.Sort([] (TSharedPtr<FWidgetViewModel> L, TSharedPtr<FWidgetViewModel> R) { return R->GetName().CompareTo(L->GetName()) > 0; });
+
+		WidgetViewModels.Add(Header);
+	}	
+
+	// Remove all Favorites that may be left in the list.Typically happening when the list of favorite contains widget that were deleted since the last opening.
+	for (const FString& favoriteName : FavoritesList)
+	{
+		FavoritesLibrary->Remove(favoriteName);
+	}
+
+	// Sort the view models by name
+	WidgetViewModels.Sort([] (TSharedPtr<FWidgetViewModel> L, TSharedPtr<FWidgetViewModel> R) { return R->GetName().CompareTo(L->GetName()) > 0; });
+
+	// Add the Favorite section at the top
+	if (FavoriteHeader->Children.Num() != 0)
+	{
+		// We force expansion of the favorite header when we add favorites for the first time.
+		FavoriteHeader->SetForceExpansion(!bHasFavorites);
+		FavoriteHeader->Children.Sort([](TSharedPtr<FWidgetViewModel> L, TSharedPtr<FWidgetViewModel> R) { return R->GetName().CompareTo(L->GetName()) > 0; });
+		WidgetViewModels.Insert(FavoriteHeader, 0);
+	}
+	
+	// Take the Advanced Section, and put it at the end.
+	TSharedPtr<FWidgetViewModel>* advancedSectionPtr = WidgetViewModels.FindByPredicate([](TSharedPtr<FWidgetViewModel> widget) {return widget->GetName().CompareTo(LOCTEXT("Advanced", "Advanced")) == 0; });
+	if (advancedSectionPtr)
+	{
+		TSharedPtr<FWidgetViewModel> advancedSection = *advancedSectionPtr;
+		WidgetViewModels.Remove(advancedSection);
+		WidgetViewModels.Push(advancedSection);
+	}
+}
+
+void FLibraryViewModel::BuildClassWidgetList()
+{
+	static const FName DevelopmentStatusKey(TEXT("DevelopmentStatus"));
+
+	TMap<FName, TSubclassOf<UUserWidget>> LoadedWidgetBlueprintClassesByName;
+
+	auto ActiveWidgetBlueprintClass = GetBlueprint()->GeneratedClass;
+	FName ActiveWidgetBlueprintClassName = ActiveWidgetBlueprintClass->GetFName();
+
+	TArray<FSoftClassPath> WidgetClassesToHide = GetDefault<UUMGEditorProjectSettings>()->WidgetClassesToHide;
+
+	// Locate all UWidget classes from code and loaded widget BPs
+	for (TObjectIterator<UClass> ClassIt; ClassIt; ++ClassIt)
+	{
+		UClass* WidgetClass = *ClassIt;
+
+		if (!FWidgetBlueprintEditorUtils::IsUsableWidgetClass(WidgetClass))
+		{
+			continue;
+		}
+
+		if (WidgetClass->HasAnyClassFlags(CLASS_HideDropDown))
+		{
+			continue;
+		}
+
+		// Initialize AssetData for checking PackagePath
+		FAssetData WidgetAssetData = FAssetData(WidgetClass);
+
+		// Excludes engine content if user sets it to false
+		if (!GetDefault<UContentBrowserSettings>()->GetDisplayEngineFolder() || !GetDefault<UUMGEditorProjectSettings>()->bShowWidgetsFromEngineContent)
+		{
+			if (WidgetAssetData.PackagePath.ToString().Find(TEXT("/Engine")) == 0)
+			{
+				continue;
+			}
+		}
+
+		// Excludes developer content if user sets it to false
+		if (!GetDefault<UContentBrowserSettings>()->GetDisplayDevelopersFolder() || !GetDefault<UUMGEditorProjectSettings>()->bShowWidgetsFromDeveloperContent)
+		{
+			if (WidgetAssetData.PackagePath.ToString().Find(TEXT("/Game/Developers")) == 0)
+			{
+				continue;
+			}
+		}
+
+		// Excludes this widget if it is on the hide list
+		bool bIsOnList = false;
+		for (FSoftClassPath Widget : WidgetClassesToHide)
+		{
+			if (WidgetAssetData.ObjectPath.ToString().Find(Widget.ToString()) == 0)
+			{
+				bIsOnList = true;
+				break;
+			}
+		}
+		if (bIsOnList)
+		{
+			continue;
+		}
+
+		const bool bIsSameClass = WidgetClass->GetFName() == ActiveWidgetBlueprintClassName;
+
+		// Check that the asset that generated this class is valid (necessary b/c of a larger issue wherein force delete does not wipe the generated class object)
+		if ( bIsSameClass )
+		{
+			continue;
+		}
+
+		if (WidgetClass->IsChildOf(UUserWidget::StaticClass()))
+		{
+			if ( WidgetClass->ClassGeneratedBy )
+			{
+				// Track the widget blueprint classes that are already loaded
+				LoadedWidgetBlueprintClassesByName.Add(WidgetClass->ClassGeneratedBy->GetFName()) = WidgetClass;
+			}
+		}
+		else
+		{
+			TSharedPtr<FWidgetTemplateClass> Template = MakeShareable(new FWidgetTemplateClass(WidgetClass));
+
+			AddWidgetTemplate(Template);
+		}
+
+		//TODO UMG does not prevent deep nested circular references
+	}
+
+	// Locate all widget BP assets (include unloaded)
+	const FAssetRegistryModule& AssetRegistryModule = FModuleManager::LoadModuleChecked<FAssetRegistryModule>(TEXT("AssetRegistry"));
+	TArray<FAssetData> AllBPsAssetData;
+	AssetRegistryModule.Get().GetAssetsByClass(UBlueprint::StaticClass()->GetFName(), AllBPsAssetData, true);
+
+	for (FAssetData& BPAssetData : AllBPsAssetData)
+	{
+		// Blueprints get the class type actions for their parent native class - this avoids us having to load the blueprint
+		UClass* ParentClass = nullptr;
+		FString ParentClassName;
+		if (!BPAssetData.GetTagValue(FBlueprintTags::NativeParentClassPath, ParentClassName))
+		{
+			BPAssetData.GetTagValue(FBlueprintTags::ParentClassPath, ParentClassName);
+		}
+		if (!ParentClassName.IsEmpty())
+		{
+			UObject* Outer = nullptr;
+			ResolveName(Outer, ParentClassName, false, false);
+			ParentClass = FindObject<UClass>(ANY_PACKAGE, *ParentClassName);
+			// UUserWidgets have their own loading section, and we don't want to process any blueprints that don't have UWidget parents
+			if (!ParentClass->IsChildOf(UWidget::StaticClass()) || ParentClass->IsChildOf(UUserWidget::StaticClass()))
+			{
+				continue;
+			}
+		}
+
+		if (!FilterAssetData(BPAssetData))
+		{
+			// If this object isn't currently loaded, add it to the Library view
+			if (BPAssetData.ToSoftObjectPath().ResolveObject() == nullptr)
+			{
+				auto Template = MakeShareable(new FWidgetTemplateClass(BPAssetData, nullptr));
+				AddWidgetTemplate(Template);
+			}
+		}
+	}
+
+	TArray<FAssetData> AllWidgetBPsAssetData;
+	AssetRegistryModule.Get().GetAssetsByClass(UWidgetBlueprint::StaticClass()->GetFName(), AllWidgetBPsAssetData, true);
+
+	FName ActiveWidgetBlueprintName = ActiveWidgetBlueprintClass->ClassGeneratedBy->GetFName();
+	for (FAssetData& WidgetBPAssetData : AllWidgetBPsAssetData)
+	{
+		// Excludes the blueprint you're currently in
+		if (WidgetBPAssetData.AssetName == ActiveWidgetBlueprintName)
+		{
+			continue;
+		}
+
+		if (!FilterAssetData(WidgetBPAssetData))
+		{
+			// Excludes this widget if it is on the hide list
+			bool bIsOnList = false;
+			for (FSoftClassPath Widget : WidgetClassesToHide)
+			{
+				if (Widget.ToString().Find(WidgetBPAssetData.ObjectPath.ToString()) == 0)
+				{
+					bIsOnList = true;
+					break;
+				}
+			}
+			if (bIsOnList)
+			{
+				continue;
+			}
+
+			// If the blueprint generated class was found earlier, pass it to the template
+			TSubclassOf<UUserWidget> WidgetBPClass = nullptr;
+			auto LoadedWidgetBPClass = LoadedWidgetBlueprintClassesByName.Find(WidgetBPAssetData.AssetName);
+			if (LoadedWidgetBPClass)
+			{
+				WidgetBPClass = *LoadedWidgetBPClass;
+			}
+
+			uint32 BPFlags = WidgetBPAssetData.GetTagValueRef<uint32>(FBlueprintTags::ClassFlags);
+			if ((BPFlags & (CLASS_Abstract | CLASS_Deprecated | CLASS_HideDropDown)) == 0)
+			{
+				auto Template = MakeShareable(new FWidgetTemplateBlueprintClass(WidgetBPAssetData, WidgetBPClass));
+
+				AddWidgetTemplate(Template);
+			}
+		}
+	}
+}
+
+bool FLibraryViewModel::FilterAssetData(FAssetData &InAssetData)
+{
+	// Excludes engine content if user sets it to false
+	if (!GetDefault<UContentBrowserSettings>()->GetDisplayEngineFolder() || !GetDefault<UUMGEditorProjectSettings>()->bShowWidgetsFromEngineContent)
+	{
+		if (InAssetData.PackagePath.ToString().Find(TEXT("/Engine")) == 0)
+		{
+			return true;
+		}
+	}
+
+	// Excludes developer content if user sets it to false
+	if (!GetDefault<UContentBrowserSettings>()->GetDisplayDevelopersFolder() || !GetDefault<UUMGEditorProjectSettings>()->bShowWidgetsFromDeveloperContent)
+	{
+		if (InAssetData.PackagePath.ToString().Find(TEXT("/Game/Developers")) == 0)
+		{
+			return true;
+		}
+	}
+	return false;
+}
+
+void FLibraryViewModel::AddWidgetTemplate(TSharedPtr<FWidgetTemplate> Template)
+{
+	FString Category = Template->GetCategory().ToString();
+
+	// Hide user specific categories
+	TArray<FString> CategoriesToHide = GetDefault<UUMGEditorProjectSettings>()->CategoriesToHide;
+	for (FString CategoryName : CategoriesToHide)
+	{
+		if (Category == CategoryName)
+		{
+			return;
+		}
+	}
+	WidgetTemplateArray& Group = WidgetTemplateCategories.FindOrAdd(Category);
+	Group.Add(Template);
+}
+
+void FLibraryViewModel::OnObjectsReplaced(const TMap<UObject*, UObject*>& ReplacementMap)
+{
+}
+
+void FLibraryViewModel::OnBlueprintReinstanced()
+{
+	bRebuildRequested = true;
+}
+
+void FLibraryViewModel::OnFavoritesUpdated()
+{
+	bRebuildRequested = true;
+}
+
+void FLibraryViewModel::OnReloadComplete(EReloadCompleteReason Reason)
+{
+	bRebuildRequested = true;
+}
+
+void FLibraryViewModel::HandleOnAssetsDeleted(const TArray<UClass*>& DeletedAssetClasses)
+{
+	for (auto DeletedAssetClass : DeletedAssetClasses)
+	{
+		if (DeletedAssetClass->IsChildOf(UWidgetBlueprint::StaticClass()))
+		{
+			bRebuildRequested = true;
+		}
+	}
+}
+
+#undef LOCTEXT_NAMESPACE
+
+FWidgetTemplateListViewModel::FWidgetTemplateListViewModel()
+{
+}
+
+FText FWidgetTemplateListViewModel::GetName() const
+{
+	return FText();
+}
+
+bool FWidgetTemplateListViewModel::IsTemplate() const
+{
+	return false;
+}
+
+void FWidgetTemplateListViewModel::GetFilterStrings(TArray<FString>& OutStrings) const
+{
+}
+
+TSharedRef<ITableRow> FWidgetTemplateListViewModel::BuildRow(const TSharedRef<STableViewBase>& OwnerTable)
+{
+	// @TODO: DarenC - Temp construct on demand. Move to construct earlier?
+	if (!TemplatesFilter)
+	{
+		// Generate filter text
+		bool bHasFilters = false;
+		TStringBuilder<2048> FilterString; 
+		for (TSharedPtr<FWidgetTemplate> Template : Templates)
+		{
+			if (TSharedPtr<FWidgetTemplateClass> TemplateClass = StaticCastSharedPtr<FWidgetTemplateClass>(Template))
+			{
+				TWeakObjectPtr<UClass> WidgetClass = TemplateClass->GetWidgetClass();
+				if (WidgetClass.IsValid())
+				{
+					FString TemplateName = WidgetClass->GetFName().ToString();
+					if (!TemplateName.IsEmpty())
+					{
+						TemplateName.RemoveFromEnd(TEXT("_C"));
+						FilterString += bHasFilters ? "|+" : "+";
+						FilterString += TemplateName;
+						bHasFilters = true;
+					}
+				}
+			}
+		}
+
+		TSharedPtr<FFrontendFilter_Text> WidgetTextFilter = MakeShared<FFrontendFilter_Text>();
+		WidgetTextFilter->SetActive(true);
+		WidgetTextFilter->SetIncludeClassName(true);
+		WidgetTextFilter->SetRawFilterText(FText::FromString(FilterString.ToString()));
+
+		TemplatesFilter = MakeShared<FAssetFilterCollectionType>();
+		TemplatesFilter->Add(WidgetTextFilter);
+	}
+
+	if (!AssetViewPtr)
+	{
+		FARFilter WidgetTemplateFilter;
+		WidgetTemplateFilter.ClassNames.Add(UClass::StaticClass()->GetFName());
+		WidgetTemplateFilter.ClassNames.Add(UWidgetBlueprint::StaticClass()->GetFName());
+
+		AssetViewPtr = SNew(SAssetView)
+			.InitialCategoryFilter(EContentBrowserItemCategoryFilter::IncludeAssets | EContentBrowserItemCategoryFilter::IncludeClasses)
+			.InitialBackendFilter(WidgetTemplateFilter)
+			.FrontendFilters(TemplatesFilter)
+			.ForceShowEngineContent(true)
+			.ForceShowPluginContent(true)
+			;
+	}
+	AssetViewPtr->RequestSlowFullListRefresh();
+
+	return SNew(STableRow<TSharedPtr<FWidgetViewModel>>, OwnerTable)
+		.Padding(2.0f)
+		//.OnDragDetected(this, &FWidgetTemplateViewModel::OnDraggingWidgetTemplateItem)
+		[
+			AssetViewPtr.ToSharedRef()
+		];
+}
diff --git a/Engine/Source/Editor/UMGEditor/Private/Library/SLibraryViewModel.h b/Engine/Source/Editor/UMGEditor/Private/Library/SLibraryViewModel.h
new file mode 100644
--- /dev/null
+++ b/Engine/Source/Editor/UMGEditor/Private/Library/SLibraryViewModel.h
@@ -0,0 +1,118 @@
+// Copyright Epic Games, Inc. All Rights Reserved.
+
+#pragma once
+
+#include "CoreMinimal.h"
+#include "WidgetBlueprintEditor.h"
+#include "AssetData.h"
+
+#include "IContentBrowserSingleton.h"
+#include "Widgets/Views/STableViewBase.h"
+#include "Widgets/Views/STableRow.h"
+#include "SPaletteViewModel.h"
+
+class FWidgetTemplate;
+class FWidgetBlueprintEditor;
+class UWidgetBlueprint;
+class SAssetView;
+class SView;
+
+class FWidgetTemplateListViewModel : public FWidgetViewModel
+{
+public:
+	FWidgetTemplateListViewModel();
+
+	virtual FText GetName() const override;
+
+	virtual bool IsTemplate() const override;
+
+	virtual void GetFilterStrings(TArray<FString>& OutStrings) const override;
+
+	virtual TSharedRef<ITableRow> BuildRow(const TSharedRef<STableViewBase>& OwnerTable) override;
+
+	// FReply OnDraggingWidgetTemplateItem(const FGeometry& MyGeometry, const FPointerEvent& MouseEvent);
+
+	TArray<TSharedPtr<FWidgetTemplate>> Templates;
+	TSharedPtr<FAssetFilterCollectionType> TemplatesFilter;
+private:
+
+	/** The asset view widget */
+	TSharedPtr<SAssetView> AssetViewPtr;
+};
+
+class FLibraryViewModel : public FFavortiesViewModel
+{
+public:
+
+	DECLARE_MULTICAST_DELEGATE(FOnUpdating)
+	DECLARE_MULTICAST_DELEGATE(FOnUpdated)
+
+public:
+	FLibraryViewModel(TSharedPtr<FWidgetBlueprintEditor> InBlueprintEditor);
+	virtual ~FLibraryViewModel();
+
+	/** Register the View Model to events that should trigger a update of the Library*/
+	void RegisterToEvents();
+
+	/** Update the view model if needed and returns true if it did. */
+	void Update();
+
+	/** Returns true if the view model needs to be updated */
+	bool NeedUpdate() const { return bRebuildRequested; }
+	   
+	/** Add the widget template to the list of favorites */
+	virtual void AddToFavorites(const FWidgetTemplateViewModel* WidgetTemplateViewModel) override;
+
+	/** Remove the widget template to the list of favorites */
+	virtual void RemoveFromFavorites(const FWidgetTemplateViewModel* WidgetTemplateViewModel) override;
+
+	typedef TArray< TSharedPtr<FWidgetViewModel> > ViewModelsArray;
+	ViewModelsArray& GetWidgetViewModels() { return WidgetViewModels; }
+
+	/** Fires before the view model is updated */
+	FOnUpdating OnUpdating;
+
+	/** Fires after the view model is updated */
+	FOnUpdated OnUpdated;
+
+private:
+	FLibraryViewModel() {};
+
+	UWidgetBlueprint* GetBlueprint() const;
+
+	void BuildWidgetList();
+	void BuildClassWidgetList();
+
+	static bool FilterAssetData(FAssetData &BPAssetData);
+
+	void AddWidgetTemplate(TSharedPtr<FWidgetTemplate> Template);
+
+	/** Called when a Blueprint is recompiled and live objects are swapped out for replacements */
+	void OnObjectsReplaced(const TMap<UObject*, UObject*>& ReplacementMap);
+
+	/** Requests a rebuild of the widget list if a widget blueprint was compiled */
+	void OnBlueprintReinstanced();
+
+	/** Called when the favorite list is changed */
+	void OnFavoritesUpdated();
+
+	/** Requests a rebuild of the widget list */
+	void OnReloadComplete(EReloadCompleteReason Reason);
+
+	/** Requests a rebuild of the widget list if a widget blueprint was deleted */
+	void HandleOnAssetsDeleted(const TArray<UClass*>& DeletedAssetClasses);
+	
+	TWeakPtr<class FWidgetBlueprintEditor> BlueprintEditor;
+
+	typedef TArray<TSharedPtr<FWidgetTemplate>> WidgetTemplateArray;
+	TMap<FString, WidgetTemplateArray> WidgetTemplateCategories;
+
+	/** The source root view models for the tree. */
+	ViewModelsArray WidgetViewModels;
+	   
+	/** Controls rebuilding the list of spawnable widgets */
+	bool bRebuildRequested;
+
+	TSharedPtr<FWidgetHeaderViewModel> FavoriteHeader;
+};
+
diff --git a/Engine/Source/Editor/UMGEditor/Private/Palette/SPaletteViewModel.cpp b/Engine/Source/Editor/UMGEditor/Private/Palette/SPaletteViewModel.cpp
--- a/Engine/Source/Editor/UMGEditor/Private/Palette/SPaletteViewModel.cpp
+++ b/Engine/Source/Editor/UMGEditor/Private/Palette/SPaletteViewModel.cpp
@@ -27,7 +27,7 @@
 #define LOCTEXT_NAMESPACE "UMG"
 
 FWidgetTemplateViewModel::FWidgetTemplateViewModel()
-	: PaletteViewModel(nullptr),
+	: FavortiesViewModel(nullptr),
 	bIsFavorite(false)
 {
 }
@@ -54,7 +54,7 @@ TSharedRef<ITableRow> FWidgetTemplateViewModel::BuildRow(const TSharedRef<STable
 		.OnDragDetected(this, &FWidgetTemplateViewModel::OnDraggingWidgetTemplateItem)
 		[
 			SNew(SPaletteViewItem, SharedThis(this))
-			.HighlightText(PaletteViewModel, &FPaletteViewModel::GetSearchText)
+			.HighlightText(FavortiesViewModel, &FFavortiesViewModel::GetSearchText)
 		];
 }
 
@@ -66,13 +66,13 @@ FReply FWidgetTemplateViewModel::OnDraggingWidgetTemplateItem(const FGeometry& M
 void FWidgetTemplateViewModel::AddToFavorites()
 {
 	bIsFavorite = true;
-	PaletteViewModel->AddToFavorites(this);
+	FavortiesViewModel->AddToFavorites(this);
 }
 
 void FWidgetTemplateViewModel::RemoveFromFavorites()
 {
 	bIsFavorite = false;
-	PaletteViewModel->RemoveFromFavorites(this);
+	FavortiesViewModel->RemoveFromFavorites(this);
 }
 
 TSharedRef<ITableRow> FWidgetHeaderViewModel::BuildRow(const TSharedRef<STableViewBase>& OwnerTable)
@@ -194,7 +194,7 @@ void FPaletteViewModel::BuildWidgetList()
 		{
 			TSharedPtr<FWidgetTemplateViewModel> TemplateViewModel = MakeShareable(new FWidgetTemplateViewModel());
 			TemplateViewModel->Template = Template;
-			TemplateViewModel->PaletteViewModel = this;
+			TemplateViewModel->FavortiesViewModel = this;
 			Header->Children.Add(TemplateViewModel);
 
 			// If it's a favorite, we also add it to the Favorite section
@@ -206,7 +206,7 @@ void FPaletteViewModel::BuildWidgetList()
 				// We have to create a second copy of the ViewModel for the treeview has it doesn't support to have the same element twice.
 				TSharedPtr<FWidgetTemplateViewModel> FavoriteTemplateViewModel = MakeShareable(new FWidgetTemplateViewModel());
 				FavoriteTemplateViewModel->Template = Template;
-				FavoriteTemplateViewModel->PaletteViewModel = this;
+				FavoriteTemplateViewModel->FavortiesViewModel = this;
 				FavoriteTemplateViewModel->SetFavorite();
 
 				FavoriteHeader->Children.Add(FavoriteTemplateViewModel);
diff --git a/Engine/Source/Editor/UMGEditor/Private/Palette/SPaletteViewModel.h b/Engine/Source/Editor/UMGEditor/Private/Palette/SPaletteViewModel.h
--- a/Engine/Source/Editor/UMGEditor/Private/Palette/SPaletteViewModel.h
+++ b/Engine/Source/Editor/UMGEditor/Private/Palette/SPaletteViewModel.h
@@ -13,6 +13,7 @@ class FWidgetTemplate;
 class FWidgetBlueprintEditor;
 class UWidgetBlueprint;
 class SPaletteView;
+class FFavortiesViewModel;
 
 /** View model for the items in the widget template list */
 class FWidgetViewModel : public TSharedFromThis<FWidgetViewModel>
@@ -72,7 +73,7 @@ public:
 	virtual void SetFavorite() override { bIsFavorite = true; }
 
 	TSharedPtr<FWidgetTemplate> Template;
-	FPaletteViewModel* PaletteViewModel;
+	FFavortiesViewModel* FavortiesViewModel;
 private:
 	/** True is the widget is a favorite. It's keep as a state to prevent a search in the favorite list. */
 	bool bIsFavorite;
@@ -117,7 +118,25 @@ private:
 	bool bForceExpansion = false;
 };
 
-class FPaletteViewModel: public TSharedFromThis<FPaletteViewModel>
+class FFavortiesViewModel : public TSharedFromThis<FFavortiesViewModel>
+{
+public:
+	virtual ~FFavortiesViewModel() { }
+
+	/** Add the widget template to the list of favorites */
+	virtual void AddToFavorites(const FWidgetTemplateViewModel* WidgetTemplateViewModel) = 0;
+
+	/** Remove the widget template to the list of favorites */
+	virtual void RemoveFromFavorites(const FWidgetTemplateViewModel* WidgetTemplateViewModel) = 0;
+
+	void SetSearchText(const FText& inSearchText) { SearchText = inSearchText; }
+	FText GetSearchText() const { return SearchText; }
+
+protected:
+	FText SearchText;
+};
+
+class FPaletteViewModel : public FFavortiesViewModel
 {
 public:
 
@@ -126,7 +145,7 @@ public:
 
 public:
 	FPaletteViewModel(TSharedPtr<FWidgetBlueprintEditor> InBlueprintEditor);
-	~FPaletteViewModel();
+	virtual ~FPaletteViewModel();
 
 	/** Register the View Model to events that should trigger a update of the Palette*/
 	void RegisterToEvents();
@@ -146,9 +165,6 @@ public:
 	typedef TArray< TSharedPtr<FWidgetViewModel> > ViewModelsArray;
 	ViewModelsArray& GetWidgetViewModels() { return WidgetViewModels; }
 
-	void SetSearchText(const FText& inSearchText) { SearchText = inSearchText; }
-	FText GetSearchText() const { return SearchText; }
-
 	/** Fires before the view model is updated */
 	FOnUpdating OnUpdating;
 
@@ -193,8 +209,6 @@ private:
 	/** Controls rebuilding the list of spawnable widgets */
 	bool bRebuildRequested;
 
-	FText SearchText;
-
 	TSharedPtr<FWidgetHeaderViewModel> FavoriteHeader;
 };
 
diff --git a/Engine/Source/Editor/UMGEditor/Private/TabFactory/LibraryTabSummoner.cpp b/Engine/Source/Editor/UMGEditor/Private/TabFactory/LibraryTabSummoner.cpp
new file mode 100644
--- /dev/null
+++ b/Engine/Source/Editor/UMGEditor/Private/TabFactory/LibraryTabSummoner.cpp
@@ -0,0 +1,32 @@
+// Copyright Epic Games, Inc. All Rights Reserved.
+
+#include "TabFactory/LibraryTabSummoner.h"
+#include "Library/SLibraryView.h"
+#include "UMGStyle.h"
+
+#define LOCTEXT_NAMESPACE "UMG"
+
+const FName FLibraryTabSummoner::TabID(TEXT("WidgetLibrary"));
+
+FLibraryTabSummoner::FLibraryTabSummoner(TSharedPtr<class FWidgetBlueprintEditor> InBlueprintEditor)
+		: FWorkflowTabFactory(TabID, InBlueprintEditor)
+		, BlueprintEditor(InBlueprintEditor)
+{
+	TabLabel = LOCTEXT("WidgetLibraryTabLabel", "Library");
+	TabIcon = FSlateIcon(FUMGStyle::GetStyleSetName(), "Palette.TabIcon");
+
+	bIsSingleton = true;
+
+	ViewMenuDescription = LOCTEXT("WidgetLibrary_ViewMenu_Desc", "Library");
+	ViewMenuTooltip = LOCTEXT("WidgetLibrary_ViewMenu_ToolTip", "Show the Widget Library");
+}
+
+TSharedRef<SWidget> FLibraryTabSummoner::CreateTabBody(const FWorkflowTabSpawnInfo& Info) const
+{
+	TSharedPtr<FWidgetBlueprintEditor> BlueprintEditorPtr = StaticCastSharedPtr<FWidgetBlueprintEditor>(BlueprintEditor.Pin());
+
+	return SNew(SLibraryView, BlueprintEditorPtr)
+		.AddMetaData<FTagMetaData>(FTagMetaData(TEXT("Library")));
+}
+
+#undef LOCTEXT_NAMESPACE 
diff --git a/Engine/Source/Editor/UMGEditor/Private/TabFactory/LibraryTabSummoner.h b/Engine/Source/Editor/UMGEditor/Private/TabFactory/LibraryTabSummoner.h
new file mode 100644
--- /dev/null
+++ b/Engine/Source/Editor/UMGEditor/Private/TabFactory/LibraryTabSummoner.h
@@ -0,0 +1,22 @@
+// Copyright Epic Games, Inc. All Rights Reserved.
+
+#pragma once
+
+#include "CoreMinimal.h"
+#include "Widgets/SWidget.h"
+#include "WorkflowOrientedApp/WorkflowTabFactory.h"
+#include "WidgetBlueprintEditor.h"
+
+struct FLibraryTabSummoner : public FWorkflowTabFactory
+{
+public:
+	static const FName TabID;
+	
+public:
+	FLibraryTabSummoner(TSharedPtr<class FWidgetBlueprintEditor> InBlueprintEditor);
+	
+	virtual TSharedRef<SWidget> CreateTabBody(const FWorkflowTabSpawnInfo& Info) const override;
+	
+protected:
+	TWeakPtr<class FWidgetBlueprintEditor> BlueprintEditor;
+};
diff --git a/Engine/Source/Editor/UMGEditor/Private/WidgetBlueprintEditor.cpp b/Engine/Source/Editor/UMGEditor/Private/WidgetBlueprintEditor.cpp
--- a/Engine/Source/Editor/UMGEditor/Private/WidgetBlueprintEditor.cpp
+++ b/Engine/Source/Editor/UMGEditor/Private/WidgetBlueprintEditor.cpp
@@ -8,6 +8,7 @@
 #include "Widgets/Text/STextBlock.h"
 #include "Widgets/Docking/SDockTab.h"
 #include "Framework/MultiBox/MultiBoxBuilder.h"
+#include "HAL/FileManager.h"
 #include "Engine/SimpleConstructionScript.h"
 #include "Blueprint/WidgetBlueprintGeneratedClass.h"
 #include "Kismet2/BlueprintEditorUtils.h"
@@ -57,6 +58,7 @@
 #include "GameProjectGenerationModule.h"
 
 #include "SPaletteViewModel.h"
+#include "SLibraryViewModel.h"
 
 #include "DesktopPlatformModule.h"
 #include "Engine/TextureRenderTarget2D.h"
@@ -125,6 +127,9 @@ void FWidgetBlueprintEditor::InitWidgetBlueprintEditor(const EToolkitMode::Type
 	PaletteViewModel = MakeShareable(new FPaletteViewModel(ThisPtr));
 	PaletteViewModel->RegisterToEvents();
 
+	LibraryViewModel = MakeShareable(new FLibraryViewModel(ThisPtr));
+	LibraryViewModel->RegisterToEvents();
+
 	WidgetToolbar = MakeShareable(new FWidgetBlueprintEditorToolbar(ThisPtr));
 
 	BindToolkitCommands();
@@ -746,11 +751,16 @@ void FWidgetBlueprintEditor::Tick(float DeltaTime)
 		RefreshPreview();
 	}
 
-	// Updat the palette view model.
+	// Update the palette view model.
 	if (PaletteViewModel->NeedUpdate())
 	{
 		PaletteViewModel->Update();
 	}
+
+	if (LibraryViewModel->NeedUpdate())
+	{
+		LibraryViewModel->Update();
+	}
 }
 
 static bool MigratePropertyValue(UObject* SourceObject, UObject* DestinationObject, FEditPropertyChain::TDoubleLinkedListNode* PropertyChainNode, FProperty* MemberProperty, bool bIsModify)
diff --git a/Engine/Source/Editor/UMGEditor/Private/WidgetBlueprintEditorUtils.cpp b/Engine/Source/Editor/UMGEditor/Private/WidgetBlueprintEditorUtils.cpp
--- a/Engine/Source/Editor/UMGEditor/Private/WidgetBlueprintEditorUtils.cpp
+++ b/Engine/Source/Editor/UMGEditor/Private/WidgetBlueprintEditorUtils.cpp
@@ -16,6 +16,10 @@
 #include "HAL/PlatformApplicationMisc.h"
 
 #include "Dialogs/Dialogs.h"
+#include "DragAndDrop/DecoratedDragDropOp.h"
+#include "DragAndDrop/AssetDragDropOp.h"
+#include "DragAndDrop/ClassDragDropOp.h"
+#include "DragDrop/WidgetTemplateDragDropOp.h"
 #include "Exporters/Exporter.h"
 #include "ObjectEditorUtils.h"
 #include "Components/CanvasPanelSlot.h"
@@ -25,6 +29,7 @@
 #include "Kismet2/Kismet2NameValidators.h"
 #include "Kismet2/BlueprintEditorUtils.h"
 #include "Templates/WidgetTemplateClass.h"
+#include "Templates/WidgetTemplateImageClass.h"
 #include "Templates/WidgetTemplateBlueprintClass.h"
 #include "Factories.h"
 #include "UnrealExporter.h"
@@ -1347,6 +1352,61 @@ TArray<UWidget*> FWidgetBlueprintEditorUtils::DuplicateWidgets(TSharedRef<FWidge
 
 	return DuplicatedWidgets;
 }
+UWidget* FWidgetBlueprintEditorUtils::GetWidgetTemplateFromDragDrop(UWidgetBlueprint* Blueprint, UWidgetTree* RootWidgetTree, TSharedPtr<FDragDropOperation>& DragDropOp)
+{
+	UWidget* Widget = nullptr;
+
+	if (!DragDropOp.IsValid())
+	{
+		return nullptr;
+	}
+
+	if (DragDropOp->IsOfType<FWidgetTemplateDragDropOp>())
+	{
+		TSharedPtr<FWidgetTemplateDragDropOp> TemplateDragDropOp = StaticCastSharedPtr<FWidgetTemplateDragDropOp>(DragDropOp);
+		Widget = TemplateDragDropOp->Template->Create(RootWidgetTree);
+	}
+	else if (DragDropOp->IsOfType<FAssetDragDropOp>())
+	{
+		TSharedPtr<FAssetDragDropOp> AssetDragDropOp = StaticCastSharedPtr<FAssetDragDropOp>(DragDropOp);
+		if (AssetDragDropOp->GetAssets().Num() > 0)
+		{
+			// Only handle first valid dragged widget, multi widget drag drop is not practically useful
+			const FAssetData& AssetData = AssetDragDropOp->GetAssets()[0];
+
+			bool CodeClass = AssetData.AssetClass == "Class";
+			FName ClassName = CodeClass ? AssetData.ObjectPath : AssetData.AssetClass;
+			UClass* AssetClass = FindObjectChecked<UClass>(ANY_PACKAGE, *ClassName.ToString());
+
+			if (FWidgetTemplateBlueprintClass::Supports(AssetClass))
+			{
+				// Allows a UMG Widget Blueprint to be dragged from the Content Browser to another Widget Blueprint...as long as we're not trying to place a
+				// blueprint inside itself.
+				FString BlueprintPath = Blueprint->GetPathName();
+				if (BlueprintPath != AssetData.ObjectPath.ToString())
+				{
+					Widget = FWidgetTemplateBlueprintClass(AssetData).Create(RootWidgetTree);
+				}
+			}
+			else if (CodeClass && AssetClass && AssetClass->IsChildOf(UWidget::StaticClass()))
+			{
+				Widget = FWidgetTemplateClass(AssetClass).Create(RootWidgetTree);
+			}
+			else if (FWidgetTemplateImageClass::Supports(AssetClass))
+			{
+				Widget = FWidgetTemplateImageClass(AssetData).Create(RootWidgetTree);
+			}
+		}
+	}
+
+	// Check to make sure that this widget can be added to the current blueprint
+	if (Cast<UUserWidget>(Widget) && !Blueprint->IsWidgetFreeFromCircularReferences(Cast<UUserWidget>(Widget)))
+	{
+		Widget = nullptr;
+	}
+
+	return Widget;
+}
 
 void FWidgetBlueprintEditorUtils::ExportWidgetsToText(TArray<UWidget*> WidgetsToExport, /*out*/ FString& ExportedText)
 {
diff --git a/Engine/Source/Editor/UMGEditor/Private/WidgetBlueprintEditorUtils.h b/Engine/Source/Editor/UMGEditor/Private/WidgetBlueprintEditorUtils.h
--- a/Engine/Source/Editor/UMGEditor/Private/WidgetBlueprintEditorUtils.h
+++ b/Engine/Source/Editor/UMGEditor/Private/WidgetBlueprintEditorUtils.h
@@ -7,6 +7,7 @@
 #include "WidgetBlueprint.h"
 #include "WidgetReference.h"
 
+class FDragDropOperation;
 class FHittestGrid;
 class FMenuBuilder;
 class UWidgetBlueprint;
@@ -43,6 +44,8 @@ public:
 
 	static TArray<UWidget*> DuplicateWidgets(TSharedRef<FWidgetBlueprintEditor> BlueprintEditor, UWidgetBlueprint* BP, TSet<FWidgetReference> Widgets);
 
+	static UWidget* GetWidgetTemplateFromDragDrop(UWidgetBlueprint* Blueprint, UWidgetTree* RootWidgetTree, TSharedPtr<FDragDropOperation>& DragDropOp);
+
 	static bool IsBindWidgetProperty(FProperty* InProperty);
 	static bool IsBindWidgetProperty(FProperty* InProperty, bool& bIsOptional);
 
diff --git a/Engine/Source/Editor/UMGEditor/Public/WidgetBlueprintEditor.h b/Engine/Source/Editor/UMGEditor/Public/WidgetBlueprintEditor.h
--- a/Engine/Source/Editor/UMGEditor/Public/WidgetBlueprintEditor.h
+++ b/Engine/Source/Editor/UMGEditor/Public/WidgetBlueprintEditor.h
@@ -23,6 +23,7 @@ class UPanelSlot;
 class UWidgetAnimation;
 class UWidgetBlueprint;
 class FPaletteViewModel;
+class FLibraryViewModel;
 
 struct FNamedSlotSelection
 {
@@ -177,6 +178,7 @@ public:
 	void SetIsRespectingLocks(bool Value);
 
 	TSharedPtr<FPaletteViewModel> GetPaletteViewModel() { return PaletteViewModel; };
+	TSharedPtr<FLibraryViewModel> GetLibraryViewModel() { return LibraryViewModel; };
 
 public:
 	/** Fires whenever a new widget is being hovered over */
@@ -392,6 +394,9 @@ private:
 	/** ViewModel used by the Palette and Palette Favorite Views */
 	TSharedPtr<FPaletteViewModel> PaletteViewModel;
 
+	/** ViewModel used by the Library View */
+	TSharedPtr<FLibraryViewModel> LibraryViewModel;
+
 	/** When true the sequencer selection is being updated from changes to the external selection. */
 	bool bUpdatingSequencerSelection;
 
diff --git a/Engine/Source/Editor/UMGEditor/UMGEditor.Build.cs b/Engine/Source/Editor/UMGEditor/UMGEditor.Build.cs
--- a/Engine/Source/Editor/UMGEditor/UMGEditor.Build.cs
+++ b/Engine/Source/Editor/UMGEditor/UMGEditor.Build.cs
@@ -17,6 +17,7 @@ public class UMGEditor : ModuleRules
 				"Editor/UMGEditor/Private/Designer",
 				"Editor/UMGEditor/Private/Hierarchy",
 				"Editor/UMGEditor/Private/Palette",
+				"Editor/UMGEditor/Private/Library",
 				"Editor/UMGEditor/Private/Details",
 				"Editor/UMGEditor/Private/DragDrop",
                 "Editor/UMGEditor/Private/Utility",
@@ -40,6 +41,7 @@ public class UMGEditor : ModuleRules
 			new string[] {
 				"Core",
 				"CoreUObject",
+				"ContentBrowser",
 				"ApplicationCore",
 				"InputCore",
 				"Engine",
