Fix private include AssetViewSortManager.

diff --git a/Engine/Source/Editor/ContentBrowser/Private/AssetViewSortManager.h b/Engine/Source/Editor/ContentBrowser/Private/AssetViewSortManager.h
deleted file mode 100644
diff --git a/Engine/Source/Editor/ContentBrowser/Public/AssetViewSortManager.h b/Engine/Source/Editor/ContentBrowser/Public/AssetViewSortManager.h
new file mode 100644
--- /dev/null
+++ b/Engine/Source/Editor/ContentBrowser/Public/AssetViewSortManager.h
@@ -0,0 +1,61 @@
+// Copyright Epic Games, Inc. All Rights Reserved.
+
+#pragma once
+
+#include "CoreMinimal.h"
+#include "Containers/StringView.h"
+#include "Widgets/Views/SHeaderRow.h"
+#include "ContentBrowserDelegates.h"
+
+class FAssetViewItem;
+
+class FAssetViewSortManager
+{
+public:
+	/** Constructor */
+	FAssetViewSortManager();
+
+	/** Reset the sort mode back to default */
+	void ResetSort();
+
+	/** Sorts a list using the current ColumnId and Mode. Supply a MajorityAssetType to help discover sorting type (numerical vs alphabetical) */
+	void SortList(TArray<TSharedPtr<FAssetViewItem>>& AssetItems, const FName& MajorityAssetType, const TArray<FAssetViewCustomColumn>& CustomColumns) const;
+
+	/** Exports the list of asset items to CSV, in order and with the listed columns */
+	void ExportColumnsToCSV(TArray<TSharedPtr<FAssetViewItem>>& AssetItems, TArray<FName>& ColumnList, const TArray<FAssetViewCustomColumn>& CustomColumns, FString& OutString) const;
+
+	/** Sets the sort mode based on the column that was clicked, returns true if newly assigned */
+	bool SetOrToggleSortColumn(const EColumnSortPriority::Type InSortPriority, const FName& InColumnId);
+
+	/**	Sets the column to sort */
+	void SetSortColumnId(const EColumnSortPriority::Type InSortPriority, const FName& InColumnId);
+
+	/**	Sets the current sort mode */
+	void SetSortMode(const EColumnSortPriority::Type InSortPriority, const EColumnSortMode::Type InSortMode);
+
+	/** Gets the current sort mode */
+	EColumnSortMode::Type GetSortMode(const EColumnSortPriority::Type InSortPriority) const { check(InSortPriority < EColumnSortPriority::Max); return SortMode[InSortPriority]; }
+
+	/** Gets the current sort column id */
+	FName GetSortColumnId(const EColumnSortPriority::Type InSortPriority) const { check(InSortPriority < EColumnSortPriority::Max); return SortColumnId[InSortPriority]; }
+
+	/** Refresh a custom column if needed.  If found, returns true with TagType parameter set */
+	bool FindAndRefreshCustomColumn(const TArray<TSharedPtr<FAssetViewItem>>& AssetItems, const FName ColumnName, const TArray<FAssetViewCustomColumn>& CustomColumns, UObject::FAssetRegistryTag::ETagType& TagType) const;
+
+public:
+	/** The names of non-type specific columns in the columns view. */
+	static const FName NameColumnId;
+	static const FName ClassColumnId;
+	static const FName PathColumnId;
+
+public:
+	/** Compare the two strings, correctly ordering any numeric suffixes they may have */
+	static int32 CompareWithNumericSuffix(FStringView A, FStringView B);
+
+private:
+	/** The name of the column that is currently used for sorting. */
+	FName SortColumnId[EColumnSortPriority::Max];
+
+	/** Whether the sort is ascending or descending. */
+	EColumnSortMode::Type SortMode[EColumnSortPriority::Max];
+};
diff --git a/Engine/Source/Editor/ContentBrowser/Public/SAssetView.h b/Engine/Source/Editor/ContentBrowser/Public/SAssetView.h
--- a/Engine/Source/Editor/ContentBrowser/Public/SAssetView.h
+++ b/Engine/Source/Editor/ContentBrowser/Public/SAssetView.h
@@ -19,7 +19,7 @@
 #include "Animation/CurveSequence.h"
 #include "Widgets/Views/STableViewBase.h"
 #include "Widgets/Views/STableRow.h"
-#include "Editor/ContentBrowser/Private/AssetViewSortManager.h"
+#include "AssetViewSortManager.h"
 #include "AssetViewTypes.h"
 #include "HistoryManager.h"
 #include "Misc/BlacklistNames.h"
