Add _API to AssetViewSortManager to fix Mac Builds.

diff --git a/Engine/Source/Editor/ContentBrowser/Public/AssetViewSortManager.h b/Engine/Source/Editor/ContentBrowser/Public/AssetViewSortManager.h
--- a/Engine/Source/Editor/ContentBrowser/Public/AssetViewSortManager.h
+++ b/Engine/Source/Editor/ContentBrowser/Public/AssetViewSortManager.h
@@ -9,7 +9,7 @@
 
 class FAssetViewItem;
 
-class FAssetViewSortManager
+class CONTENTBROWSER_API FAssetViewSortManager
 {
 public:
 	/** Constructor */
