Improve the performence when expending all the folders of the SPathView widget.
From 4s to less than a 1 ms in the project TestEngine while expending all the folders.
The processing has been moved to the function populate where it take between 3 to 6 ms to run when needed.

diff --git a/Engine/Source/Editor/ContentBrowser/Private/SPathView.cpp b/Engine/Source/Editor/ContentBrowser/Private/SPathView.cpp
--- a/Engine/Source/Editor/ContentBrowser/Private/SPathView.cpp
+++ b/Engine/Source/Editor/ContentBrowser/Private/SPathView.cpp
@@ -532,7 +532,7 @@ TSharedPtr<FTreeItem> SPathView::AddFolderItem(FContentBrowserItemData&& InItem,
 	{
 		const bool bIsLeafmostItem = PathItemIndex == PathItemList.Num() - 1;
 
-		const FString FolderNameStr = PathItemList[PathItemIndex];
+		const FString& FolderNameStr = PathItemList[PathItemIndex];
 		const FName FolderName = *FolderNameStr;
 		FPathViews::Append(CurrentPathStr, FolderNameStr);
 
@@ -1225,23 +1225,7 @@ void SPathView::TreeExpansionChanged( TSharedPtr< FTreeItem > TreeItem, bool bIs
 {
 	if ( ShouldAllowTreeItemChangedDelegate() )
 	{
-		TSet<TSharedPtr<FTreeItem>> ExpandedItemSet;
-		TreeViewPtr->GetExpandedItems(ExpandedItemSet);
-		const TArray<TSharedPtr<FTreeItem>> ExpandedItems = ExpandedItemSet.Array();
-
-		LastExpandedPaths.Empty();
-		for (int32 ItemIdx = 0; ItemIdx < ExpandedItems.Num(); ++ItemIdx)
-		{
-			const TSharedPtr<FTreeItem> Item = ExpandedItems[ItemIdx];
-			if ( !ensure(Item.IsValid()) )
-			{
-				// All items must exist
-				continue;
-			}
-
-			// Keep track of the last paths that we broadcasted for expansion reasons when filtering
-			LastExpandedPaths.Add(Item->GetItem().GetInvariantPath());
-		}
+		DirtyLastExpandedPaths();
 
 		if (!bIsExpanded)
 		{
@@ -1290,6 +1274,9 @@ FText SPathView::GetHighlightText() const
 
 void SPathView::Populate(const bool bIsRefreshingFilter)
 {
+	// Update the list of expanded path before removing the items
+	UpdateLastExpandedPathsIfDirty();
+
 	const bool bFilteringByText = !SearchBoxFolderFilter->GetRawFilterText().IsEmpty();
 
 	// Batch the selection changed event
@@ -1936,6 +1923,35 @@ TArray<FName> SPathView::GetDefaultPathsToExpand() const
 	return VirtualPaths;
 }
 
+void SPathView::DirtyLastExpandedPaths()
+{
+	bLastExpandedPathsDirty = true;
+}
+
+void SPathView::UpdateLastExpandedPathsIfDirty()
+{
+	if (bLastExpandedPathsDirty)
+	{
+		TSet<TSharedPtr<FTreeItem>> ExpandedItemSet;
+		TreeViewPtr->GetExpandedItems(ExpandedItemSet);
+
+		LastExpandedPaths.Empty(ExpandedItemSet.Num());
+		for (const TSharedPtr<FTreeItem>& Item : ExpandedItemSet)
+		{
+			if (!ensure(Item.IsValid()))
+			{
+				// All items must exist
+				continue;
+			}
+
+			// Keep track of the last paths that we broadcasted for expansion reasons when filtering
+			LastExpandedPaths.Add(Item->GetItem().GetInvariantPath());
+		}
+
+		bLastExpandedPathsDirty = false;
+	}
+}
+
 void SFavoritePathView::Construct(const FArguments& InArgs)
 {
 	SAssignNew(TreeViewPtr, STreeView< TSharedPtr<FTreeItem> >)
diff --git a/Engine/Source/Editor/ContentBrowser/Private/SPathView.h b/Engine/Source/Editor/ContentBrowser/Private/SPathView.h
--- a/Engine/Source/Editor/ContentBrowser/Private/SPathView.h
+++ b/Engine/Source/Editor/ContentBrowser/Private/SPathView.h
@@ -307,6 +307,12 @@ private:
 
 	TArray<FName> GetDefaultPathsToExpand() const;
 
+	/** Tell the tree that the LastExpandedPath set should be refreshed */
+	void DirtyLastExpandedPaths();
+
+	/** Update the LastExpandedPath if required */
+	void UpdateLastExpandedPathsIfDirty();
+
 protected:
 	/** A helper class to manage PreventTreeItemChangedDelegateCount by incrementing it when constructed (on the stack) and decrementing when destroyed */
 	class FScopedPreventTreeItemChangedDelegate
@@ -381,6 +387,8 @@ protected:
 	TSharedPtr<FBlacklistPaths> CustomFolderBlacklist;
 
 private:
+	/** Used to track if the list of last expanded path should be updated */
+	bool bLastExpandedPathsDirty = false;
 
 	/** The paths that were last reported by OnPathExpanded event. Used in preserving expansion when filtering folders */
 	TSet<FName> LastExpandedPaths;
