Scroll Bars now highlight when hovering anywhere along the slider well
Moved Show Empty Folders to the top level in Content Browser Settings
Replaced the icon for SingleFrameAdvance in the play control toolbar buttons

diff --git a/Engine/Source/Editor/ContentBrowser/Private/SAssetView.cpp b/Engine/Source/Editor/ContentBrowser/Private/SAssetView.cpp
--- a/Engine/Source/Editor/ContentBrowser/Private/SAssetView.cpp
+++ b/Engine/Source/Editor/ContentBrowser/Private/SAssetView.cpp
@@ -2419,35 +2419,32 @@ void SAssetView::PopulateViewButtonMenu(UToolMenu* Menu)
 
 	{
 		FToolMenuSection& Section = Menu->AddSection("View", LOCTEXT("ViewHeading", "View"));
-		auto CreateShowFoldersSubMenu = [this](UToolMenu* SubMenu)
-		{
-			FToolMenuSection& ShowEmptyFoldersSection = SubMenu->AddSection("ShowEmptyFolders");
-			ShowEmptyFoldersSection.AddMenuEntry(
-				"ShowEmptyFolders",
-				LOCTEXT("ShowEmptyFoldersOption", "Show Empty Folders"),
-				LOCTEXT("ShowEmptyFoldersOptionToolTip", "Show empty folders in the view as well as assets?"),
-				FSlateIcon(),
-				FUIAction(
-					FExecuteAction::CreateSP( this, &SAssetView::ToggleShowEmptyFolders ),
-					FCanExecuteAction::CreateSP( this, &SAssetView::IsToggleShowEmptyFoldersAllowed ),
-					FIsActionChecked::CreateSP( this, &SAssetView::IsShowingEmptyFolders )
-				),
-				EUserInterfaceActionType::ToggleButton
-			);
-		};
 
-		Section.AddEntry(FToolMenuEntry::InitSubMenu(
+		Section.AddMenuEntry(
 			"ShowFolders",
 			LOCTEXT("ShowFoldersOption", "Show Folders"),
 			LOCTEXT("ShowFoldersOptionToolTip", "Show folders in the view as well as assets?"),
-			FNewToolMenuDelegate::CreateLambda(CreateShowFoldersSubMenu),
+			FSlateIcon(),
 			FUIAction(
 				FExecuteAction::CreateSP( this, &SAssetView::ToggleShowFolders ),
 				FCanExecuteAction::CreateSP( this, &SAssetView::IsToggleShowFoldersAllowed ),
 				FIsActionChecked::CreateSP( this, &SAssetView::IsShowingFolders )
 			),
 			EUserInterfaceActionType::ToggleButton
-		));
+		);
+
+		Section.AddMenuEntry(
+			"ShowEmptyFolders",
+			LOCTEXT("ShowEmptyFoldersOption", "Show Empty Folders"),
+			LOCTEXT("ShowEmptyFoldersOptionToolTip", "Show empty folders in the view as well as assets?"),
+			FSlateIcon(),
+			FUIAction(
+				FExecuteAction::CreateSP(this, &SAssetView::ToggleShowEmptyFolders),
+				FCanExecuteAction::CreateSP(this, &SAssetView::IsToggleShowEmptyFoldersAllowed),
+				FIsActionChecked::CreateSP(this, &SAssetView::IsShowingEmptyFolders)
+			),
+			EUserInterfaceActionType::ToggleButton
+		);
 
 		Section.AddMenuEntry(
 			"ShowFavorite",
diff --git a/Engine/Source/Editor/EditorStyle/Private/StarshipStyle.cpp b/Engine/Source/Editor/EditorStyle/Private/StarshipStyle.cpp
--- a/Engine/Source/Editor/EditorStyle/Private/StarshipStyle.cpp
+++ b/Engine/Source/Editor/EditorStyle/Private/StarshipStyle.cpp
@@ -5539,8 +5539,8 @@ void FStarshipEditorStyle::FStyle::SetupPersonaStyle()
 		Set("PlayWorld.ResumePlaySession.Small", new IMAGE_BRUSH_SVG("Starship/MainToolbar/simulate", Icon20x20));
 		Set( "PlayWorld.PausePlaySession", new IMAGE_BRUSH_SVG("Starship/MainToolbar/pause", Icon40x40));
 		Set( "PlayWorld.PausePlaySession.Small", new IMAGE_BRUSH_SVG("Starship/MainToolbar/pause", Icon20x20));
-		Set( "PlayWorld.SingleFrameAdvance", new IMAGE_BRUSH( "Icons/icon_advance_40x", Icon40x40 ) );
-		Set( "PlayWorld.SingleFrameAdvance.Small", new IMAGE_BRUSH( "Icons/icon_advance_40x", Icon20x20 ) );
+		Set( "PlayWorld.SingleFrameAdvance", new IMAGE_BRUSH_SVG( "Starship/MainToolbar/SingleFrameAdvance", Icon40x40 ) );
+		Set( "PlayWorld.SingleFrameAdvance.Small", new IMAGE_BRUSH_SVG( "Starship/MainToolbar/SingleFrameAdvance", Icon20x20 ) );
 
 		Set( "PlayWorld.StopPlaySession", new IMAGE_BRUSH_SVG("Starship/MainToolbar/stop", Icon40x40));
 		Set( "PlayWorld.StopPlaySession.Small", new IMAGE_BRUSH_SVG("Starship/MainToolbar/stop", Icon20x20));
diff --git a/Engine/Source/Runtime/Slate/Private/Widgets/Layout/SScrollBar.cpp b/Engine/Source/Runtime/Slate/Private/Widgets/Layout/SScrollBar.cpp
--- a/Engine/Source/Runtime/Slate/Private/Widgets/Layout/SScrollBar.cpp
+++ b/Engine/Source/Runtime/Slate/Private/Widgets/Layout/SScrollBar.cpp
@@ -285,7 +285,7 @@ const FSlateBrush* SScrollBar::GetDragThumbImage() const
 	{
 		return DraggedThumbImage;
 	}
-	else if (DragThumb->IsDirectlyHovered())
+	else if (IsHovered())
 	{
 		return HoveredThumbImage;
 	}
