Fixed an issue that prevented the content browser from being dockable and appearing in the Window Menu

diff --git a/Engine/Source/Editor/ContentBrowser/Private/ContentBrowserSingleton.cpp b/Engine/Source/Editor/ContentBrowser/Private/ContentBrowserSingleton.cpp
--- a/Engine/Source/Editor/ContentBrowser/Private/ContentBrowserSingleton.cpp
+++ b/Engine/Source/Editor/ContentBrowser/Private/ContentBrowserSingleton.cpp
@@ -52,6 +52,21 @@ FContentBrowserSingleton::FContentBrowserSingleton()
 		ContentBrowserIcon,
 		true);
 
+	for (int32 BrowserIdx = 0; BrowserIdx < UE_ARRAY_COUNT(ContentBrowserTabIDs); BrowserIdx++)
+	{
+		const FName TabID = FName(*FString::Printf(TEXT("ContentBrowserTab%d"), BrowserIdx + 1));
+		ContentBrowserTabIDs[BrowserIdx] = TabID;
+
+		const FText DefaultDisplayName = GetContentBrowserLabelWithIndex(BrowserIdx);
+
+		FTabSpawnerEntry& ContentBrowserTabSpawner = FGlobalTabmanager::Get()->RegisterNomadTabSpawner(TabID, FOnSpawnTab::CreateRaw(this, &FContentBrowserSingleton::SpawnContentBrowserTab, BrowserIdx))
+			.SetDisplayName(DefaultDisplayName)
+			.SetTooltipText(LOCTEXT("ContentBrowserMenuTooltipText", "Open a Content Browser tab."))
+			.SetGroup(ContentBrowserGroup)
+			.SetIcon(ContentBrowserIcon);
+
+		ContentBrowserTabs.Add(ContentBrowserTabSpawner.AsSpawnerEntry());
+	}
 
 	UToolMenu* ContentMenu = UToolMenus::Get()->ExtendMenu("LevelEditor.LevelEditorToolBar.AddQuickMenu");
 	FToolMenuSection& Section = ContentMenu->FindOrAddSection("Content");
@@ -972,33 +987,24 @@ void FContentBrowserSingleton::GetContentBrowserSubMenu(UToolMenu* Menu, TShared
 		FUIAction(FExecuteAction::CreateRaw(this, &FContentBrowserSingleton::FocusPrimaryContentBrowser, true), FCanExecuteAction())
 	);
 
-	auto AddSpawnerEntryToMenuSection = [](FToolMenuSection& InSection, FTabSpawnerEntry& InSpawnerNode, FName InTabID)
+	auto AddSpawnerEntryToMenuSection = [](FToolMenuSection& InSection, TSharedPtr<FTabSpawnerEntry> InSpawnerNode, FName InTabID)
 	{
 		InSection.AddMenuEntry(
 			InTabID,
-			InSpawnerNode.GetDisplayName().IsEmpty() ? FText::FromName(InTabID) : InSpawnerNode.GetDisplayName(),
-			InSpawnerNode.GetTooltipText(),
+			InSpawnerNode->GetDisplayName().IsEmpty() ? FText::FromName(InTabID) : InSpawnerNode->GetDisplayName(),
+			InSpawnerNode->GetTooltipText(),
 			FSlateIcon(),
-			FGlobalTabmanager::Get()->GetUIActionForTabSpawnerMenuEntry(InSpawnerNode.AsSpawnerEntry()),
+			FGlobalTabmanager::Get()->GetUIActionForTabSpawnerMenuEntry(InSpawnerNode),
 			EUserInterfaceActionType::Check
 		);
 	};
 
 	for (int32 BrowserIdx = 0; BrowserIdx < UE_ARRAY_COUNT(ContentBrowserTabIDs); BrowserIdx++)
 	{
-		const FName TabID = FName(*FString::Printf(TEXT("ContentBrowserTab%d"), BrowserIdx + 1));
-		ContentBrowserTabIDs[BrowserIdx] = TabID;
-
-		const FText DefaultDisplayName = GetContentBrowserLabelWithIndex(BrowserIdx);
-
-		FTabSpawnerEntry& ContentBrowserTabSpawner = FGlobalTabmanager::Get()->RegisterNomadTabSpawner(TabID, FOnSpawnTab::CreateRaw(this, &FContentBrowserSingleton::SpawnContentBrowserTab, BrowserIdx))
-			.SetDisplayName(DefaultDisplayName)
-			.SetTooltipText(LOCTEXT("ContentBrowserMenuTooltipText", "Open a Content Browser tab."))
-			.SetGroup(ContentBrowserGroup)
-			.SetIcon(ContentBrowserIcon);
-
-		AddSpawnerEntryToMenuSection(Section, ContentBrowserTabSpawner, TabID);
+		const FName TabID = ContentBrowserTabIDs[BrowserIdx];
+		AddSpawnerEntryToMenuSection(Section, ContentBrowserTabs[BrowserIdx], TabID);
 	}
+
 }
 
 #undef LOCTEXT_NAMESPACE
diff --git a/Engine/Source/Editor/ContentBrowser/Private/ContentBrowserSingleton.h b/Engine/Source/Editor/ContentBrowser/Private/ContentBrowserSingleton.h
--- a/Engine/Source/Editor/ContentBrowser/Private/ContentBrowserSingleton.h
+++ b/Engine/Source/Editor/ContentBrowser/Private/ContentBrowserSingleton.h
@@ -16,6 +16,7 @@ class SContentBrowser;
 class UFactory;
 class UToolMenu;
 class FWorkspaceItem;
+struct FTabSpawnerEntry;
 
 #define MAX_CONTENT_BROWSERS 4
 
@@ -154,6 +155,8 @@ private:
 
 	TArray<FContentBrowserPluginSettings> PluginSettings;
 
+	TArray<TSharedPtr<FTabSpawnerEntry>> ContentBrowserTabs;
+
 	/** An incrementing int32 which is used when making unique settings strings */
 	int32 SettingsStringID;
 };
