Fix Ctrl + Space input twice needed to dismiss content browser.

diff --git a/Engine/Source/Editor/ContentBrowser/Private/SContentBrowser.cpp b/Engine/Source/Editor/ContentBrowser/Private/SContentBrowser.cpp
--- a/Engine/Source/Editor/ContentBrowser/Private/SContentBrowser.cpp
+++ b/Engine/Source/Editor/ContentBrowser/Private/SContentBrowser.cpp
@@ -68,6 +68,7 @@
 #include "IContentBrowserDataModule.h"
 #include "ContentBrowserDataSource.h"
 #include "ContentBrowserDataSubsystem.h"
+#include "StatusBarSubsystem.h"
 #include "Brushes/SlateColorBrush.h"
 #include "ToolMenu.h"
 #include "Widgets/Input/SExpandableButton.h"
@@ -1822,7 +1823,14 @@ FReply SContentBrowser::OnSearchKeyDown(const FGeometry& Geometry, const FKeyEve
 	// Clear focus if the content browser drawer key is clicked so it will close the opened content browser
 	if (FGlobalEditorCommonCommands::Get().OpenContentBrowserDrawer->HasActiveChord(CheckChord))
 	{
-		return FReply::Handled().ClearUserFocus(EFocusCause::SetDirectly);
+		FReply Reply = FReply::Handled().ClearUserFocus(EFocusCause::SetDirectly);
+		
+		// If we are a drawer, re-attempting to open will toggle. We must be open as we received input.
+		if (bIsDrawer)
+		{
+			GEditor->GetEditorSubsystem<UStatusBarSubsystem>()->OpenContentBrowserDrawer();
+		}
+		return Reply;
 	}
 
 	return FReply::Unhandled();
