Fix content browser path view filters submenu

diff --git a/Engine/Source/Editor/ContentBrowser/Private/SContentBrowser.cpp b/Engine/Source/Editor/ContentBrowser/Private/SContentBrowser.cpp
--- a/Engine/Source/Editor/ContentBrowser/Private/SContentBrowser.cpp
+++ b/Engine/Source/Editor/ContentBrowser/Private/SContentBrowser.cpp
@@ -824,6 +824,7 @@ TSharedRef<SWidget> SContentBrowser::CreatePathView(const FContentBrowserConfig*
 			.AllowClassesFolder(true)
 			.AddMetaData<FTagMetaData>(FTagMetaData(TEXT("ContentBrowserSources")))
 			.ExternalSearch(SourcesSearch)
+			.PluginPathFilters(PluginPathFilters)
 		];
 }
 
@@ -2331,9 +2332,16 @@ void SContentBrowser::RegisterPathViewFiltersMenu()
 		UToolMenu* Menu = UToolMenus::Get()->RegisterMenu(PathViewFiltersMenuName);
 		Menu->AddDynamicSection("DynamicContent", FNewToolMenuDelegate::CreateLambda([](UToolMenu* InMenu)
 		{
-			if (const UContentBrowserMenuContext* ContextObject = InMenu->FindContext<UContentBrowserMenuContext>())
+			if (const UContentBrowserAssetViewContextMenuContext* AssetViewContext = InMenu->FindContext<UContentBrowserAssetViewContextMenuContext>())
 			{
-				if (TSharedPtr<SContentBrowser> ContentBrowser = ContextObject->ContentBrowser.Pin())
+				if (TSharedPtr<SContentBrowser> ContentBrowser = AssetViewContext->OwningContentBrowser.Pin())
+				{
+					ContentBrowser->PopulatePathViewFiltersMenu(InMenu);
+				}
+			}
+			else if (const UContentBrowserMenuContext* ContentBrowserContext = InMenu->FindContext<UContentBrowserMenuContext>())
+			{
+				if (TSharedPtr<SContentBrowser> ContentBrowser = ContentBrowserContext->ContentBrowser.Pin())
 				{
 					ContentBrowser->PopulatePathViewFiltersMenu(InMenu);
 				}
