Fixes AssetView not keeping Source Control state icon when switching view. SCCStateWidget was not constructed yet during Source Control State initialization

#changelist validated

diff --git a/Engine/Source/Editor/ContentBrowser/Private/AssetViewWidgets.cpp b/Engine/Source/Editor/ContentBrowser/Private/AssetViewWidgets.cpp
--- a/Engine/Source/Editor/ContentBrowser/Private/AssetViewWidgets.cpp
+++ b/Engine/Source/Editor/ContentBrowser/Private/AssetViewWidgets.cpp
@@ -256,8 +256,10 @@ TSharedRef<SWidget> FAssetViewItemHelper::CreateListTileItemContents(T* const In
 			.SmallView(!InTileOrListItem->CanDisplayPrimitiveTools())
 			.Visibility(InTileOrListItem, &T::GetThumbnailEditModeUIVisibility)
 		];
-	}
 
+		InTileOrListItem->HandleSourceControlStateChanged();
+	}
+	
 	return ItemContentsOverlay;
 }
 
@@ -2236,6 +2238,8 @@ TSharedRef<SWidget> SAssetColumnItem::GenerateWidgetForColumn( const FName& Colu
 			.Text( TAttribute<FText>::Create( TAttribute<FText>::FGetter::CreateSP(this, &SAssetColumnItem::GetAssetTagText, ColumnName) ) );
 	}
 
+	HandleSourceControlStateChanged();
+
 	return SNew(SBox)
 		.Padding( ColumnItemPadding )
 		.VAlign(VAlign_Center)
