Move header search boxes in content browser's left-hand tree view to second line and adjust spacing

This gives them a bit more space and prevents them from being squished when the left-hand tree view is narrow

diff --git a/Engine/Source/Editor/ContentBrowser/Private/SContentBrowser.cpp b/Engine/Source/Editor/ContentBrowser/Private/SContentBrowser.cpp
--- a/Engine/Source/Editor/ContentBrowser/Private/SContentBrowser.cpp
+++ b/Engine/Source/Editor/ContentBrowser/Private/SContentBrowser.cpp
@@ -72,7 +72,7 @@
 #include "Brushes/SlateColorBrush.h"
 #include "ToolMenu.h"
 #include "Widgets/Input/SExpandableButton.h"
-#include "SExpandableSearchArea.h"
+#include "SSearchToggleButton.h"
 #include "SEditorHeaderButton.h"
 
 #define LOCTEXT_NAMESPACE "ContentBrowser"
@@ -145,6 +145,10 @@ void SContentBrowser::Construct( const FArguments& InArgs, const FName& InInstan
 
 	PluginPathFilters = MakeShareable(new FPluginFilterCollectionType());
 
+	FavoritesSearch = MakeShared<FSourcesSearch>();
+	FavoritesSearch->Initialize();
+	FavoritesSearch->SetHintText(LOCTEXT("SearchFavoritesHint", "Search Favorites"));
+
 	SourcesSearch = MakeShared<FSourcesSearch>();
 	SourcesSearch->Initialize();
 	SourcesSearch->SetHintText(LOCTEXT("SearchPathsHint", "Search Paths"));
@@ -759,28 +763,57 @@ TSharedRef<SWidget> SContentBrowser::CreateFavoritesView(const FContentBrowserCo
 		SAssignNew(FavoritesArea, SExpandableArea)
 		.BorderImage(FAppStyle::Get().GetBrush("Brushes.Header"))
 		.BodyBorderImage(FAppStyle::Get().GetBrush("Brushes.Recessed"))
-		.HeaderPadding(FMargin(5.0f, 7.0f))
-		.Visibility(this, &SContentBrowser::GetFavoriteFolderVisibility)
+		.HeaderPadding(FMargin(4.0f, 2.0f))
+		.Padding(0)
 		.AllowAnimatedTransition(false)
+		.OnAreaExpansionChanged_Lambda([this](bool bIsExpanded) { if (!bIsExpanded) FavoritesSearchToggleButton->SetExpanded(false); })
 		.HeaderContent()
 		[
-			SNew(STextBlock)
-			.Text(LOCTEXT("Favorites", "Favorites"))
-			.TextStyle(FAppStyle::Get(), "ButtonText")
-			.Font(FAppStyle::Get().GetFontStyle("NormalFontBold"))
-			.TransformPolicy(ETextTransformPolicy::ToUpper)
+			SNew(SHorizontalBox)
+			+ SHorizontalBox::Slot()
+			.VAlign(VAlign_Center)
+			[
+				SNew(STextBlock)
+				.Text(LOCTEXT("Favorites", "Favorites"))
+				.TextStyle(FAppStyle::Get(), "ButtonText")
+				.Font(FAppStyle::Get().GetFontStyle("NormalFontBold"))
+			]
+			+SHorizontalBox::Slot()
+			.VAlign(VAlign_Center)
+			.HAlign(HAlign_Right)
+			.AutoWidth()
+			.Padding(4.0f, 0.0f)
+			[
+				SAssignNew(FavoritesSearchToggleButton, SSearchToggleButton, FavoritesSearch->GetWidget())
+				.OnSearchBoxShown_Lambda([this]() { FavoritesArea->SetExpanded(true); })
+			]
 		]
 		.BodyContent()
 		[
-			SAssignNew(FavoritePathViewPtr, SFavoritePathView)
-			.OnItemSelectionChanged(this, &SContentBrowser::OnItemSelectionChanged, EContentBrowserViewContext::FavoriteView)
-			.OnGetItemContextMenu(this, &SContentBrowser::GetItemContextMenu, EContentBrowserViewContext::FavoriteView)
-			.FocusSearchBoxWhenOpened(false)
-			.ShowTreeTitle(false)
-			.ShowSeparator(false)
-			.AllowClassesFolder(true)
-			.AddMetaData<FTagMetaData>(FTagMetaData(TEXT("ContentBrowserFavorites")))
-			.ExternalSearch(SourcesSearch)
+			SNew(SVerticalBox)
+			+ SVerticalBox::Slot()
+			.AutoHeight()
+			[
+				// Should blend in visually with the header but technically acts like part of the body
+				SNew(SBorder)
+				.BorderImage(FAppStyle::Get().GetBrush("Brushes.Header"))
+				.Padding(FMargin(4.0f, 2.0f))
+				[
+					FavoritesSearch->GetWidget()
+				]
+			]
+			+ SVerticalBox::Slot()
+			[
+				SAssignNew(FavoritePathViewPtr, SFavoritePathView)
+				.OnItemSelectionChanged(this, &SContentBrowser::OnItemSelectionChanged, EContentBrowserViewContext::FavoriteView)
+				.OnGetItemContextMenu(this, &SContentBrowser::GetItemContextMenu, EContentBrowserViewContext::FavoriteView)
+				.FocusSearchBoxWhenOpened(false)
+				.ShowTreeTitle(false)
+				.ShowSeparator(false)
+				.AllowClassesFolder(true)
+				.AddMetaData<FTagMetaData>(FTagMetaData(TEXT("ContentBrowserFavorites")))
+				.ExternalSearch(FavoritesSearch)
+			]
 		];
 }
 
@@ -790,41 +823,62 @@ TSharedRef<SWidget> SContentBrowser::CreatePathView(const FContentBrowserConfig*
 		SAssignNew(PathArea, SExpandableArea)
 		.BorderImage(FAppStyle::Get().GetBrush("Brushes.Header"))
 		.BodyBorderImage(FAppStyle::Get().GetBrush("Brushes.Recessed"))
-		.HeaderPadding(FMargin(5.0f, 3.0f))
+		.HeaderPadding(FMargin(4.0f, 2.0f))
+		.Padding(0)
 		.AllowAnimatedTransition(false)
+		.OnAreaExpansionChanged_Lambda([this](bool bIsExpanded) { if (!bIsExpanded) PathSearchToggleButton->SetExpanded(false); })
 		.HeaderContent()
 		[
 			SNew(SHorizontalBox)
 			+ SHorizontalBox::Slot()
 			.VAlign(VAlign_Center)
-			.AutoWidth()
 			[
 				SNew(STextBlock)
 				.Text(FText::FromString(FApp::GetProjectName()))
 				.TextStyle(FAppStyle::Get(), "ButtonText")
 				.Font(FAppStyle::Get().GetFontStyle("NormalFontBold"))
-				.TransformPolicy(ETextTransformPolicy::ToUpper)
 			]
 			+SHorizontalBox::Slot()
 			.VAlign(VAlign_Center)
 			.HAlign(HAlign_Right)
-			.Padding(5.0f, 0.0f)
+			.AutoWidth()
+			.Padding(4.0f, 0.0f)
 			[
-				SAssignNew(PathSearchArea, SExpandableSearchArea, SourcesSearch->GetWidget())
+				SAssignNew(PathSearchToggleButton, SSearchToggleButton, SourcesSearch->GetWidget())
+				.OnSearchBoxShown_Lambda([this]() { PathArea->SetExpanded(true); })
 			]
 		]
 		.BodyContent()
 		[
-			SAssignNew(PathViewPtr, SPathView)
-			.OnItemSelectionChanged(this, &SContentBrowser::OnItemSelectionChanged, EContentBrowserViewContext::PathView)
-			.OnGetItemContextMenu(this, &SContentBrowser::GetItemContextMenu, EContentBrowserViewContext::PathView)
-			.FocusSearchBoxWhenOpened(false)
-			.ShowTreeTitle(false)
-			.ShowSeparator(false)
-			.AllowClassesFolder(true)
-			.AddMetaData<FTagMetaData>(FTagMetaData(TEXT("ContentBrowserSources")))
-			.ExternalSearch(SourcesSearch)
-			.PluginPathFilters(PluginPathFilters)
+			SNew(SVerticalBox)
+			+ SVerticalBox::Slot()
+			.AutoHeight()
+			[
+				// Should blend in visually with the header but technically acts like part of the body
+				SNew(SBorder)
+				.BorderImage(FAppStyle::Get().GetBrush("Brushes.Header"))
+				.Padding(FMargin(4.0f, 2.0f))
+				[
+					SourcesSearch->GetWidget()
+				]
+			]
+			+ SVerticalBox::Slot()
+			[
+				SNew(SBorder)
+				.BorderImage(FAppStyle::Get().GetBrush("Brushes.Recessed"))
+				[
+					SAssignNew(PathViewPtr, SPathView)
+					.OnItemSelectionChanged(this, &SContentBrowser::OnItemSelectionChanged, EContentBrowserViewContext::PathView)
+					.OnGetItemContextMenu(this, &SContentBrowser::GetItemContextMenu, EContentBrowserViewContext::PathView)
+					.FocusSearchBoxWhenOpened(false)
+					.ShowTreeTitle(false)
+					.ShowSeparator(false)
+					.AllowClassesFolder(true)
+					.AddMetaData<FTagMetaData>(FTagMetaData(TEXT("ContentBrowserSources")))
+					.ExternalSearch(SourcesSearch)
+					.PluginPathFilters(PluginPathFilters)
+				]
+			]
 		];
 }
 
@@ -834,27 +888,27 @@ TSharedRef<SWidget> SContentBrowser::CreateDockedCollectionsView(const FContentB
 		SAssignNew(CollectionArea, SExpandableArea)
 		.BorderImage(FAppStyle::Get().GetBrush("Brushes.Header"))
 		.BodyBorderImage(FAppStyle::Get().GetBrush("Brushes.Recessed"))
-		.HeaderPadding(FMargin(5.0f, 3.0f))
+		.HeaderPadding(FMargin(4.0f, 2.0f))
+		.Padding(0)
 		.Visibility(this, &SContentBrowser::GetDockedCollectionsVisibility)
 		.AllowAnimatedTransition(false)
+		.OnAreaExpansionChanged_Lambda([this](bool bIsExpanded) { if (!bIsExpanded) CollectionSearchToggleButton->SetExpanded(false); })
 		.HeaderContent()
 		[
 			SNew(SHorizontalBox)
 			+ SHorizontalBox::Slot()
 			.VAlign(VAlign_Center)
-			.AutoWidth()
 			[
 				SNew(STextBlock)
 				.Text(LOCTEXT("CollectionsTitle", "Collections"))
 				.TextStyle(FAppStyle::Get(), "ButtonText")
 				.Font(FAppStyle::Get().GetFontStyle("NormalFontBold"))
-				.TransformPolicy(ETextTransformPolicy::ToUpper)
 			]
 			+ SHorizontalBox::Slot()
 			.VAlign(VAlign_Center)
 			.HAlign(HAlign_Right)
 			.AutoWidth()
-			.Padding(5.0f, 0.0f)
+			.Padding(4.0f, 0.0f, 0.0f, 0.0f)
 			[
 				SNew(SButton)
 				.ButtonStyle(FEditorStyle::Get(), "SimpleButton")
@@ -870,14 +924,35 @@ TSharedRef<SWidget> SContentBrowser::CreateDockedCollectionsView(const FContentB
 			+ SHorizontalBox::Slot()
 			.VAlign(VAlign_Center)
 			.HAlign(HAlign_Right)
-			.Padding(5.0f, 0.0f)
+			.AutoWidth()
+			.Padding(4.0f, 0.0f)
 			[
-				SAssignNew(CollectionSearchArea, SExpandableSearchArea, CollectionSearch->GetWidget())
+				SAssignNew(CollectionSearchToggleButton, SSearchToggleButton, CollectionSearch->GetWidget())
+				.OnSearchBoxShown_Lambda([this]() { CollectionArea->SetExpanded(true); })
 			]
 		]
 		.BodyContent()
 		[
-			CollectionViewPtr.ToSharedRef()
+			SNew(SVerticalBox)
+			+ SVerticalBox::Slot()
+			.AutoHeight()
+			[
+				// Should blend in visually with the header but technically acts like part of the body
+				SNew(SBorder)
+				.BorderImage(FAppStyle::Get().GetBrush("Brushes.Header"))
+				.Padding(FMargin(4.0f, 2.0f))
+				[
+					CollectionSearch->GetWidget()
+				]
+			]
+			+ SVerticalBox::Slot()
+			[
+				SNew(SBorder)
+				.BorderImage(FAppStyle::Get().GetBrush("Brushes.Recessed"))
+				[
+					CollectionViewPtr.ToSharedRef()
+				]
+			]
 		];
 }
 
@@ -1350,8 +1425,9 @@ void SContentBrowser::SaveSettings() const
 	GConfig->SetBool(*SettingsIniSection, *(SettingsString + TEXT(".PathAreaExpanded")), PathArea->IsExpanded(), GEditorPerProjectIni);
 	GConfig->SetBool(*SettingsIniSection, *(SettingsString + TEXT(".CollectionAreaExpanded")), CollectionArea->IsExpanded(), GEditorPerProjectIni);
 
-	GConfig->SetBool(*SettingsIniSection, *(SettingsString + TEXT(".PathSearchAreaExpanded")), PathSearchArea->IsExpanded(), GEditorPerProjectIni);
-	GConfig->SetBool(*SettingsIniSection, *(SettingsString + TEXT(".CollectionSearchAreaExpanded")), CollectionSearchArea->IsExpanded(), GEditorPerProjectIni);
+	GConfig->SetBool(*SettingsIniSection, *(SettingsString + TEXT(".FavoritesSearchAreaExpanded")), FavoritesSearchToggleButton->IsExpanded(), GEditorPerProjectIni);
+	GConfig->SetBool(*SettingsIniSection, *(SettingsString + TEXT(".PathSearchAreaExpanded")), PathSearchToggleButton->IsExpanded(), GEditorPerProjectIni);
+	GConfig->SetBool(*SettingsIniSection, *(SettingsString + TEXT(".CollectionSearchAreaExpanded")), CollectionSearchToggleButton->IsExpanded(), GEditorPerProjectIni);
 
 	for(int32 SlotIndex = 0; SlotIndex < PathAssetSplitterPtr->GetChildren()->Num(); SlotIndex++)
 	{
@@ -1491,13 +1567,17 @@ void SContentBrowser::LoadSettings(const FName& InInstanceName)
 	GConfig->GetBool(*SettingsIniSection, *(SettingsString + TEXT(".CollectionAreaExpanded")), bCollectionAreaExpanded, GEditorPerProjectIni);
 	CollectionArea->SetExpanded(bCollectionAreaExpanded);
 
+	bool bFavoritesSearchAreaExpanded = false;
+	GConfig->GetBool(*SettingsIniSection, *(SettingsString + TEXT(".FavoritesSearchAreaExpanded")), bFavoritesSearchAreaExpanded, GEditorPerProjectIni);
+	FavoritesSearchToggleButton->SetExpanded(bFavoritesSearchAreaExpanded);
+
 	bool bPathSearchAreaExpanded = false;
 	GConfig->GetBool(*SettingsIniSection, *(SettingsString + TEXT(".PathSearchAreaExpanded")), bPathSearchAreaExpanded, GEditorPerProjectIni);
-	PathSearchArea->SetExpanded(bPathSearchAreaExpanded);
+	PathSearchToggleButton->SetExpanded(bPathSearchAreaExpanded);
 
 	bool bCollectionSearchAreaExpanded = false;
 	GConfig->GetBool(*SettingsIniSection, *(SettingsString + TEXT(".CollectionSearchAreaExpanded")), bCollectionSearchAreaExpanded, GEditorPerProjectIni);
-	CollectionSearchArea->SetExpanded(bCollectionSearchAreaExpanded);
+	CollectionSearchToggleButton->SetExpanded(bCollectionSearchAreaExpanded);
 
 	for(int32 SlotIndex = 0; SlotIndex < PathAssetSplitterPtr->GetChildren()->Num(); SlotIndex++)
 	{
diff --git a/Engine/Source/Editor/ContentBrowser/Private/SContentBrowser.h b/Engine/Source/Editor/ContentBrowser/Private/SContentBrowser.h
--- a/Engine/Source/Editor/ContentBrowser/Private/SContentBrowser.h
+++ b/Engine/Source/Editor/ContentBrowser/Private/SContentBrowser.h
@@ -32,7 +32,7 @@ class SPathView;
 class SSplitter;
 class UFactory;
 class UToolMenu;
-class SExpandableSearchArea;
+class SSearchToggleButton;
 class UContentBrowserToolbarMenuContext;
 
 struct FToolMenuContext;
@@ -556,6 +556,9 @@ private:
 	/** The context menu manager for the path view */
 	TSharedPtr<class FPathContextMenu> PathContextMenu;
 
+	/** The sources search for favorites */
+	TSharedPtr<FSourcesSearch> FavoritesSearch;
+
 	/** The sources search for paths */
 	TSharedPtr<FSourcesSearch> SourcesSearch;
 
@@ -589,17 +592,20 @@ private:
 	/** Favorites area widget */
 	TSharedPtr<SExpandableArea> FavoritesArea;
 
+	/** Toggle button for showing/hiding favorites search area */
+	TSharedPtr<SSearchToggleButton> FavoritesSearchToggleButton;
+
 	/** Path area widget */
 	TSharedPtr<SExpandableArea> PathArea;
 
-	/** Search area for paths */
-	TSharedPtr<SExpandableSearchArea> PathSearchArea;
+	/** Toggle button for showing/hiding path search area */
+	TSharedPtr<SSearchToggleButton> PathSearchToggleButton;
 
 	/** Collection area widget */
 	TSharedPtr<SExpandableArea> CollectionArea;
 
-	/** Search area for collections */
-	TSharedPtr<SExpandableSearchArea> CollectionSearchArea;
+	/** Toggle button for showing/hiding collection search area */
+	TSharedPtr<SSearchToggleButton> CollectionSearchToggleButton;
 
 	/** Index of the active sources widget */
 	int32 ActiveSourcesWidgetIndex = 0;
diff --git a/Engine/Source/Editor/ContentBrowser/Private/SExpandableSearchArea.cpp b/Engine/Source/Editor/ContentBrowser/Private/SExpandableSearchArea.cpp
deleted file mode 100644
diff --git a/Engine/Source/Editor/ContentBrowser/Private/SExpandableSearchArea.h b/Engine/Source/Editor/ContentBrowser/Private/SExpandableSearchArea.h
deleted file mode 100644
diff --git a/Engine/Source/Editor/ContentBrowser/Private/SSearchToggleButton.cpp b/Engine/Source/Editor/ContentBrowser/Private/SSearchToggleButton.cpp
new file mode 100644
--- /dev/null
+++ b/Engine/Source/Editor/ContentBrowser/Private/SSearchToggleButton.cpp
@@ -0,0 +1,71 @@
+// Copyright Epic Games, Inc. All Rights Reserved.
+
+#include "SSearchToggleButton.h"
+#include "Widgets/SBoxPanel.h"
+#include "Widgets/Images/SImage.h"
+#include "Widgets/Input/SButton.h"
+#include "Widgets/Input/SSearchBox.h"
+#include "Widgets/Layout/SBox.h"
+#include "Widgets/Layout/SExpandableArea.h"
+
+void SSearchToggleButton::Construct(const FArguments& InArgs, TSharedRef<SSearchBox> SearchBox)
+{
+	bIsExpanded = false;
+	OnSearchBoxShown = InArgs._OnSearchBoxShown;
+
+	SearchStyle = InArgs._Style;
+
+	SearchBox->SetVisibility(TAttribute<EVisibility>::CreateSP(this, &SSearchToggleButton::GetSearchBoxVisibility));
+	SearchBoxPtr = SearchBox;
+
+	ChildSlot
+	[
+		SNew(SCheckBox)
+		.IsChecked(this, &SSearchToggleButton::GetToggleButtonState)
+		.OnCheckStateChanged(this, &SSearchToggleButton::OnToggleButtonStateChanged)
+		.Style(FAppStyle::Get(), "ToggleButtonCheckbox")
+		.Padding(4.0f)
+		.ToolTipText(NSLOCTEXT("ExpandableSearchArea", "ExpandCollapseSearchButton", "Expands or collapses the search text box"))
+		[
+			SNew(SImage)
+			.Image(&SearchStyle->GlassImage)
+			.ColorAndOpacity(FSlateColor::UseForeground())
+		]
+	];
+}
+
+void SSearchToggleButton::SetExpanded(bool bInExpanded)
+{
+	bIsExpanded = bInExpanded;
+
+	if (TSharedPtr<SSearchBox> SearchBox = SearchBoxPtr.Pin())
+	{
+		if (bIsExpanded)
+		{
+			OnSearchBoxShown.ExecuteIfBound();
+
+			// Focus the search box when it's shown
+			FSlateApplication::Get().SetUserFocus(FSlateApplication::Get().GetUserIndexForKeyboard(), SearchBox, EFocusCause::SetDirectly);
+		}
+		else
+		{
+			// Clear the search box when it's hidden
+			SearchBox->SetText(FText::GetEmpty());
+		}
+	}
+}
+
+void SSearchToggleButton::OnToggleButtonStateChanged(ECheckBoxState CheckBoxState)
+{
+	SetExpanded(CheckBoxState == ECheckBoxState::Checked);
+}
+
+ECheckBoxState SSearchToggleButton::GetToggleButtonState() const
+{
+	return bIsExpanded ? ECheckBoxState::Checked : ECheckBoxState::Unchecked;
+}
+
+EVisibility SSearchToggleButton::GetSearchBoxVisibility() const
+{
+	return bIsExpanded ? EVisibility::Visible : EVisibility::Collapsed;
+}
diff --git a/Engine/Source/Editor/ContentBrowser/Private/SSearchToggleButton.h b/Engine/Source/Editor/ContentBrowser/Private/SSearchToggleButton.h
new file mode 100644
--- /dev/null
+++ b/Engine/Source/Editor/ContentBrowser/Private/SSearchToggleButton.h
@@ -0,0 +1,46 @@
+// Copyright Epic Games, Inc. All Rights Reserved.
+
+#pragma once 
+
+#include "Widgets/SCompoundWidget.h"
+#include "Widgets/DeclarativeSyntaxSupport.h"
+#include "Styling/SlateTypes.h"
+#include "Styling/AppStyle.h"
+#include "Styling/SlateWidgetStyleAsset.h"
+
+class SSearchBox;
+
+DECLARE_DELEGATE(FOnSearchBoxShown)
+
+class SSearchToggleButton : public SCompoundWidget
+{
+public:
+	SLATE_BEGIN_ARGS(SSearchToggleButton)
+		: _Style(&FAppStyle::Get().GetWidgetStyle<FSearchBoxStyle>("SearchBox"))
+	{}
+		/** Search box style (used to match the glass icon) */
+		SLATE_STYLE_ARGUMENT(FSearchBoxStyle, Style)
+
+		/** Event fired when the associated search box is made visible */
+		SLATE_EVENT(FOnSearchBoxShown, OnSearchBoxShown)
+
+	SLATE_END_ARGS()
+
+	void Construct(const FArguments& InArgs, TSharedRef<SSearchBox> SearchBox);
+
+	/** @return true if the search area is expanded and the search box exposed */
+	bool IsExpanded() const { return bIsExpanded; }
+
+	/** Sets whether or not the search area is expanded to expose the search box */
+	void SetExpanded(bool bInExpanded);
+private:
+
+	ECheckBoxState GetToggleButtonState() const;
+	void OnToggleButtonStateChanged(ECheckBoxState CheckBoxState);
+	EVisibility GetSearchBoxVisibility() const;
+private:
+	const FSearchBoxStyle* SearchStyle;
+	bool bIsExpanded;
+	FOnSearchBoxShown OnSearchBoxShown;
+	TWeakPtr<SSearchBox> SearchBoxPtr;
+};
