Fix non-unity build/missing includes in SSearchToggleButton

diff --git a/Engine/Source/Editor/ContentBrowser/Private/SSearchToggleButton.cpp b/Engine/Source/Editor/ContentBrowser/Private/SSearchToggleButton.cpp
--- a/Engine/Source/Editor/ContentBrowser/Private/SSearchToggleButton.cpp
+++ b/Engine/Source/Editor/ContentBrowser/Private/SSearchToggleButton.cpp
@@ -1,9 +1,11 @@
 // Copyright Epic Games, Inc. All Rights Reserved.
 
 #include "SSearchToggleButton.h"
+#include "Framework/Application/SlateApplication.h"
 #include "Widgets/SBoxPanel.h"
 #include "Widgets/Images/SImage.h"
 #include "Widgets/Input/SButton.h"
+#include "Widgets/Input/SCheckBox.h"
 #include "Widgets/Input/SSearchBox.h"
 #include "Widgets/Layout/SBox.h"
 #include "Widgets/Layout/SExpandableArea.h"
