Performance improvements to widget library.
Add callback for source data generation to asset views.
Add UObject enumeration to content browser API (By [REDACTED_PERSON]).
Intentionally reserve thumbnail pool sizes instead of flat 1024 reservation.

diff --git a/Engine/Plugins/Editor/ContentBrowser/ContentBrowserAssetDataSource/Source/ContentBrowserAssetDataSource/Private/ContentBrowserAssetDataSource.cpp b/Engine/Plugins/Editor/ContentBrowser/ContentBrowserAssetDataSource/Source/ContentBrowserAssetDataSource/Private/ContentBrowserAssetDataSource.cpp
--- a/Engine/Plugins/Editor/ContentBrowser/ContentBrowserAssetDataSource/Source/ContentBrowserAssetDataSource/Private/ContentBrowserAssetDataSource.cpp
+++ b/Engine/Plugins/Editor/ContentBrowser/ContentBrowserAssetDataSource/Source/ContentBrowserAssetDataSource/Private/ContentBrowserAssetDataSource.cpp
@@ -945,6 +945,58 @@ void UContentBrowserAssetDataSource::EnumerateItemsAtPath(const FName InPath, co
 	}
 }
 
+bool UContentBrowserAssetDataSource::EnumerateItemsAtPaths(const TArrayView<FContentBrowserItemPath> InPaths, const EContentBrowserItemTypeFilter InItemTypeFilter, TFunctionRef<bool(FContentBrowserItemData&&)> InCallback)
+{
+	if (EnumHasAnyFlags(InItemTypeFilter, EContentBrowserItemTypeFilter::IncludeFolders))
+	{
+		for (const FContentBrowserItemPath& InPath : InPaths)
+		{
+			if (InPath.HasInternalPath())
+			{
+				if (AssetRegistry->PathExists(InPath.GetInternalPathName()))
+				{
+					if (!InCallback(CreateAssetFolderItem(InPath.GetInternalPathName())))
+					{
+						return false;
+					}
+				}
+			}
+		}
+	}
+
+	if (EnumHasAnyFlags(InItemTypeFilter, EContentBrowserItemTypeFilter::IncludeFiles) && !InPaths.IsEmpty())
+	{
+		FARFilter ARFilter;
+		
+		// TODO: EnumerateAssets for in memory assets needs optimization, currently enumerates every UObject in memory instead of calling find
+		ARFilter.bIncludeOnlyOnDiskAssets = true;
+
+		for (const FContentBrowserItemPath& InPath : InPaths)
+		{
+			if (InPath.HasInternalPath())
+			{
+				ARFilter.PackageNames.Add(InPath.GetInternalPathName());
+			}
+		}
+
+		auto FileFoundCallback = [this, &InCallback](const FAssetData& AssetData)
+		{
+			if (ContentBrowserAssetData::IsPrimaryAsset(AssetData))
+			{
+				return InCallback(CreateAssetFileItem(AssetData));
+			}
+			return true;
+		};
+
+		if (!AssetRegistry->EnumerateAssets(ARFilter, FileFoundCallback))
+		{
+			return false;
+		}
+	}
+
+	return true;
+}
+
 bool UContentBrowserAssetDataSource::IsDiscoveringItems(FText* OutStatus)
 {
 	if (AssetRegistry->IsLoadingAssets())
diff --git a/Engine/Plugins/Editor/ContentBrowser/ContentBrowserAssetDataSource/Source/ContentBrowserAssetDataSource/Public/ContentBrowserAssetDataSource.h b/Engine/Plugins/Editor/ContentBrowser/ContentBrowserAssetDataSource/Source/ContentBrowserAssetDataSource/Public/ContentBrowserAssetDataSource.h
--- a/Engine/Plugins/Editor/ContentBrowser/ContentBrowserAssetDataSource/Source/ContentBrowserAssetDataSource/Public/ContentBrowserAssetDataSource.h
+++ b/Engine/Plugins/Editor/ContentBrowser/ContentBrowserAssetDataSource/Source/ContentBrowserAssetDataSource/Public/ContentBrowserAssetDataSource.h
@@ -160,6 +160,8 @@ public:
 
 	virtual void EnumerateItemsAtPath(const FName InPath, const EContentBrowserItemTypeFilter InItemTypeFilter, TFunctionRef<bool(FContentBrowserItemData&&)> InCallback) override;
 
+	virtual bool EnumerateItemsAtPaths(const TArrayView<FContentBrowserItemPath> InPaths, const EContentBrowserItemTypeFilter InItemTypeFilter, TFunctionRef<bool(FContentBrowserItemData&&)> InCallback) override;
+
 	virtual bool IsDiscoveringItems(FText* OutStatus = nullptr) override;
 
 	virtual bool PrioritizeSearchPath(const FName InPath) override;
diff --git a/Engine/Plugins/Editor/ContentBrowser/ContentBrowserClassDataSource/Source/ContentBrowserClassDataSource/Private/ContentBrowserClassDataSource.cpp b/Engine/Plugins/Editor/ContentBrowser/ContentBrowserClassDataSource/Source/ContentBrowserClassDataSource/Private/ContentBrowserClassDataSource.cpp
--- a/Engine/Plugins/Editor/ContentBrowser/ContentBrowserClassDataSource/Source/ContentBrowserClassDataSource/Private/ContentBrowserClassDataSource.cpp
+++ b/Engine/Plugins/Editor/ContentBrowser/ContentBrowserClassDataSource/Source/ContentBrowserClassDataSource/Private/ContentBrowserClassDataSource.cpp
@@ -16,6 +16,7 @@
 #include "Framework/Docking/TabManager.h"
 #include "ContentBrowserDataSubsystem.h"
 #include "ContentBrowserDataUtils.h"
+#include "ContentBrowserItemPath.h"
 
 #define LOCTEXT_NAMESPACE "ContentBrowserClassDataSource"
 
@@ -374,6 +375,30 @@ void UContentBrowserClassDataSource::EnumerateItemsAtPath(const FName InPath, co
 	}
 }
 
+bool UContentBrowserClassDataSource::EnumerateItemsForObjects(const TArrayView<UObject*> InObjects, TFunctionRef<bool(FContentBrowserItemData&&)> InCallback)
+{
+	ConditionalCreateNativeClassHierarchy();
+
+	FString InternalPath;
+	for (UObject* InObject : InObjects)
+	{
+		if (UClass* InClass = Cast<UClass>(InObject))
+		{
+			InternalPath.Reset();
+			if (NativeClassHierarchy->GetClassPath(InClass, InternalPath, NativeClassHierarchyGetClassPathCache))
+			{
+				const FContentBrowserItemPath ContentBrowserItemPath(InternalPath, EContentBrowserPathType::Internal);
+				if (!InCallback(ContentBrowserClassData::CreateClassFileItem(this, ContentBrowserItemPath.GetVirtualPathName(), FName(InternalPath), InClass)))
+				{
+					return false;
+				}
+			}
+		}
+	}
+
+	return true;
+}
+
 bool UContentBrowserClassDataSource::IsFolderVisibleIfHidingEmpty(const FName InPath)
 {
 	FName ConvertedPath;
@@ -643,6 +668,7 @@ void UContentBrowserClassDataSource::ConditionalCreateNativeClassHierarchy()
 
 void UContentBrowserClassDataSource::ClassHierarchyUpdated()
 {
+	NativeClassHierarchyGetClassPathCache.Reset();
 	SetVirtualPathTreeNeedsRebuild();
 	NotifyItemDataRefreshed();
 }
diff --git a/Engine/Plugins/Editor/ContentBrowser/ContentBrowserClassDataSource/Source/ContentBrowserClassDataSource/Private/NativeClassHierarchy.cpp b/Engine/Plugins/Editor/ContentBrowser/ContentBrowserClassDataSource/Source/ContentBrowserClassDataSource/Private/NativeClassHierarchy.cpp
--- a/Engine/Plugins/Editor/ContentBrowser/ContentBrowserClassDataSource/Source/ContentBrowserClassDataSource/Private/NativeClassHierarchy.cpp
+++ b/Engine/Plugins/Editor/ContentBrowser/ContentBrowserClassDataSource/Source/ContentBrowserClassDataSource/Private/NativeClassHierarchy.cpp
@@ -66,7 +66,7 @@ FNativeClassHierarchy::~FNativeClassHierarchy()
 TSharedPtr<const FNativeClassHierarchyNode> FNativeClassHierarchy::FindNode(const FName InClassPath, const ENativeClassHierarchyNodeType InType) const
 {
 	TArray<TSharedRef<FNativeClassHierarchyNode>, TInlineAllocator<4>> NodesFound;
-	GatherMatchingNodesForPaths(TArrayView<const FName>(&InClassPath, 1), NodesFound);
+	GatherMatchingNodesForPaths(TArrayView<const FName>(&InClassPath, 1), NodesFound, InType);
 
 	for(const auto& NodeFound : NodesFound)
 	{
@@ -263,7 +263,7 @@ void FNativeClassHierarchy::GetFoldersRecursive(const TSharedRef<FNativeClassHie
 	}
 }
 
-void FNativeClassHierarchy::GatherMatchingNodesForPaths(const TArrayView<const FName>& InClassPaths, TArray<TSharedRef<FNativeClassHierarchyNode>, TInlineAllocator<4>>& OutMatchingNodes) const
+void FNativeClassHierarchy::GatherMatchingNodesForPaths(const TArrayView<const FName>& InClassPaths, TArray<TSharedRef<FNativeClassHierarchyNode>, TInlineAllocator<4>>& OutMatchingNodes, const ENativeClassHierarchyNodeType InType) const
 {
 	if(InClassPaths.Num() == 0)
 	{
@@ -286,7 +286,16 @@ void FNativeClassHierarchy::GatherMatchingNodesForPaths(const TArrayView<const F
 			{
 				// Try and find the node associated with this part of the path...
 				const FName ClassPathPartName = *ClassPathPart;
-				CurrentNode = (CurrentNode.IsValid()) ? CurrentNode->Children.FindRef(FNativeClassHierarchyNodeKey(ClassPathPartName, ENativeClassHierarchyNodeType::Folder)) : RootNodes.FindRef(ClassPathPartName);
+
+				if (InType == ENativeClassHierarchyNodeType::Class && ClassPathPart == ClassPathParts.Last())
+				{
+
+					CurrentNode = (CurrentNode.IsValid()) ? CurrentNode->Children.FindRef(FNativeClassHierarchyNodeKey(ClassPathPartName, ENativeClassHierarchyNodeType::Class)) : RootNodes.FindRef(ClassPathPartName);
+				}
+				else
+				{
+					CurrentNode = (CurrentNode.IsValid()) ? CurrentNode->Children.FindRef(FNativeClassHierarchyNodeKey(ClassPathPartName, ENativeClassHierarchyNodeType::Folder)) : RootNodes.FindRef(ClassPathPartName);
+				}
 
 				// ... bail out if we didn't find a valid node
 				if(!CurrentNode.IsValid())
@@ -502,7 +511,7 @@ bool FNativeClassHierarchy::GetFileSystemPath(const FString& InClassPath, FStrin
 	return false;
 }
 
-bool FNativeClassHierarchy::GetClassPath(UClass* InClass, FString& OutClassPath, FNativeClassHierarchyGetClassPathCache& InCache, const bool bIncludeClassName) const
+bool FNativeClassHierarchy::GetClassPath(const UClass* InClass, FString& OutClassPath, FNativeClassHierarchyGetClassPathCache& InCache, const bool bIncludeClassName) const
 {
 	const FName ClassModuleName = GetClassModuleName(InClass);
 	if(ClassModuleName.IsNone())
@@ -566,7 +575,7 @@ void FNativeClassHierarchy::OnReloadComplete(EReloadCompleteReason Reason)
 	PopulateHierarchy();
 }
 
-FName FNativeClassHierarchy::GetClassModuleName(UClass* InClass)
+FName FNativeClassHierarchy::GetClassModuleName(const UClass* InClass)
 {
 	UPackage* const ClassPackage = InClass->GetOuterUPackage();
 
diff --git a/Engine/Plugins/Editor/ContentBrowser/ContentBrowserClassDataSource/Source/ContentBrowserClassDataSource/Private/NativeClassHierarchy.h b/Engine/Plugins/Editor/ContentBrowser/ContentBrowserClassDataSource/Source/ContentBrowserClassDataSource/Private/NativeClassHierarchy.h
--- a/Engine/Plugins/Editor/ContentBrowser/ContentBrowserClassDataSource/Source/ContentBrowserClassDataSource/Private/NativeClassHierarchy.h
+++ b/Engine/Plugins/Editor/ContentBrowser/ContentBrowserClassDataSource/Source/ContentBrowserClassDataSource/Private/NativeClassHierarchy.h
@@ -36,6 +36,12 @@ struct FNativeClassHierarchyGetClassPathCache
 {
 	TSet<FName> GameModules;
 	TMap<FName, FNativeClassHierarchyPluginModuleInfo> PluginModules;
+
+	void Reset()
+	{
+		GameModules.Reset();
+		PluginModules.Reset();
+	}
 };
 
 /**
@@ -249,7 +255,7 @@ public:
 	 * 
 	 * @return true if the class path could be resolved and OutClassPath was filled in, false otherwise
 	 */
-	bool GetClassPath(UClass* InClass, FString& OutClassPath, FNativeClassHierarchyGetClassPathCache& InCache, const bool bIncludeClassName = true) const;
+	bool GetClassPath(const UClass* InClass, FString& OutClassPath, FNativeClassHierarchyGetClassPathCache& InCache, const bool bIncludeClassName = true) const;
 
 	/**
 	 * This will add a transient folder into the hierarchy
@@ -327,8 +333,9 @@ private:
 	 *
 	 * @param InClassPaths - The class paths to find the nodes for, or an empty list to get all root nodes
 	 * @param OutMatchingNodes - Array to be populated the nodes that correspond to the given class paths
+	 * @param InType - Type of node to search for, ie: class / folder
 	 */
-	void GatherMatchingNodesForPaths(const TArrayView<const FName>& InClassPaths, TArray<TSharedRef<FNativeClassHierarchyNode>, TInlineAllocator<4>>& OutMatchingNodes) const;
+	void GatherMatchingNodesForPaths(const TArrayView<const FName>& InClassPaths, TArray<TSharedRef<FNativeClassHierarchyNode>, TInlineAllocator<4>>& OutMatchingNodes, const ENativeClassHierarchyNodeType InType = ENativeClassHierarchyNodeType::Folder) const;
 
 	/**
 	 * Completely clear and re-populate the known class hierarchy
@@ -381,7 +388,7 @@ private:
 	 *
 	 * @return The name of the module that holds the class, eg) "CoreUObject"
 	 */
-	static FName GetClassModuleName(UClass* InClass);
+	static FName GetClassModuleName(const UClass* InClass);
 
 	/**
 	 * Given a module, work out which root path it should use as a parent
diff --git a/Engine/Plugins/Editor/ContentBrowser/ContentBrowserClassDataSource/Source/ContentBrowserClassDataSource/Public/ContentBrowserClassDataSource.h b/Engine/Plugins/Editor/ContentBrowser/ContentBrowserClassDataSource/Source/ContentBrowserClassDataSource/Public/ContentBrowserClassDataSource.h
--- a/Engine/Plugins/Editor/ContentBrowser/ContentBrowserClassDataSource/Source/ContentBrowserClassDataSource/Public/ContentBrowserClassDataSource.h
+++ b/Engine/Plugins/Editor/ContentBrowser/ContentBrowserClassDataSource/Source/ContentBrowserClassDataSource/Public/ContentBrowserClassDataSource.h
@@ -6,6 +6,7 @@
 #include "ContentBrowserDataSource.h"
 #include "ContentBrowserDataMenuContexts.h"
 #include "ContentBrowserClassDataPayload.h"
+#include "NativeClassHierarchy.h"
 #include "ContentBrowserClassDataSource.generated.h"
 
 class IAssetTypeActions;
@@ -13,7 +14,6 @@ class ICollectionManager;
 class UToolMenu;
 class FNativeClassHierarchy;
 struct FCollectionNameType;
-struct FNativeClassHierarchyGetClassPathCache;
 
 USTRUCT()
 struct CONTENTBROWSERCLASSDATASOURCE_API FContentBrowserCompiledClassDataFilter
@@ -44,6 +44,8 @@ public:
 
 	virtual void EnumerateItemsAtPath(const FName InPath, const EContentBrowserItemTypeFilter InItemTypeFilter, TFunctionRef<bool(FContentBrowserItemData&&)> InCallback) override;
 
+	virtual bool EnumerateItemsForObjects(const TArrayView<UObject*> InObjects, TFunctionRef<bool(FContentBrowserItemData&&)> InCallback) override;
+
 	virtual bool IsFolderVisibleIfHidingEmpty(const FName InPath) override;
 
 	virtual bool DoesItemPassFilter(const FContentBrowserItemData& InItem, const FContentBrowserDataCompiledFilter& InFilter) override;
@@ -100,6 +102,7 @@ private:
 	void ClassHierarchyUpdated();
 
 	TSharedPtr<FNativeClassHierarchy> NativeClassHierarchy;
+	FNativeClassHierarchyGetClassPathCache NativeClassHierarchyGetClassPathCache;
 
 	TSharedPtr<IAssetTypeActions> ClassTypeActions;
 
diff --git a/Engine/Source/Editor/ContentBrowser/Private/SAssetView.cpp b/Engine/Source/Editor/ContentBrowser/Private/SAssetView.cpp
--- a/Engine/Source/Editor/ContentBrowser/Private/SAssetView.cpp
+++ b/Engine/Source/Editor/ContentBrowser/Private/SAssetView.cpp
@@ -220,7 +220,7 @@ void SAssetView::Construct( const FArguments& InArgs )
 	ThumbnailScaleRangeScalar = ( DisplaySize.Y / 2160 );
 
 	// Create a thumbnail pool for rendering thumbnails	
-	AssetThumbnailPool = MakeShared<FAssetThumbnailPool>(1024);
+	AssetThumbnailPool = MakeShared<FAssetThumbnailPool>(InArgs._InitialThumbnailPoolSize);
 	NumOffscreenThumbnails = 64;
 	ListViewThumbnailResolution = 128;
 	ListViewThumbnailSize = 64;
@@ -1817,79 +1817,92 @@ void SAssetView::RefreshSourceItems()
 	ItemsPendingFrontendFilter.Reset();
 	{
 		UContentBrowserDataSubsystem* ContentBrowserData = IContentBrowserDataModule::Get().GetSubsystem();
-		const FContentBrowserDataFilter DataFilter = CreateBackendDataFilter();
-
-		bWereItemsRecursivelyFiltered = DataFilter.bRecursivePaths;
 
-		if (SourcesData.HasCollections() && EnumHasAnyFlags(DataFilter.ItemCategoryFilter, EContentBrowserItemCategoryFilter::IncludeCollections))
+		auto AddNewItem = [this, &PreviousAvailableBackendItems](FContentBrowserItemData&& InItemData)
 		{
-			// If we are showing collections then we may need to add dummy folder items for the child collections
-			// Note: We don't check the IncludeFolders flag here, as that is forced to false when collections are selected,
-			// instead we check the state of bIncludeChildCollections which will be false when we want to show collection folders
-			const FContentBrowserDataCollectionFilter* CollectionFilter = DataFilter.ExtraFilters.FindFilter<FContentBrowserDataCollectionFilter>();
-			if (CollectionFilter && !CollectionFilter->bIncludeChildCollections)
+			const FContentBrowserItemKey ItemDataKey(InItemData);
+			const uint32 ItemDataKeyHash = GetTypeHash(ItemDataKey);
+
+			TSharedPtr<FAssetViewItem>& NewItem = AvailableBackendItems.FindOrAddByHash(ItemDataKeyHash, ItemDataKey);
+			if (!NewItem && InItemData.IsFile())
 			{
-				FCollectionManagerModule& CollectionManagerModule = FCollectionManagerModule::GetModule();
-				
-				TArray<FCollectionNameType> ChildCollections;
-				for(const FCollectionNameType& Collection : SourcesData.Collections)
+				// Re-use the old view item where possible to avoid list churn when our backend view already included the item
+				if (TSharedPtr<FAssetViewItem>* PreviousItem = PreviousAvailableBackendItems.FindByHash(ItemDataKeyHash, ItemDataKey))
 				{
-					ChildCollections.Reset();
-					CollectionManagerModule.Get().GetChildCollections(Collection.Name, Collection.Type, ChildCollections);
-
-					for (const FCollectionNameType& ChildCollection : ChildCollections)
-					{
-						// Use "Collections" as the root of the path to avoid this being confused with other view folders - see ContentBrowserUtils::IsCollectionPath
-						FContentBrowserItemData FolderItemData(
-							nullptr, 
-							EContentBrowserItemFlags::Type_Folder | EContentBrowserItemFlags::Category_Collection, 
-							*FString::Printf(TEXT("/Collections/%s/%s"), ECollectionShareType::ToString(ChildCollection.Type), *ChildCollection.Name.ToString()), 
-							ChildCollection.Name, 
-							FText::FromName(ChildCollection.Name), 
-							nullptr
-							);
-
-						const FContentBrowserItemKey FolderItemDataKey(FolderItemData);
-						AvailableBackendItems.Add(FolderItemDataKey, MakeShared<FAssetViewItem>(MoveTemp(FolderItemData)));
-					}
+					NewItem = *PreviousItem;
+					NewItem->ClearCachedCustomColumns();
 				}
 			}
+			if (NewItem)
+			{
+				NewItem->AppendItemData(InItemData);
+				NewItem->CacheCustomColumns(CustomColumns, true, true, false /*bUpdateExisting*/);
+			}
+			else
+			{
+				NewItem = MakeShared<FAssetViewItem>(MoveTemp(InItemData));
+			}
+
+			return true;
+		};
+
+		if (SourcesData.OnEnumerateCustomSourceItemDatas.IsBound())
+		{
+			SourcesData.OnEnumerateCustomSourceItemDatas.Execute(AddNewItem);
 		}
 
-		static const FName RootPath = "/";
-		const TArrayView<const FName> DataSourcePaths = SourcesData.HasVirtualPaths() ? MakeArrayView(SourcesData.VirtualPaths) : MakeArrayView(&RootPath, 1);
-		for (const FName& DataSourcePath : DataSourcePaths)
+		if (SourcesData.IsIncludingVirtualPaths() || SourcesData.HasCollections()) 
 		{
-			// Ensure paths do not contain trailing slash
-			ensure(DataSourcePath == RootPath || !FStringView(FNameBuilder(DataSourcePath)).EndsWith(TEXT('/')));
+			const FContentBrowserDataFilter DataFilter = CreateBackendDataFilter();
+
+			bWereItemsRecursivelyFiltered = DataFilter.bRecursivePaths;
 
-			ContentBrowserData->EnumerateItemsUnderPath(DataSourcePath, DataFilter, [this, &PreviousAvailableBackendItems](FContentBrowserItemData&& InItemData)
+			if (SourcesData.HasCollections() && EnumHasAnyFlags(DataFilter.ItemCategoryFilter, EContentBrowserItemCategoryFilter::IncludeCollections))
 			{
-				const FContentBrowserItemKey ItemDataKey(InItemData);
-				const uint32 ItemDataKeyHash = GetTypeHash(ItemDataKey);
-				
-				TSharedPtr<FAssetViewItem>& NewItem = AvailableBackendItems.FindOrAddByHash(ItemDataKeyHash, ItemDataKey);
-				if (!NewItem && InItemData.IsFile())
+				// If we are showing collections then we may need to add dummy folder items for the child collections
+				// Note: We don't check the IncludeFolders flag here, as that is forced to false when collections are selected,
+				// instead we check the state of bIncludeChildCollections which will be false when we want to show collection folders
+				const FContentBrowserDataCollectionFilter* CollectionFilter = DataFilter.ExtraFilters.FindFilter<FContentBrowserDataCollectionFilter>();
+				if (CollectionFilter && !CollectionFilter->bIncludeChildCollections)
 				{
-					// Re-use the old view item where possible to avoid list churn when our backend view already included the item
-					if (TSharedPtr<FAssetViewItem>* PreviousItem = PreviousAvailableBackendItems.FindByHash(ItemDataKeyHash, ItemDataKey))
+					FCollectionManagerModule& CollectionManagerModule = FCollectionManagerModule::GetModule();
+				
+					TArray<FCollectionNameType> ChildCollections;
+					for(const FCollectionNameType& Collection : SourcesData.Collections)
 					{
-						NewItem = *PreviousItem;
-						NewItem->ClearCachedCustomColumns();
+						ChildCollections.Reset();
+						CollectionManagerModule.Get().GetChildCollections(Collection.Name, Collection.Type, ChildCollections);
+
+						for (const FCollectionNameType& ChildCollection : ChildCollections)
+						{
+							// Use "Collections" as the root of the path to avoid this being confused with other view folders - see ContentBrowserUtils::IsCollectionPath
+							FContentBrowserItemData FolderItemData(
+								nullptr, 
+								EContentBrowserItemFlags::Type_Folder | EContentBrowserItemFlags::Category_Collection, 
+								*FString::Printf(TEXT("/Collections/%s/%s"), ECollectionShareType::ToString(ChildCollection.Type), *ChildCollection.Name.ToString()), 
+								ChildCollection.Name, 
+								FText::FromName(ChildCollection.Name), 
+								nullptr
+								);
+
+							const FContentBrowserItemKey FolderItemDataKey(FolderItemData);
+							AvailableBackendItems.Add(FolderItemDataKey, MakeShared<FAssetViewItem>(MoveTemp(FolderItemData)));
+						}
 					}
 				}
-				if (NewItem)
-				{
-					NewItem->AppendItemData(InItemData);
-					NewItem->CacheCustomColumns(CustomColumns, true, true, false /*bUpdateExisting*/);
-				}
-				else
+			}
+
+			if (SourcesData.IsIncludingVirtualPaths())
+			{
+				static const FName RootPath = "/";
+				const TArrayView<const FName> DataSourcePaths = SourcesData.HasVirtualPaths() ? MakeArrayView(SourcesData.VirtualPaths) : MakeArrayView(&RootPath, 1);
+				for (const FName& DataSourcePath : DataSourcePaths)
 				{
-					NewItem = MakeShared<FAssetViewItem>(MoveTemp(InItemData));
+					// Ensure paths do not contain trailing slash
+					ensure(DataSourcePath == RootPath || !FStringView(FNameBuilder(DataSourcePath)).EndsWith(TEXT('/')));
+					ContentBrowserData->EnumerateItemsUnderPath(DataSourcePath, DataFilter, AddNewItem);
 				}
-
-				return true;
-			});
+			}
 		}
 	}
 
@@ -4345,6 +4358,7 @@ void SAssetView::HandleItemDataUpdated(TArrayView<const FContentBrowserItemDataU
 	UContentBrowserDataSubsystem* ContentBrowserData = IContentBrowserDataModule::Get().GetSubsystem();
 
 	TArray<FContentBrowserDataCompiledFilter> CompiledDataFilters;
+	if (SourcesData.IsIncludingVirtualPaths())
 	{
 		const FContentBrowserDataFilter DataFilter = CreateBackendDataFilter();
 
diff --git a/Engine/Source/Editor/ContentBrowser/Public/SAssetView.h b/Engine/Source/Editor/ContentBrowser/Public/SAssetView.h
--- a/Engine/Source/Editor/ContentBrowser/Public/SAssetView.h
+++ b/Engine/Source/Editor/ContentBrowser/Public/SAssetView.h
@@ -58,6 +58,7 @@ public:
 		, _AllowThumbnailHintLabel(true)
 		, _InitialViewType(EAssetViewType::Tile)
 		, _InitialThumbnailSize(EThumbnailSize::Medium)
+		, _InitialThumbnailPoolSize(1024)
 		, _ShowBottomToolbar(true)
 		, _ShowViewOptions(true)
 		, _AllowThumbnailEditMode(false)
@@ -151,6 +152,9 @@ public:
 		/** Initial thumbnail size */
 		SLATE_ARGUMENT( EThumbnailSize, InitialThumbnailSize )
 
+		/** Initial thumbnail pool size */
+		SLATE_ARGUMENT( uint32, InitialThumbnailPoolSize )
+
 		/** Should the toolbar indicating number of selected assets, mode switch buttons, etc... be shown? */
 		SLATE_ARGUMENT( bool, ShowBottomToolbar )
 
diff --git a/Engine/Source/Editor/ContentBrowser/Public/SourcesData.h b/Engine/Source/Editor/ContentBrowser/Public/SourcesData.h
--- a/Engine/Source/Editor/ContentBrowser/Public/SourcesData.h
+++ b/Engine/Source/Editor/ContentBrowser/Public/SourcesData.h
@@ -9,8 +9,12 @@
 
 struct FSourcesData
 {
+	DECLARE_DELEGATE_RetVal_OneParam(bool, FOnEnumerateCustomSourceItemDatas, TFunctionRef<bool(class FContentBrowserItemData&&)>)
+
 	TArray<FName> VirtualPaths;
 	TArray<FCollectionNameType> Collections;
+	FOnEnumerateCustomSourceItemDatas OnEnumerateCustomSourceItemDatas;
+	bool bIncludeVirtualPaths = true;
 
 	FSourcesData() = default;
 
@@ -51,6 +55,11 @@ struct FSourcesData
 		return Collections.Num() > 0;
 	}
 
+	FORCEINLINE bool IsIncludingVirtualPaths() const
+	{
+		return bIncludeVirtualPaths;
+	}
+
 	bool IsDynamicCollection() const
 	{
 		if ( Collections.Num() == 1 && FCollectionManagerModule::IsModuleAvailable() )
diff --git a/Engine/Source/Editor/ContentBrowserData/Private/ContentBrowserDataSource.cpp b/Engine/Source/Editor/ContentBrowserData/Private/ContentBrowserDataSource.cpp
--- a/Engine/Source/Editor/ContentBrowserData/Private/ContentBrowserDataSource.cpp
+++ b/Engine/Source/Editor/ContentBrowserData/Private/ContentBrowserDataSource.cpp
@@ -207,6 +207,16 @@ void UContentBrowserDataSource::EnumerateItemsAtPath(const FName InPath, const E
 {
 }
 
+bool UContentBrowserDataSource::EnumerateItemsAtPaths(const TArrayView<FContentBrowserItemPath> InPaths, const EContentBrowserItemTypeFilter InItemTypeFilter, TFunctionRef<bool(FContentBrowserItemData&&)> InCallback)
+{
+	return true;
+}
+
+bool UContentBrowserDataSource::EnumerateItemsForObjects(const TArrayView<UObject*> InObjects, TFunctionRef<bool(FContentBrowserItemData&&)> InCallback)
+{
+	return true;
+}
+
 bool UContentBrowserDataSource::IsDiscoveringItems(FText* OutStatus)
 {
 	return false;
diff --git a/Engine/Source/Editor/ContentBrowserData/Private/ContentBrowserDataSubsystem.cpp b/Engine/Source/Editor/ContentBrowserData/Private/ContentBrowserDataSubsystem.cpp
--- a/Engine/Source/Editor/ContentBrowserData/Private/ContentBrowserDataSubsystem.cpp
+++ b/Engine/Source/Editor/ContentBrowserData/Private/ContentBrowserDataSubsystem.cpp
@@ -334,6 +334,34 @@ void UContentBrowserDataSubsystem::EnumerateItemsAtPath(const FName InPath, cons
 	}
 }
 
+bool UContentBrowserDataSubsystem::EnumerateItemsAtPaths(const TArrayView<class FContentBrowserItemPath> InItemPaths, const EContentBrowserItemTypeFilter InItemTypeFilter, TFunctionRef<bool(FContentBrowserItemData&&)> InCallback) const
+{
+	for (const auto& ActiveDataSourcePair : ActiveDataSources)
+	{
+		UContentBrowserDataSource* DataSource = ActiveDataSourcePair.Value;
+		if (!DataSource->EnumerateItemsAtPaths(InItemPaths, InItemTypeFilter, InCallback))
+		{
+			return false;
+		}
+	}
+
+	return true;
+}
+
+bool UContentBrowserDataSubsystem::EnumerateItemsForObjects(TArrayView<UObject*> InObjects, TFunctionRef<bool(FContentBrowserItemData&&)> InCallback) const
+{
+	for (const auto& ActiveDataSourcePair : ActiveDataSources)
+	{
+		UContentBrowserDataSource* DataSource = ActiveDataSourcePair.Value;
+		if (!DataSource->EnumerateItemsForObjects(InObjects, InCallback))
+		{
+			return false;
+		}
+	}
+
+	return true;
+}
+
 TArray<FContentBrowserItem> UContentBrowserDataSubsystem::GetItemsAtPath(const FName InPath, const EContentBrowserItemTypeFilter InItemTypeFilter) const
 {
 	TMap<FContentBrowserItemKey, FContentBrowserItem> FoundItems;
diff --git a/Engine/Source/Editor/ContentBrowserData/Public/ContentBrowserDataSource.h b/Engine/Source/Editor/ContentBrowserData/Public/ContentBrowserDataSource.h
--- a/Engine/Source/Editor/ContentBrowserData/Public/ContentBrowserDataSource.h
+++ b/Engine/Source/Editor/ContentBrowserData/Public/ContentBrowserDataSource.h
@@ -186,6 +186,26 @@ public:
 	 */
 	virtual void EnumerateItemsAtPath(const FName InPath, const EContentBrowserItemTypeFilter InItemTypeFilter, TFunctionRef<bool(FContentBrowserItemData&&)> InCallback);
 
+	
+	/**
+	 * Enumerate the items (folders and/or files) that exist at the given content browser paths.
+	 * @note Multiple items may have the same virtual path if they are different types, or come from different data sources.
+	 *
+	 * @param InPaths The paths to search for
+	 * @param InItemTypeFilter The types of items we want to find.
+	 * @param InCallback The function to invoke for each matching item (return true to continue enumeration).
+	 */
+	virtual bool EnumerateItemsAtPaths(const TArrayView<FContentBrowserItemPath> InPaths, const EContentBrowserItemTypeFilter InItemTypeFilter, TFunctionRef<bool(FContentBrowserItemData&&)> InCallback);
+
+	/**
+	 * Enumerate the items (files) that exist for the given objects.
+	 * @note Multiple items may have the same virtual path if they are different types, or come from different data sources.
+	 *
+	 * @param InObjects The objects to enumerate
+	 * @param InCallback The function to invoke for each matching item (return true to continue enumeration).
+	 */
+	virtual bool EnumerateItemsForObjects(const TArrayView<UObject*> InObjects, TFunctionRef<bool(FContentBrowserItemData&&)> InCallback);
+
 	/**
 	 * Query whether this data source instance is currently discovering content, and retrieve an optional status message that can be shown in the UI.
 	 */
diff --git a/Engine/Source/Editor/ContentBrowserData/Public/ContentBrowserDataSubsystem.h b/Engine/Source/Editor/ContentBrowserData/Public/ContentBrowserDataSubsystem.h
--- a/Engine/Source/Editor/ContentBrowserData/Public/ContentBrowserDataSubsystem.h
+++ b/Engine/Source/Editor/ContentBrowserData/Public/ContentBrowserDataSubsystem.h
@@ -157,6 +157,25 @@ public:
 	void EnumerateItemsAtPath(const FName InPath, const EContentBrowserItemTypeFilter InItemTypeFilter, TFunctionRef<bool(FContentBrowserItem&&)> InCallback) const;
 	void EnumerateItemsAtPath(const FName InPath, const EContentBrowserItemTypeFilter InItemTypeFilter, TFunctionRef<bool(FContentBrowserItemData&&)> InCallback) const;
 
+	/**
+	 * Enumerate the items (files) that exist at the given paths.
+	 * @note Multiple items may have the same virtual path if they are different types, or come from different data sources.
+	 *
+	 * @param InItemPaths The paths to enumerate
+	 * @param InItemTypeFilter The types of items we want to find.
+	 * @param InCallback The function to invoke for each matching item (return true to continue enumeration).
+	 */
+	bool EnumerateItemsAtPaths(const TArrayView<class FContentBrowserItemPath> InItemPaths, const EContentBrowserItemTypeFilter InItemTypeFilter, TFunctionRef<bool(FContentBrowserItemData&&)> InCallback) const;
+
+	/**
+	 * Enumerate the items (files) that exist for the given objects.
+	 * @note Multiple items may have the same virtual path if they are different types, or come from different data sources.
+	 *
+	 * @param InObjects The objects to enumerate
+	 * @param InCallback The function to invoke for each matching item (return true to continue enumeration).
+	 */
+	bool EnumerateItemsForObjects(const TArrayView<UObject*> InObjects, TFunctionRef<bool(FContentBrowserItemData&&)> InCallback) const;
+
 	/**
 	 * Get the items (folders and/or files) that exist at the given virtual path.
 	 * @note Multiple items may have the same virtual path if they are different types, or come from different data sources.
diff --git a/Engine/Source/Editor/ContentBrowserData/Public/ContentBrowserItemPath.h b/Engine/Source/Editor/ContentBrowserData/Public/ContentBrowserItemPath.h
--- a/Engine/Source/Editor/ContentBrowserData/Public/ContentBrowserItemPath.h
+++ b/Engine/Source/Editor/ContentBrowserData/Public/ContentBrowserItemPath.h
@@ -14,10 +14,24 @@ class CONTENTBROWSERDATA_API FContentBrowserItemPath
 {
 public:
 	FContentBrowserItemPath();
+
+	FContentBrowserItemPath(const FName InVirtualPath, const FName InInternalPath)
+		: VirtualPath(InVirtualPath)
+		, InternalPath(InInternalPath)
+	{
+		check(!VirtualPath.IsNone());
+		check(!InternalPath.IsNone());
+	}
+
 	FContentBrowserItemPath(const FStringView InPath, const EContentBrowserPathType InPathType);
 	FContentBrowserItemPath(const TCHAR* InPath, const EContentBrowserPathType InPathType);
 	FContentBrowserItemPath(const FName InPath, const EContentBrowserPathType InPathType);
 
+	bool operator==(const FContentBrowserItemPath& Other) const 
+	{ 
+		return VirtualPath == Other.VirtualPath && InternalPath == Other.InternalPath;
+	}
+
 	/**
 	 * Set the path being stored
 	 */
diff --git a/Engine/Source/Editor/UMGEditor/Private/Library/SLibraryViewModel.cpp b/Engine/Source/Editor/UMGEditor/Private/Library/SLibraryViewModel.cpp
--- a/Engine/Source/Editor/UMGEditor/Private/Library/SLibraryViewModel.cpp
+++ b/Engine/Source/Editor/UMGEditor/Private/Library/SLibraryViewModel.cpp
@@ -21,6 +21,7 @@
 #include "AssetRegistryModule.h"
 #include "WidgetBlueprintEditorUtils.h"
 
+#include "IContentBrowserDataModule.h"
 #include "Settings/ContentBrowserSettings.h"
 #include "Settings/WidgetDesignerSettings.h"
 #include "UMGEditorProjectSettings.h"
@@ -111,6 +112,10 @@ void FWidgetTemplateListViewModel::ConstructListView(TArray<TSharedPtr<FWidgetTe
 
 	if (!TemplatesFilter)
 	{
+		UContentBrowserDataSubsystem* ContentBrowserDataSubsystem = IContentBrowserDataModule::Get().GetSubsystem();
+		NumAssets = 0;
+		NumClasses = 0;
+
 		// Generate filter text
 		bool bHasFilters = false;
 		TStringBuilder<2048> FilterString;
@@ -118,43 +123,57 @@ void FWidgetTemplateListViewModel::ConstructListView(TArray<TSharedPtr<FWidgetTe
 		{
 			if (TSharedPtr<FWidgetTemplateClass> TemplateClass = StaticCastSharedPtr<FWidgetTemplateClass>(Template))
 			{
-				TWeakObjectPtr<UClass> WidgetClass = TemplateClass->GetWidgetClass();
-				FName TemplateName = WidgetClass.IsValid() ? WidgetClass->GetFName() : TemplateClass->GetWidgetAssetData().AssetName;
-				FString TemplateString = TemplateName.ToString();
+				FString TemplateString;
+				FString TemplatePath;
+
+				if (TemplateClass->GetWidgetAssetData().IsValid())
+				{
+					TemplateString = TemplateClass->GetWidgetAssetData().AssetName.ToString();
+					NumAssets++;
+				}
+				else if (UClass* WidgetClass = TemplateClass->GetWidgetClass().Get())
+				{
+					TemplateString = WidgetClass->GetFName().ToString();
+					NumClasses++;
+				}
 
 				if (!TemplateString.IsEmpty())
 				{
 					TemplateString.RemoveFromEnd(TEXT("_C"));
 					FilterString += bHasFilters ? "|+" : "+";
 					FilterString += TemplateString;
-					bHasFilters = true;
+					bHasFilters = false;
 				}
 			}
 		}
 
-		WidgetTextFilter = MakeShared<FFrontendFilter_Text>();
-		WidgetTextFilter->SetActive(true);
-		WidgetTextFilter->SetIncludeClassName(true);
-		WidgetTextFilter->SetRawFilterText(FText::FromString(FilterString.ToString()));
-		CachedLowercaseWidgetFilter = WidgetTextFilter->GetRawFilterText().ToString().ToLower();
+		LibrarySourceData = MakeShared<FSourcesData>();
+		// Provide a dummy invalid virtual path to make sure nothing tries to enumerate root "/"
+		LibrarySourceData->VirtualPaths.Add(FName(TEXT("/UMGWidgetTemplateListViewModel")));
+		// Disable any enumerate of virtual path folders
+		LibrarySourceData->bIncludeVirtualPaths = false;
+		// Supply a custom list of source items to display
+		LibrarySourceData->OnEnumerateCustomSourceItemDatas.BindSP(this, &FWidgetTemplateListViewModel::EnumerateCustomSourceItemDatas);
+
+		CachedLowercaseWidgetFilter = FString(FilterString.ToString()).ToLower();
 
 		SearchFilter = MakeShared<FFrontendFilter_Text>();
 		SearchFilter->SetActive(false);
 
 		TemplatesFilter = MakeShared<FAssetFilterCollectionType>();
-		TemplatesFilter->Add(WidgetTextFilter);
 		TemplatesFilter->Add(SearchFilter);
 	}
 
 	if (!AssetViewPtr)
 	{
-		FARFilter WidgetTemplateFilter;
-		WidgetTemplateFilter.ClassNames.Add(UClass::StaticClass()->GetFName());
-		WidgetTemplateFilter.ClassNames.Add(UWidgetBlueprint::StaticClass()->GetFName());
+		// Reserve just enough thumbnails, plus a few extra if this category has user widgets
+		uint32 ThumbnailPoolSize = NumAssets + NumClasses;
+		ThumbnailPoolSize += NumAssets > 0 ? 32 : 0;
 
 		AssetViewPtr = SNew(SAssetView)
 			.InitialCategoryFilter(EContentBrowserItemCategoryFilter::IncludeAssets | EContentBrowserItemCategoryFilter::IncludeClasses)
-			.InitialBackendFilter(WidgetTemplateFilter)
+			.InitialSourcesData(*LibrarySourceData)
+			.InitialThumbnailPoolSize(ThumbnailPoolSize)
 			.InitialThumbnailSize(EThumbnailSize::Small)
 			.FrontendFilters(TemplatesFilter)
 			.ForceShowEngineContent(true)
@@ -168,6 +187,50 @@ void FWidgetTemplateListViewModel::ConstructListView(TArray<TSharedPtr<FWidgetTe
 	AssetViewPtr->RequestSlowFullListRefresh();
 }
 
+bool FWidgetTemplateListViewModel::EnumerateCustomSourceItemDatas(TFunctionRef<bool(FContentBrowserItemData&&)> InCallback)
+{
+	SourceItemPaths.Reset();
+	SourceItemPaths.Reserve(NumAssets);
+
+	TArray<UObject*> ClassObjects;
+	ClassObjects.Reserve(NumClasses);
+
+	for (TSharedPtr<FWidgetTemplate> Template : Templates)
+	{
+		if (TSharedPtr<FWidgetTemplateClass> TemplateClass = StaticCastSharedPtr<FWidgetTemplateClass>(Template))
+		{
+			if (TemplateClass->GetWidgetAssetData().IsValid())
+			{
+				SourceItemPaths.Add(FContentBrowserItemPath(TemplateClass->GetWidgetAssetData().PackageName, EContentBrowserPathType::Internal));
+			}
+			else if (UClass* WidgetClass = TemplateClass->GetWidgetClass().Get())
+			{
+				ClassObjects.Add(WidgetClass);
+			}
+		}
+	}
+	
+	UContentBrowserDataSubsystem* ContentBrowserDataSubsystem = IContentBrowserDataModule::Get().GetSubsystem();
+
+	if (!SourceItemPaths.IsEmpty())
+	{
+		if (!ContentBrowserDataSubsystem->EnumerateItemsAtPaths(SourceItemPaths, EContentBrowserItemTypeFilter::IncludeFiles, InCallback))
+		{
+			return false;
+		}
+	}
+
+	if (!ClassObjects.IsEmpty())
+	{
+		if (!ContentBrowserDataSubsystem->EnumerateItemsForObjects(ClassObjects, InCallback))
+		{
+			return false;
+		}
+	}
+
+	return true;
+}
+
 FText FWidgetTemplateListViewModel::GetName() const
 {
 	return FText();
diff --git a/Engine/Source/Editor/UMGEditor/Private/Library/SLibraryViewModel.h b/Engine/Source/Editor/UMGEditor/Private/Library/SLibraryViewModel.h
--- a/Engine/Source/Editor/UMGEditor/Private/Library/SLibraryViewModel.h
+++ b/Engine/Source/Editor/UMGEditor/Private/Library/SLibraryViewModel.h
@@ -44,17 +44,28 @@ public:
 	TSharedPtr<FAssetFilterCollectionType> TemplatesFilter;
 private:
 
+	bool EnumerateCustomSourceItemDatas(TFunctionRef<bool(FContentBrowserItemData&&)> InCallback);
+
 	/** The asset view widget */
 	TSharedPtr<SAssetView> AssetViewPtr;
 
-	/** Filter we use to narrow down the libary to select widgets asset view */
-	TSharedPtr<class FFrontendFilter_Text> WidgetTextFilter;
-
 	/** Filter we forward our search text to for the asset view */
 	TSharedPtr<class FFrontendFilter_Text> SearchFilter;
 
-	/** Lowercase string used by widget text filter */
+	/** Source data for the asset view */
+	TSharedPtr<struct FSourcesData> LibrarySourceData;
+
+	/** Source item path storage for asset view, accessed via delegate */
+	TArray<FContentBrowserItemPath> SourceItemPaths;
+
+	/** Lowercase string representing widget text filter */
 	FString CachedLowercaseWidgetFilter;
+
+	/** Number of widget assets in view */
+	uint32 NumAssets = 0;
+
+	/** Number of widget classes in view */
+	uint32 NumClasses = 0;
 };
 
 class FLibraryViewModel : public FWidgetCatalogViewModel
diff --git a/Engine/Source/Editor/UMGEditor/Private/Templates/WidgetTemplateClass.h b/Engine/Source/Editor/UMGEditor/Private/Templates/WidgetTemplateClass.h
--- a/Engine/Source/Editor/UMGEditor/Private/Templates/WidgetTemplateClass.h
+++ b/Engine/Source/Editor/UMGEditor/Private/Templates/WidgetTemplateClass.h
@@ -44,7 +44,7 @@ public:
 	TWeakObjectPtr<UClass> GetWidgetClass() const { return WidgetClass; }
 
 	/** Returns the asset data for this widget which might be invalid. */
-	FAssetData GetWidgetAssetData() { return WidgetAssetData; }
+	const FAssetData& GetWidgetAssetData() { return WidgetAssetData; }
 
 protected:
 	/** Creates a widget template class without any class reference */
diff --git a/Engine/Source/Editor/UMGEditor/UMGEditor.Build.cs b/Engine/Source/Editor/UMGEditor/UMGEditor.Build.cs
--- a/Engine/Source/Editor/UMGEditor/UMGEditor.Build.cs
+++ b/Engine/Source/Editor/UMGEditor/UMGEditor.Build.cs
@@ -42,6 +42,7 @@ public class UMGEditor : ModuleRules
 				"Core",
 				"CoreUObject",
 				"ContentBrowser",
+				"ContentBrowserData",
 				"ApplicationCore",
 				"InputCore",
 				"Engine",
