PR #8475: Asset view filters icons (Contributed by [REDACTED_PERSON])

diff --git a/Engine/Source/Editor/ContentBrowser/Private/SFilterList.cpp b/Engine/Source/Editor/ContentBrowser/Private/SFilterList.cpp
--- a/Engine/Source/Editor/ContentBrowser/Private/SFilterList.cpp
+++ b/Engine/Source/Editor/ContentBrowser/Private/SFilterList.cpp
@@ -31,6 +31,7 @@
 #include "ToolMenus.h"
 #include "ContentBrowserMenuContexts.h"
 #include "Widgets/Images/SImage.h"
+#include "Styling/SlateIconFinder.h"
 
 #define LOCTEXT_NAMESPACE "ContentBrowser"
 
@@ -1210,7 +1211,7 @@ void SFilterList::CreateFiltersMenuCategory(FToolMenuSection& Section, const TAr
 					NAME_None,
 					LabelText,
 					FText::Format( LOCTEXT("FilterByTooltipPrefix", "Filter by {0}"), LabelText ),
-					FSlateIcon(),
+					FSlateIconFinder::FindIconForClass(TypeActions->GetSupportedClass()),
 					FUIAction(
 						FExecuteAction::CreateSP( const_cast<SFilterList*>(this), &SFilterList::FilterByTypeClicked, WeakTypeActions ),
 						FCanExecuteAction(),
@@ -1373,9 +1374,11 @@ void SFilterList::PopulateAddFilterMenu(UToolMenu* Menu)
 			"ResetFilters",
 			LOCTEXT("FilterListResetFilters", "Reset Filters"),
 			LOCTEXT("FilterListResetToolTip", "Resets current filter selection"),
-			FSlateIcon(),
-			FUIAction(FExecuteAction::CreateSP(this, &SFilterList::OnResetFilters))
-			);
+			FSlateIcon(FAppStyle::Get().GetStyleSetName(), "PropertyWindow.DiffersFromDefault"),
+			FUIAction(
+				FExecuteAction::CreateSP(this, &SFilterList::OnResetFilters),
+				FCanExecuteAction::CreateLambda([this]() { return HasAnyFilters(); }))
+		);
 	}
 
 	// First add the expanded category, this appears as standard entries in the list (Note: intentionally not using FindChecked here as removing it from the map later would cause the ref to be garbage)
@@ -1391,7 +1394,7 @@ void SFilterList::PopulateAddFilterMenu(UToolMenu* Menu)
 				NAME_None,
 				ExpandedCategory->Name,
 				ExpandedCategory->Tooltip,
-				FSlateIcon(),
+				FSlateIcon(FAppStyle::Get().GetStyleSetName(), "PlacementBrowser.Icons.Basic"),
 				FUIAction(
 				FExecuteAction::CreateSP( this, &SFilterList::FilterByTypeCategoryClicked, MenuExpansion ),
 				FCanExecuteAction(),
