Fix crash when saving all assets while renaming

diff --git a/Engine/Source/Editor/ContentBrowser/Private/SContentBrowser.cpp b/Engine/Source/Editor/ContentBrowser/Private/SContentBrowser.cpp
--- a/Engine/Source/Editor/ContentBrowser/Private/SContentBrowser.cpp
+++ b/Engine/Source/Editor/ContentBrowser/Private/SContentBrowser.cpp
@@ -1402,7 +1402,8 @@ void SContentBrowser::CopySettingsFromBrowser(TSharedPtr<SContentBrowser> OtherB
 
 FReply SContentBrowser::OnKeyDown( const FGeometry& MyGeometry, const FKeyEvent& InKeyEvent )
 {
-	if( Commands->ProcessCommandBindings( InKeyEvent ) )
+	bool bIsRenamingAsset = AssetViewPtr && AssetViewPtr->IsRenamingAsset();
+	if(bIsRenamingAsset || Commands->ProcessCommandBindings( InKeyEvent ) )
 	{
 		return FReply::Handled();
 	}
