[ContentBrowser] Hide show in explorer if physical path does not exist
#rnx

diff --git a/Engine/Source/Developer/AssetTools/Private/AssetViewUtils.cpp b/Engine/Source/Developer/AssetTools/Private/AssetViewUtils.cpp
--- a/Engine/Source/Developer/AssetTools/Private/AssetViewUtils.cpp
+++ b/Engine/Source/Developer/AssetTools/Private/AssetViewUtils.cpp
@@ -1222,7 +1222,13 @@ bool AssetViewUtils::IsValidObjectPathForCreate(const FString& ObjectPath, const
 	}
 
 	// Make sure we are not creating an path that is too long for the OS
-	const FString RelativePathFilename = FPackageName::LongPackageNameToFilename(PackageName, FPackageName::GetAssetPackageExtension());	// full relative path with name + extension
+	FString RelativePathFilename;
+	if (!FPackageName::TryConvertLongPackageNameToFilename(PackageName, RelativePathFilename, FPackageName::GetAssetPackageExtension()))	// full relative path with name + extension
+	{
+		OutErrorMessage = LOCTEXT("ConvertToFilename", "Package name could not be converted to file name.");
+		return false;
+	}
+
 	const FString FullPath = FPaths::ConvertRelativePathToFull(RelativePathFilename);	// path to file on disk
 	if (ObjectPath.Len() > (FPlatformMisc::GetMaxPathLength() - MAX_CLASS_NAME_LENGTH))
 	{
diff --git a/Engine/Source/Editor/ContentBrowser/Private/PathContextMenu.cpp b/Engine/Source/Editor/ContentBrowser/Private/PathContextMenu.cpp
--- a/Engine/Source/Editor/ContentBrowser/Private/PathContextMenu.cpp
+++ b/Engine/Source/Editor/ContentBrowser/Private/PathContextMenu.cpp
@@ -138,7 +138,7 @@ void FPathContextMenu::MakePathViewContextMenu(UToolMenu* Menu)
 			if (!Context->bNoFolderOnDisk)
 			{
 				Section.AddMenuEntry(
-					"ExploreTooltip",
+					"Explore",
 					ContentBrowserUtils::GetExploreFolderText(),
 					LOCTEXT("ExploreTooltip", "Finds this folder on disk."),
 					FSlateIcon(),
diff --git a/Engine/Source/Editor/ContentBrowser/Private/SAssetDialog.cpp b/Engine/Source/Editor/ContentBrowser/Private/SAssetDialog.cpp
--- a/Engine/Source/Editor/ContentBrowser/Private/SAssetDialog.cpp
+++ b/Engine/Source/Editor/ContentBrowser/Private/SAssetDialog.cpp
@@ -506,11 +506,52 @@ void SAssetDialog::SelectDefaultPaths()
 	AssetPicker->GetAssetView()->SetSourcesData(DefaultSourcesData);
 }
 
+bool SAssetDialog::ExecuteExploreInternal(bool bTest)
+{
+	bool bCanExplore = false;
+
+	TArray<FContentBrowserItem> SelectedItems = AssetPicker->GetAssetView()->GetSelectedItems();
+	if (SelectedItems.Num() == 0)
+	{
+		SelectedItems = PathPicker->GetPathView()->GetSelectedFolderItems();
+	}
+
+	for (const FContentBrowserItem& SelectedItem : SelectedItems)
+	{
+		FString ItemFilename;
+		if (SelectedItem.GetItemPhysicalPath(ItemFilename))
+		{
+			const bool bExists = SelectedItem.IsFile() ? FPaths::FileExists(ItemFilename) : FPaths::DirectoryExists(ItemFilename);
+			if (bExists)
+			{
+				bCanExplore = true;
+
+				if (!bTest)
+				{
+					FPlatformProcess::ExploreFolder(*IFileManager::Get().ConvertToAbsolutePathForExternalAppForRead(*ItemFilename));
+				}
+				else
+				{
+					break;
+				}
+			}
+		}
+	}
+
+	return bCanExplore;
+}
+
 void SAssetDialog::ExecuteExplore()
 {
+	//ExecuteExploreInternal(/*bTest*/ false);
 	ContentBrowserUtils::ExploreFolders(AssetPicker->GetAssetView()->GetSelectedItems(), AssetPicker->GetAssetView().ToSharedRef());
 }
 
+bool SAssetDialog::CanExecuteExplore()
+{
+	return ExecuteExploreInternal(/*bTest*/ true);
+}
+
 bool SAssetDialog::CanExecuteCreateNewFolder() const
 {	
 	// We can only create folders when we have a single path selected
@@ -595,14 +636,15 @@ void SAssetDialog::SetupContextMenuContent(FMenuBuilder& MenuBuilder, const TArr
 
 	MenuBuilder.EndSection();
 
-	MenuBuilder.BeginSection("AssetDialogExplore", LOCTEXT("AssetDialogExploreHeading", "Explore"));
-
-	MenuBuilder.AddMenuEntry(ContentBrowserUtils::GetExploreFolderText(), 
-							 LOCTEXT("ExploreTooltip", "Finds this folder on disk."), 
-							 FSlateIcon(FEditorStyle::GetStyleSetName(), "SystemWideCommands.FindInContentBrowser"), 
-							 FUIAction(FExecuteAction::CreateSP(this, &SAssetDialog::ExecuteExplore)));
-
-	MenuBuilder.EndSection();
+	if (CanExecuteExplore())
+	{
+		MenuBuilder.BeginSection("AssetDialogExplore", LOCTEXT("AssetDialogExploreHeading", "Explore"));
+		MenuBuilder.AddMenuEntry(ContentBrowserUtils::GetExploreFolderText(),
+			LOCTEXT("ExploreTooltip", "Finds this folder on disk."),
+			FSlateIcon(FEditorStyle::GetStyleSetName(), "SystemWideCommands.FindInContentBrowser"),
+			FUIAction(FExecuteAction::CreateSP(this, &SAssetDialog::ExecuteExplore)));
+		MenuBuilder.EndSection();
+	}
 }
 
 EActiveTimerReturnType SAssetDialog::SetFocusPostConstruct( double InCurrentTime, float InDeltaTime )
@@ -918,6 +960,10 @@ FString SAssetDialog::GetObjectPathForSave() const
 		{
 			Base = ConvertedPath.ToString();
 		}
+		else
+		{
+			return FString();
+		}
 	}
 
 	return Base / CurrentlyEnteredAssetName + TEXT(".") + CurrentlyEnteredAssetName;
diff --git a/Engine/Source/Editor/ContentBrowser/Private/SAssetDialog.h b/Engine/Source/Editor/ContentBrowser/Private/SAssetDialog.h
--- a/Engine/Source/Editor/ContentBrowser/Private/SAssetDialog.h
+++ b/Engine/Source/Editor/ContentBrowser/Private/SAssetDialog.h
@@ -119,6 +119,8 @@ private:
 
 	/** Handler for show in explorer */
 	void ExecuteExplore();
+	bool CanExecuteExplore();
+	bool ExecuteExploreInternal(bool bTest=false);
 
 	/** Setup function for the context menu creation of folder and assets */
 	void SetupContextMenuContent(FMenuBuilder& MenuBuilder, const TArray<FString>& SelectedPaths);
diff --git a/Engine/Source/Editor/ContentBrowser/Private/SContentBrowser.cpp b/Engine/Source/Editor/ContentBrowser/Private/SContentBrowser.cpp
--- a/Engine/Source/Editor/ContentBrowser/Private/SContentBrowser.cpp
+++ b/Engine/Source/Editor/ContentBrowser/Private/SContentBrowser.cpp
@@ -3459,12 +3459,22 @@ TSharedPtr<SWidget> SContentBrowser::GetItemContextMenu(TArrayView<const FConten
 		}
 
 		TArray<FString> SelectedPackagePaths;
+		bool bPhysicalPathExists = false;
 		for (const FContentBrowserItem& SelectedFolder : SelectedFolders)
 		{
 			FName PackagePath;
 			if (SelectedFolder.Legacy_TryGetPackagePath(PackagePath))
 			{
 				SelectedPackagePaths.Add(PackagePath.ToString());
+
+				if (!bPhysicalPathExists)
+				{
+					FString PhysicalPath;
+					if (SelectedFolder.GetItemPhysicalPath(PhysicalPath) && FPaths::DirectoryExists(PhysicalPath))
+					{
+						bPhysicalPathExists = true;
+					}
+				}
 			}
 		}
 
@@ -3489,6 +3499,11 @@ TSharedPtr<SWidget> SContentBrowser::GetItemContextMenu(TArrayView<const FConten
 			Context->bCanBeModified = false;
 		}
 
+		if (!bPhysicalPathExists)
+		{
+			Context->bNoFolderOnDisk = true;
+		}
+
 		FToolMenuContext MenuContext(Commands, Extender, Context);
 
 		{
