Asset Picker: Fix an issue preventing some options in the Asset Picker from working when summoned from a menu

diff --git a/Engine/Source/Editor/ContentBrowser/Private/SAssetView.cpp b/Engine/Source/Editor/ContentBrowser/Private/SAssetView.cpp
--- a/Engine/Source/Editor/ContentBrowser/Private/SAssetView.cpp
+++ b/Engine/Source/Editor/ContentBrowser/Private/SAssetView.cpp
@@ -3033,7 +3033,6 @@ void SAssetView::SetCurrentViewTypeFromMenu(EAssetViewType::Type NewType)
 	if (NewType != CurrentViewType)
 	{
 		SetCurrentViewType(NewType);
-		FSlateApplication::Get().DismissAllMenus();
 	}
 }
 
