Changed a check to a ensure to avoid crashing in a bad scenario where an asset creation would be requested while a another is active still active.

Added a line of log to help us track the potential cause of the issue.

diff --git a/Engine/Source/Editor/ContentBrowser/Private/SAssetView.cpp b/Engine/Source/Editor/ContentBrowser/Private/SAssetView.cpp
--- a/Engine/Source/Editor/ContentBrowser/Private/SAssetView.cpp
+++ b/Engine/Source/Editor/ContentBrowser/Private/SAssetView.cpp
@@ -591,8 +591,19 @@ void SAssetView::NewFolderItemRequested(const FContentBrowserItemTemporaryContex
 		return;
 	}
 
-	// we should only be creating one deferred folder per tick
-	check(!DeferredItemToCreate.IsValid());
+	// we should only be creating one deferred folder at a time
+	if (!ensureAlwaysMsgf(!DeferredItemToCreate.IsValid(), TEXT("Deferred new asset folder creation while there is already a deferred item creation: %s"), *NewItemContext.GetItem().GetItemName().ToString()))
+	{
+		if (DeferredItemToCreate->bWasAddedToView)
+		{
+			FContentBrowserItemKey ItemToRemoveKey(DeferredItemToCreate->ItemContext.GetItem());
+			FilteredAssetItems.RemoveAll([&ItemToRemoveKey](const TSharedPtr<FAssetViewItem>& InAssetViewItem) { return ItemToRemoveKey == FContentBrowserItemKey(InAssetViewItem->GetItem()); });
+			RefreshList();
+		}
+
+		DeferredItemToCreate.Release();
+	}
+
 
 	// Folder creation requires focus to give object a name, otherwise object will not be created
 	TSharedPtr<SWindow> OwnerWindow = FSlateApplication::Get().FindWidgetWindow(AsShared());
@@ -620,7 +631,17 @@ void SAssetView::NewFileItemRequested(const FContentBrowserItemDataTemporaryCont
 	}
 
 	// We should only be creating one deferred file at a time
-	check(!DeferredItemToCreate.IsValid());
+	if (!ensureAlwaysMsgf(!DeferredItemToCreate.IsValid(), TEXT("Deferred new asset file creation while there is already a deferred item creation: %s"), *NewItemContext.GetItemData().GetItemName().ToString()))
+	{
+		if (DeferredItemToCreate->bWasAddedToView)
+		{
+			FContentBrowserItemKey ItemToRemoveKey(DeferredItemToCreate->ItemContext.GetItem());
+			FilteredAssetItems.RemoveAll([&ItemToRemoveKey](const TSharedPtr<FAssetViewItem>& InAssetViewItem){ return ItemToRemoveKey == FContentBrowserItemKey(InAssetViewItem->GetItem()); });
+			RefreshList();
+		}
+
+		DeferredItemToCreate.Release();
+	}
 
 	// File creation requires focus to give item a name, otherwise the item will not be created
 	TSharedPtr<SWindow> OwnerWindow = FSlateApplication::Get().FindWidgetWindow(AsShared());
@@ -3702,6 +3723,11 @@ void SAssetView::AssetRenameBegin(const TSharedPtr<FAssetViewItem>& Item, const
 {
 	check(!RenamingAsset.IsValid());
 	RenamingAsset = Item;
+
+	if (DeferredItemToCreate.IsValid())
+	{
+		UE_LOG(LogContentBrowser, Log, TEXT("Renaming the item being created (Deferred Item: %s)."), *Item->GetItem().GetItemName().ToString());
+	}
 }
 
 void SAssetView::AssetRenameCommit(const TSharedPtr<FAssetViewItem>& Item, const FString& NewName, const FSlateRect& MessageAnchor, const ETextCommit::Type CommitType)
