Slight increase to the sources panel splitter size and hit detection handle in the Content Browser to improve UX

diff --git a/Engine/Source/Editor/ContentBrowser/Private/SContentBrowser.cpp b/Engine/Source/Editor/ContentBrowser/Private/SContentBrowser.cpp
--- a/Engine/Source/Editor/ContentBrowser/Private/SContentBrowser.cpp
+++ b/Engine/Source/Editor/ContentBrowser/Private/SContentBrowser.cpp
@@ -393,8 +393,8 @@ void SContentBrowser::Construct( const FArguments& InArgs, const FName& InInstan
 						[
 							SAssignNew(PathFavoriteSplitterPtr, SSplitter)
 							.Clipping(EWidgetClipping::ClipToBounds)
-							.PhysicalSplitterHandleSize(1.0f)
-							.HitDetectionSplitterHandleSize(3.0f)
+							.PhysicalSplitterHandleSize(2.0f)
+							.HitDetectionSplitterHandleSize(5.0f)
 							.Orientation(EOrientation::Orient_Vertical)
 							.MinimumSlotHeight(26.0f)
 							.Visibility( this, &SContentBrowser::GetSourcesViewVisibility )
