Amortize user searching in the content browser

This avoids part of the hitch when searching the CB in large projects, as while the initial backend hitch (to show all content) remains, the frontend filtering is now amortized so the editor remains responsive as you type out your search query. The CB shows a little progress bar above the asset view while the search is running.

diff --git a/Engine/Source/Editor/ContentBrowser/Private/SAssetView.cpp b/Engine/Source/Editor/ContentBrowser/Private/SAssetView.cpp
--- a/Engine/Source/Editor/ContentBrowser/Private/SAssetView.cpp
+++ b/Engine/Source/Editor/ContentBrowser/Private/SAssetView.cpp
@@ -182,7 +182,7 @@ SAssetView::~SAssetView()
 
 void SAssetView::Construct( const FArguments& InArgs )
 {
-	bIsWorking = false;
+	InitialNumAmortizedTasks = 0;
 	TotalAmortizeTime = 0;
 	AmortizeStartTime = 0;
 	MaxSecondsPerFrame = 0.015;
@@ -354,7 +354,7 @@ void SAssetView::Construct( const FArguments& InArgs )
 		.AutoHeight()
 		[
 			SNew( SBox )
-			.Visibility_Lambda([this] { return bIsWorking ? EVisibility::SelfHitTestInvisible : EVisibility::Collapsed; })
+			.Visibility_Lambda([this] { return InitialNumAmortizedTasks > 0 ? EVisibility::SelfHitTestInvisible : EVisibility::Collapsed; })
 			.HeightOverride( 2 )
 			[
 				SNew( SProgressBar )
@@ -525,13 +525,10 @@ void SAssetView::Construct( const FArguments& InArgs )
 
 TOptional< float > SAssetView::GetIsWorkingProgressBarState() const
 {
-	if (bIsWorking)
+	if (InitialNumAmortizedTasks > 0)
 	{
-		const int32 TotalAssetCount = FilteredAssetItems.Num() + ItemsPendingFrontendFilter.Num();
-		if (TotalAssetCount > 0)
-		{
-			return static_cast<float>(FilteredAssetItems.Num()) / static_cast<float>(TotalAssetCount);
-		}
+		const int32 CompletedTasks = FMath::Max(0, InitialNumAmortizedTasks - ItemsPendingFrontendFilter.Num());
+		return static_cast<float>(CompletedTasks) / static_cast<float>(InitialNumAmortizedTasks);
 	}
 	return 0.0f;
 }
@@ -1020,16 +1017,16 @@ void SAssetView::Tick( const FGeometry& AllottedGeometry, const double InCurrent
 		bForceViewUpdate = false;
 
 		const double TickStartTime = FPlatformTime::Seconds();
-		const bool bWasWorking = bIsWorking;
+		const bool bWasWorking = InitialNumAmortizedTasks > 0;
 
 		// Mark the first amortize time
 		if (AmortizeStartTime == 0)
 		{
 			AmortizeStartTime = FPlatformTime::Seconds();
-			bIsWorking = true;
+			InitialNumAmortizedTasks = ItemsPendingFrontendFilter.Num();
 		}
 
-		ProcessItemsPendingFilter(bUserSearching ? -1.0 : TickStartTime);
+		ProcessItemsPendingFilter(TickStartTime);
 
 		if (HasItemsPendingFilter())
 		{
@@ -1046,7 +1043,7 @@ void SAssetView::Tick( const FGeometry& AllottedGeometry, const double InCurrent
 		{
 			TotalAmortizeTime += FPlatformTime::Seconds() - AmortizeStartTime;
 			AmortizeStartTime = 0;
-			bIsWorking = false;
+			InitialNumAmortizedTasks = 0;
 
 			// Update the columns in the column view now that we know the majority type
 			if (CurrentViewType == EAssetViewType::Column)
@@ -1997,6 +1994,9 @@ void SAssetView::RefreshFilteredItems()
 	FilteredAssetItemTypeCounts.Reset();
 	RelevantThumbnails.Reset();
 
+	AmortizeStartTime = 0;
+	InitialNumAmortizedTasks = 0;
+
 	LastSortTime = 0;
 	bPendingSortFilteredItems = true;
 
@@ -4065,6 +4065,11 @@ FText SAssetView::GetAssetShowWarningText() const
 		return AssetShowWarningText.Get();
 	}
 	
+	if (InitialNumAmortizedTasks > 0)
+	{
+		return LOCTEXT("ApplyingFilter", "Applying filter...");
+	}
+
 	FText NothingToShowText, DropText;
 	if (ShouldFilterRecursively())
 	{
diff --git a/Engine/Source/Editor/ContentBrowser/Public/SAssetView.h b/Engine/Source/Editor/ContentBrowser/Public/SAssetView.h
--- a/Engine/Source/Editor/ContentBrowser/Public/SAssetView.h
+++ b/Engine/Source/Editor/ContentBrowser/Public/SAssetView.h
@@ -1044,8 +1044,8 @@ private:
 	/** The total time spent amortizing the delegate filter */
 	double TotalAmortizeTime;
 
-	/** Whether the asset view is currently working on something and should display a cue to the user */
-	bool bIsWorking;
+	/** The initial number of tasks when we started performing amortized work (used to display a progress cue to the user) */
+	int32 InitialNumAmortizedTasks;
 
 	/** The text to highlight on the assets */
 	TAttribute< FText > HighlightedText;
