- Fix DisallowedClasses meta tag no longer respecting children classes; Shelf was also QA tested against original JIRA that caused the regression

diff --git a/Engine/Source/Editor/ContentBrowser/Private/ContentBrowserUtils.cpp b/Engine/Source/Editor/ContentBrowser/Private/ContentBrowserUtils.cpp
--- a/Engine/Source/Editor/ContentBrowser/Private/ContentBrowserUtils.cpp
+++ b/Engine/Source/Editor/ContentBrowser/Private/ContentBrowserUtils.cpp
@@ -600,7 +600,7 @@ void ContentBrowserUtils::AppendAssetFilterToContentBrowserFilter(const FARFilte
 		if (InAssetFilter.bRecursiveClasses)
 		{
 			ClassFilter.ClassNamesToExclude = InAssetFilter.RecursiveClassesExclusionSet.Array();
-			ClassFilter.bRecursiveClassNamesToExclude = false;
+			ClassFilter.bRecursiveClassNamesToExclude = InAssetFilter.bRecursiveClasses;
 		}
 		ClassFilter.ClassPermissionList = InAssetClassPermissionList;
 	}
