[VK] Add Public Asset UX Phase 1
- Allows for toggling an individual Asset's Public/Private state in the Content Browser by adding/removing the flag PKG_NotExternallyReferenceable
- Content Browser Tooltips now show Asset states (Public, Private, Read Only)
- Includes CVar ContentBrowser_EnablePublicAssetFeatureCVar to hide UI until we reach MVP

diff --git a/Engine/Source/Editor/ContentBrowser/Private/AssetContextMenu.cpp b/Engine/Source/Editor/ContentBrowser/Private/AssetContextMenu.cpp
--- a/Engine/Source/Editor/ContentBrowser/Private/AssetContextMenu.cpp
+++ b/Engine/Source/Editor/ContentBrowser/Private/AssetContextMenu.cpp
@@ -23,6 +23,12 @@
 #include "Framework/Commands/GenericCommands.h"
 #include "ContentBrowserCommands.h"
 
+namespace ContentBrowserConsoleVariables
+{
+	static bool ContentBrowser_EnablePublicAssetFeature = false;
+	static FAutoConsoleVariableRef ContentBrowser_EnablePublicAssetFeatureCVar(TEXT("ContentBrowser.PublicAsset.EnablePublicAssetFeature"),
+		ContentBrowser_EnablePublicAssetFeature, TEXT("Enables the Experimental Public Asset Feature (False: disabled, True:enabled"));
+}
 #define LOCTEXT_NAMESPACE "ContentBrowser"
 
 FAssetContextMenu::FAssetContextMenu(const TWeakPtr<SAssetView>& InAssetView)
@@ -451,6 +457,22 @@ bool FAssetContextMenu::AddReferenceMenuOptions(UToolMenu* Menu)
 	{
 		FToolMenuSection& Section = Menu->AddSection("AssetContextReferences", LOCTEXT("ReferencesMenuHeading", "References"));
 
+		if (ContentBrowserConsoleVariables::ContentBrowser_EnablePublicAssetFeature)
+		{
+			Section.AddMenuEntry(
+				"PublicAsset",
+				LOCTEXT("PublicAssetToggle", "Public Asset"),
+				LOCTEXT("PublicAssetToggleTooltip", "Sets whether or not an asset is publicly available for reference by other plugins"),
+				FSlateIcon(FEditorStyle::GetStyleSetName(), "ContentBrowser.AssetActions.PublicAssetToggle"),
+				FUIAction(
+					FExecuteAction::CreateSP(this, &FAssetContextMenu::ExecutePublicAssetToggle),
+					FCanExecuteAction::CreateSP(this, &FAssetContextMenu::CanExecutePublicAssetToggle),
+					FGetActionCheckState::CreateSP(this, &FAssetContextMenu::GetPublicAssetCheckState)
+				),
+				EUserInterfaceActionType::ToggleButton
+			);
+		}
+
 		Section.AddMenuEntry(
 			"CopyReference",
 			LOCTEXT("CopyReference", "Copy Reference"),
@@ -832,6 +854,52 @@ FReply FAssetContextMenu::ExecuteDeleteFolderConfirmed()
 	return FReply::Handled();
 }
 
+void FAssetContextMenu::ExecutePublicAssetToggle()
+{
+	// Toggle selected files public asset flag
+	if (ensure(SelectedFiles.Num() == 1))
+	{
+		FAssetData ItemAssetData;
+		if (SelectedFiles[0].Legacy_TryGetAssetData(ItemAssetData))
+		{
+			UPackage* ItemAssetPackage = ItemAssetData.GetPackage();
+
+			if (!ItemAssetPackage)
+			{
+				return;
+			}
+
+			ItemAssetPackage->SetIsExternallyReferenceable(!ItemAssetPackage->IsExternallyReferenceable());
+
+			ItemAssetPackage->SetDirtyFlag(true);
+		}
+	}
+}
+
+bool FAssetContextMenu::CanExecutePublicAssetToggle()
+{
+	return SelectedFiles.Num() == 1 && SelectedFiles[0].CanEdit();
+}
+
+ECheckBoxState FAssetContextMenu::GetPublicAssetCheckState()
+{
+	if (ensure(SelectedFiles.Num() == 1))
+	{
+		FAssetData ItemAssetData;
+		if (SelectedFiles[0].Legacy_TryGetAssetData(ItemAssetData))
+		{
+			UPackage* ItemAssetPackage = ItemAssetData.GetPackage();
+
+			if (ItemAssetPackage)
+			{
+				return ItemAssetPackage->IsExternallyReferenceable() ? ECheckBoxState::Checked : ECheckBoxState::Unchecked;
+			}
+		}
+	}
+
+	return ECheckBoxState::Undetermined;
+}
+
 void FAssetContextMenu::ExecuteCopyReference()
 {
 	if (SelectedFiles.Num() > 0)
diff --git a/Engine/Source/Editor/ContentBrowser/Private/AssetContextMenu.h b/Engine/Source/Editor/ContentBrowser/Private/AssetContextMenu.h
--- a/Engine/Source/Editor/ContentBrowser/Private/AssetContextMenu.h
+++ b/Engine/Source/Editor/ContentBrowser/Private/AssetContextMenu.h
@@ -114,6 +114,15 @@ private:
 	/** Get tooltip for delete */
 	FText GetDeleteToolTip() const;
 
+	/** Handler for when "Public Asset" is toggled */
+	void ExecutePublicAssetToggle();
+
+	/** Handler to check to see if a Public Asset toggle is allowed */
+	bool CanExecutePublicAssetToggle();
+
+	/** Handler for getting the Public Asset check state */
+	ECheckBoxState GetPublicAssetCheckState();
+
 	/** Handler for CopyReference */
 	void ExecuteCopyReference();
 
diff --git a/Engine/Source/Editor/ContentBrowser/Private/AssetViewWidgets.cpp b/Engine/Source/Editor/ContentBrowser/Private/AssetViewWidgets.cpp
--- a/Engine/Source/Editor/ContentBrowser/Private/AssetViewWidgets.cpp
+++ b/Engine/Source/Editor/ContentBrowser/Private/AssetViewWidgets.cpp
@@ -744,6 +744,7 @@ TSharedRef<SWidget> SAssetViewItem::CreateToolTipWidget() const
 			// The tooltip contains the name, class, path, and asset registry tags
 			const FText NameText = GetNameText();
 			const FText ClassText = FText::Format(LOCTEXT("ClassName", "({0})"), GetAssetClassText());
+			FText PublicStateText = LOCTEXT("PublicAssetState", "Public");
 
 			// Create a box to hold every line of info in the body of the tooltip
 			TSharedRef<SVerticalBox> InfoBox = SNew(SVerticalBox);
@@ -775,6 +776,18 @@ TSharedRef<SWidget> SAssetViewItem::CreateToolTipWidget() const
 				int32 MaxCookPathLen = ContentBrowserUtils::GetMaxCookPathLen();
 				AddToToolTipInfoBox(InfoBox, LOCTEXT("TileViewTooltipPathLengthForCookingKey", "Cooking Filepath Length"), FText::Format(LOCTEXT("TileViewTooltipPathLengthForCookingValue", "{0} / {1}"),
 					FText::AsNumber(PackageNameLengthForCooking), FText::AsNumber(MaxCookPathLen)), PackageNameLengthForCooking > MaxCookPathLen ? true : false);
+
+				UPackage* ItemAssetPackage = ItemAssetData.GetPackage();
+
+				if (ItemAssetPackage && !ItemAssetPackage->IsExternallyReferenceable())
+				{
+					PublicStateText = LOCTEXT("PrivateAssetState", "Private");
+				}
+			}
+
+			if (!AssetItem->GetItem().CanEdit())
+			{
+				PublicStateText = LOCTEXT("ReadOnlyAssetState", "Read Only");
 			}
 
 			// Add tags
@@ -803,6 +816,14 @@ TSharedRef<SWidget> SAssetViewItem::CreateToolTipWidget() const
 
 			TSharedRef<SVerticalBox> OverallTooltipVBox = SNew(SVerticalBox);
 
+			static const auto PublicAssetUIEnabledCVar = IConsoleManager::Get().FindConsoleVariable(TEXT("ContentBrowser.PublicAsset.EnablePublicAssetFeature"));
+			bool bIsPublicAssetUIEnabled = false;
+
+			if (PublicAssetUIEnabledCVar)
+			{
+				bIsPublicAssetUIEnabled = PublicAssetUIEnabledCVar->GetBool();
+			}		 
+
 			// Top section (asset name, type, is checked out)
 			OverallTooltipVBox->AddSlot()
 				.AutoHeight()
@@ -837,6 +858,23 @@ TSharedRef<SWidget> SAssetViewItem::CreateToolTipWidget() const
 								.Text(ClassText)
 								.HighlightText(HighlightText)
 							]
+							+ SHorizontalBox::Slot()
+							.FillWidth(1.0)
+							.VAlign(VAlign_Center)
+							.HAlign(HAlign_Right) 
+							[
+								SNew(SBorder)
+								.BorderImage(FEditorStyle::GetBrush("ContentBrowser.TileViewTooltip.PillBorder"))
+								.Padding(FMargin(12.0f, 2.0f, 12.0f, 2.0f))
+								.Visibility(bIsPublicAssetUIEnabled ? EVisibility::Visible : EVisibility::Hidden)
+								[
+									SNew(STextBlock)
+									.Text(PublicStateText)
+									.HighlightText(HighlightText)
+								    .Font(FEditorStyle::GetFontStyle("ContentBrowser.TileViewTooltip.NameFont"))
+								]
+								
+							]
 						]
 
 						+ SVerticalBox::Slot()
diff --git a/Engine/Source/Editor/EditorStyle/Private/SlateEditorStyle.cpp b/Engine/Source/Editor/EditorStyle/Private/SlateEditorStyle.cpp
--- a/Engine/Source/Editor/EditorStyle/Private/SlateEditorStyle.cpp
+++ b/Engine/Source/Editor/EditorStyle/Private/SlateEditorStyle.cpp
@@ -6652,6 +6652,7 @@ void FSlateEditorStyle::FStyle::SetupContentBrowserStyle()
 		Set( "ContentBrowser.TileViewTooltip.ToolTipBorder", new FSlateColorBrush( FLinearColor::Black ) );
 		Set( "ContentBrowser.TileViewTooltip.NonContentBorder", new BOX_BRUSH( "/Docking/TabContentArea", FMargin(4/16.0f) ) );
 		Set( "ContentBrowser.TileViewTooltip.ContentBorder", new BOX_BRUSH( "Common/GroupBorder", FMargin(4.0f/16.0f) ) );
+		Set( "ContentBrowser.TileViewTooltip.PillBorder", new FSlateRoundedBoxBrush(FLinearColor::Transparent, 12.0f, FLinearColor::White, 0.5f));
 		Set( "ContentBrowser.TileViewTooltip.NameFont", DEFAULT_FONT( "Regular", 12 ) );
 		Set( "ContentBrowser.TileViewTooltip.AssetUserDescriptionFont", DEFAULT_FONT("Regular", 12 ) );
 
@@ -6778,6 +6779,7 @@ void FSlateEditorStyle::FStyle::SetupContentBrowserStyle()
 		Set( "ContentBrowser.AssetActions.Duplicate", new IMAGE_BRUSH( "Icons/Edit/icon_Edit_Duplicate_16x", Icon16x16) );
 		Set( "ContentBrowser.AssetActions.OpenSourceLocation", new IMAGE_BRUSH( "Icons/icon_Asset_Open_Source_Location_16x", Icon16x16) );
 		Set( "ContentBrowser.AssetActions.OpenInExternalEditor", new IMAGE_BRUSH( "Icons/icon_Asset_Open_In_External_Editor_16x", Icon16x16) );
+		Set( "ContentBrowser.AssetActions.PublicAssetToggle", new IMAGE_BRUSH("Icons/hiererchy_16x", Icon16x16));
 		Set( "ContentBrowser.AssetActions.ReimportAsset", new IMAGE_BRUSH( "Icons/icon_TextureEd_Reimport_40x", Icon16x16 ) );
 		Set( "ContentBrowser.AssetActions.GoToCodeForAsset", new IMAGE_BRUSH( "GameProjectDialog/feature_code_32x", Icon16x16 ) );
 		Set( "ContentBrowser.AssetActions.FindAssetInWorld", new IMAGE_BRUSH( "/Icons/icon_Genericfinder_16x", Icon16x16 ) );
diff --git a/Engine/Source/Editor/EditorStyle/Private/StarshipStyle.cpp b/Engine/Source/Editor/EditorStyle/Private/StarshipStyle.cpp
--- a/Engine/Source/Editor/EditorStyle/Private/StarshipStyle.cpp
+++ b/Engine/Source/Editor/EditorStyle/Private/StarshipStyle.cpp
@@ -6323,6 +6323,7 @@ void FStarshipEditorStyle::FStyle::SetupContentBrowserStyle()
 		Set( "ContentBrowser.TileViewTooltip.ToolTipBorder", new FSlateColorBrush( FLinearColor::Black ) );
 		Set( "ContentBrowser.TileViewTooltip.NonContentBorder", new BOX_BRUSH( "/Docking/TabContentArea", FMargin(4/16.0f) ) );
 		Set( "ContentBrowser.TileViewTooltip.ContentBorder", new FSlateColorBrush( FStyleColors::Panel));
+		Set("ContentBrowser.TileViewTooltip.PillBorder", new FSlateRoundedBoxBrush(FStyleColors::Transparent, 12.0f, FStyleColors::White, 0.5f));
 		Set( "ContentBrowser.TileViewTooltip.NameFont", DEFAULT_FONT( "Regular", 12 ) );
 		Set( "ContentBrowser.TileViewTooltip.AssetUserDescriptionFont", DEFAULT_FONT("Regular", 12 ) );
 
@@ -6441,6 +6442,7 @@ void FStarshipEditorStyle::FStyle::SetupContentBrowserStyle()
 		Set( "ContentBrowser.AssetActions.Duplicate", new IMAGE_BRUSH( "Icons/Edit/icon_Edit_Duplicate_16x", Icon16x16) );
 		Set( "ContentBrowser.AssetActions.OpenSourceLocation", new IMAGE_BRUSH( "Icons/icon_Asset_Open_Source_Location_16x", Icon16x16) );
 		Set( "ContentBrowser.AssetActions.OpenInExternalEditor", new IMAGE_BRUSH( "Icons/icon_Asset_Open_In_External_Editor_16x", Icon16x16) );
+		Set( "ContentBrowser.AssetActions.PublicAssetToggle", new IMAGE_BRUSH("Icons/hiererchy_16x", Icon16x16));
 		Set( "ContentBrowser.AssetActions.ReimportAsset", new IMAGE_BRUSH( "Icons/icon_TextureEd_Reimport_40x", Icon16x16 ) );
 		Set( "ContentBrowser.AssetActions.GoToCodeForAsset", new IMAGE_BRUSH( "GameProjectDialog/feature_code_32x", Icon16x16 ) );
 		Set( "ContentBrowser.AssetActions.FindAssetInWorld", new IMAGE_BRUSH( "/Icons/icon_Genericfinder_16x", Icon16x16 ) );
diff --git a/Engine/Source/Runtime/CoreUObject/Public/UObject/ObjectMacros.h b/Engine/Source/Runtime/CoreUObject/Public/UObject/ObjectMacros.h
--- a/Engine/Source/Runtime/CoreUObject/Public/UObject/ObjectMacros.h
+++ b/Engine/Source/Runtime/CoreUObject/Public/UObject/ObjectMacros.h
@@ -117,7 +117,7 @@ enum EPackageFlags
 	PKG_UncookedOnly				= 0x00000100,	///< Loaded only in uncooked builds (i.e. runtime in editor)
 	PKG_Cooked						= 0x00000200,	///< Package is cooked
 	PKG_ContainsNoAsset				= 0x00000400,	///< Package doesn't contain any asset object (although asset tags can be present)
-	PKG_ExternallyReferenceable		= 0x00000800,   ///< (Not Implemented) Objects in this package can be referenced in a different plugin or mount point (i.e /Game -> /Engine)
+	PKG_NotExternallyReferenceable  = 0x00000800,   ///< (Not Implemented) Objects in this package cannot be referenced in a different plugin or mount point (i.e /Game -> /Engine)
 //	PKG_Unused						= 0x00001000,
 	PKG_UnversionedProperties		= 0x00002000,   ///< Uses unversioned property serialization instead of versioned tagged property serialization
 	PKG_ContainsMapData				= 0x00004000,   ///< Contains map data (UObjects only referenced by a single ULevel) but is stored in a different package
diff --git a/Engine/Source/Runtime/CoreUObject/Public/UObject/Package.h b/Engine/Source/Runtime/CoreUObject/Public/UObject/Package.h
--- a/Engine/Source/Runtime/CoreUObject/Public/UObject/Package.h
+++ b/Engine/Source/Runtime/CoreUObject/Public/UObject/Package.h
@@ -634,6 +634,31 @@ public:
 		return PackageFlagsPrivate;
 	}
 
+	/**
+	* @return true if the package is marked as ExternallyReferenceable
+	*/
+	FORCEINLINE bool IsExternallyReferenceable() const
+	{
+		return (PackageFlagsPrivate & PKG_NotExternallyReferenceable) == 0;
+	}
+
+	/**
+	* Sets whether or not the package is ExternallyReferenceable
+	* 
+	* @param bValue Sets the package to be ExternallyReferenceable if true or NotExternallyReferenceable if false
+	*/
+	FORCEINLINE void SetIsExternallyReferenceable(bool bValue)
+	{
+		if (bValue)
+		{
+			ClearPackageFlags(PKG_NotExternallyReferenceable);
+		}
+		else
+		{
+			SetPackageFlags(PKG_NotExternallyReferenceable);
+		}
+	}
+
 #if WITH_EDITORONLY_DATA
 	/** Returns true if this package has a thumbnail map */
 	bool HasThumbnailMap() const
