[VK] Add forward declare to AssetContextMenu to resolve NonUnity build failure

diff --git a/Engine/Source/Editor/ContentBrowser/Private/AssetContextMenu.h b/Engine/Source/Editor/ContentBrowser/Private/AssetContextMenu.h
--- a/Engine/Source/Editor/ContentBrowser/Private/AssetContextMenu.h
+++ b/Engine/Source/Editor/ContentBrowser/Private/AssetContextMenu.h
@@ -14,6 +14,7 @@ class FUICommandList;
 class SAssetView;
 class SWindow;
 
+enum class ECheckBoxState : uint8;
 enum class EContentBrowserViewContext : uint8;
 
 class FAssetContextMenu : public TSharedFromThis<FAssetContextMenu>
