Merging CL#18894607

[VK] Bandaid to prevent the loading of assets in ContentBrowser on hover. This disables the experimental Public/Private asset tooltip state UI from refreshing correctly when toggled.

diff --git a/Engine/Source/Editor/ContentBrowser/Private/AssetViewWidgets.cpp b/Engine/Source/Editor/ContentBrowser/Private/AssetViewWidgets.cpp
--- a/Engine/Source/Editor/ContentBrowser/Private/AssetViewWidgets.cpp
+++ b/Engine/Source/Editor/ContentBrowser/Private/AssetViewWidgets.cpp
@@ -777,9 +777,8 @@ TSharedRef<SWidget> SAssetViewItem::CreateToolTipWidget() const
 				AddToToolTipInfoBox(InfoBox, LOCTEXT("TileViewTooltipPathLengthForCookingKey", "Cooking Filepath Length"), FText::Format(LOCTEXT("TileViewTooltipPathLengthForCookingValue", "{0} / {1}"),
 					FText::AsNumber(PackageNameLengthForCooking), FText::AsNumber(MaxCookPathLen)), PackageNameLengthForCooking > MaxCookPathLen ? true : false);
 
-				UPackage* ItemAssetPackage = ItemAssetData.GetPackage();
 
-				if (ItemAssetPackage && !ItemAssetPackage->IsExternallyReferenceable())
+				if (ItemAssetData.PackageFlags & PKG_NotExternallyReferenceable)
 				{
 					PublicStateText = LOCTEXT("PrivateAssetState", "Private");
 				}
