Asset View: Quick jumping now affects navigation selection

Previously, quick jumping would not properly move the selection, making it so that navigation after jumping would ignore the jump.

diff --git a/Engine/Source/Editor/ContentBrowser/Private/SAssetView.cpp b/Engine/Source/Editor/ContentBrowser/Private/SAssetView.cpp
--- a/Engine/Source/Editor/ContentBrowser/Private/SAssetView.cpp
+++ b/Engine/Source/Editor/ContentBrowser/Private/SAssetView.cpp
@@ -4198,8 +4198,9 @@ bool SAssetView::PerformQuickJump(const bool bWasJumping)
 			const FString& NewSelectedItemName = NewSelectedItem->GetItem().GetDisplayName().ToString();
 			if(NewSelectedItemName.StartsWith(QuickJumpData.JumpTerm, ESearchCase::IgnoreCase))
 			{
-				SetSelection(NewSelectedItem);
 				RequestScrollIntoView(NewSelectedItem);
+				// Consider it derived from a keypress because otherwise it won't update the navigation selector
+				SetItemSelection(NewSelectedItem, true, ESelectInfo::Type::OnKeyPress);
 				return true;
 			}
 		}
