Fix for content browser quick jump expanding selection

diff --git a/Engine/Source/Editor/ContentBrowser/Private/SAssetView.cpp b/Engine/Source/Editor/ContentBrowser/Private/SAssetView.cpp
--- a/Engine/Source/Editor/ContentBrowser/Private/SAssetView.cpp
+++ b/Engine/Source/Editor/ContentBrowser/Private/SAssetView.cpp
@@ -4199,6 +4199,7 @@ bool SAssetView::PerformQuickJump(const bool bWasJumping)
 			if(NewSelectedItemName.StartsWith(QuickJumpData.JumpTerm, ESearchCase::IgnoreCase))
 			{
 				RequestScrollIntoView(NewSelectedItem);
+				ClearSelection();
 				// Consider it derived from a keypress because otherwise it won't update the navigation selector
 				SetItemSelection(NewSelectedItem, true, ESelectInfo::Type::OnKeyPress);
 				return true;
