Fixed quick jump extending the selection rather than setting it

#rnx

diff --git a/Engine/Source/Editor/ContentBrowser/Private/SAssetView.cpp b/Engine/Source/Editor/ContentBrowser/Private/SAssetView.cpp
--- a/Engine/Source/Editor/ContentBrowser/Private/SAssetView.cpp
+++ b/Engine/Source/Editor/ContentBrowser/Private/SAssetView.cpp
@@ -4198,6 +4198,7 @@ bool SAssetView::PerformQuickJump(const bool bWasJumping)
 			const FString& NewSelectedItemName = NewSelectedItem->GetItem().GetDisplayName().ToString();
 			if(NewSelectedItemName.StartsWith(QuickJumpData.JumpTerm, ESearchCase::IgnoreCase))
 			{
+				ClearSelection(true);
 				RequestScrollIntoView(NewSelectedItem);
 				ClearSelection();
 				// Consider it derived from a keypress because otherwise it won't update the navigation selector
