Content Browser: Added functionality to summon other Content Browsers by right clicking on the tab

diff --git a/Engine/Source/Editor/ContentBrowser/Private/ContentBrowserSingleton.cpp b/Engine/Source/Editor/ContentBrowser/Private/ContentBrowserSingleton.cpp
--- a/Engine/Source/Editor/ContentBrowser/Private/ContentBrowserSingleton.cpp
+++ b/Engine/Source/Editor/ContentBrowser/Private/ContentBrowserSingleton.cpp
@@ -801,8 +801,9 @@ TSharedRef<SDockTab> FContentBrowserSingleton::SpawnContentBrowserTab( const FSp
 	TSharedRef<SDockTab> NewTab = SNew(SDockTab)
 		.TabRole(ETabRole::NomadTab)
 		.Label( Label )
-		.ToolTip( IDocumentation::Get()->CreateToolTip( Label, nullptr, "Shared/ContentBrowser", "Tab" ) );
-
+		.ToolTip( IDocumentation::Get()->CreateToolTip( Label, nullptr, "Shared/ContentBrowser", "Tab" ) )
+		.OnExtendContextMenu_Raw(this, &FContentBrowserSingleton::ExtendContentBrowserTabContextMenu);
+	
 	TSharedRef<SWidget> NewBrowser = CreateContentBrowser( SpawnTabArgs.GetTabId().TabType, NewTab, nullptr );
 
 	// Add wrapper for tutorial highlighting
@@ -818,6 +819,62 @@ TSharedRef<SDockTab> FContentBrowserSingleton::SpawnContentBrowserTab( const FSp
 	return NewTab;
 }
 
+void FContentBrowserSingleton::ExtendContentBrowserTabContextMenu(FMenuBuilder& InMenuBuilder)
+{
+	InMenuBuilder.BeginSection("SummonContentBrowserTabs", LOCTEXT("ContentBrowserTabs", "Content Browser Tabs"));
+
+	for (int32 BrowserIdx = 0; BrowserIdx < UE_ARRAY_COUNT(ContentBrowserTabIDs); BrowserIdx++)
+	{
+		const FName TabID = ContentBrowserTabIDs[BrowserIdx];
+		const FText DefaultDisplayName = GetContentBrowserLabelWithIndex(BrowserIdx);
+
+		InMenuBuilder.AddMenuEntry(DefaultDisplayName,
+			LOCTEXT("ContentBrowserMenuTooltipText", "Open a Content Browser tab."),
+			FSlateIcon(),
+			FUIAction(
+					FExecuteAction::CreateLambda([TabID, this]()
+					{
+						// Go through all the content browsers to check if the current one is open
+						for (int32 BrowserIdx = 0; BrowserIdx < AllContentBrowsers.Num(); ++BrowserIdx)
+						{
+							const TWeakPtr<SContentBrowser>& Browser = AllContentBrowsers[BrowserIdx];
+
+							if (Browser.IsValid() && Browser.Pin()->GetInstanceName() == TabID)
+							{
+								FocusContentBrowser(Browser.Pin()); // Focus it if so
+							}
+						}
+
+						// If the tab was not found, try to open it
+						FGlobalTabmanager::Get()->TryInvokeTab(TabID);
+					}
+				),
+					FCanExecuteAction(),
+					FIsActionChecked::CreateLambda([this, TabID]()
+					{
+						// Go through all the content browsers to check if the current one is open
+						for (int32 BrowserIdx = 0; BrowserIdx < AllContentBrowsers.Num(); ++BrowserIdx)
+						{
+							const TWeakPtr<SContentBrowser>& Browser = AllContentBrowsers[BrowserIdx];
+
+							if (Browser.IsValid() && Browser.Pin()->GetInstanceName() == TabID)
+							{
+								return true;
+							}
+						}
+
+						return false;
+						
+					}
+				)),
+			NAME_None,
+			EUserInterfaceActionType::ToggleButton);
+	}
+
+	InMenuBuilder.EndSection();
+}
+
+
 bool FContentBrowserSingleton::IsLocked(const FName& InstanceName) const
 {
 	// First try all the open browsers, as their locked state might be newer than the configs
diff --git a/Engine/Source/Editor/ContentBrowser/Private/ContentBrowserSingleton.h b/Engine/Source/Editor/ContentBrowser/Private/ContentBrowserSingleton.h
--- a/Engine/Source/Editor/ContentBrowser/Private/ContentBrowserSingleton.h
+++ b/Engine/Source/Editor/ContentBrowser/Private/ContentBrowserSingleton.h
@@ -140,6 +140,8 @@ private:
 	/** Creates the Content Browser submenu in the Level Editor Toolbar Add menu */
 	void GetContentBrowserSubMenu(UToolMenu* Menu, TSharedRef<FWorkspaceItem> ContentBrowserGroup);
 
+	void ExtendContentBrowserTabContextMenu(FMenuBuilder& InMenuBuilder);
+
 public:
 	/** The tab identifier/instance name for content browser tabs */
 	FName ContentBrowserTabIDs[MAX_CONTENT_BROWSERS];
