Add OnShouldFilterItem delegate to SAssetView

diff --git a/Engine/Source/Editor/ContentBrowser/Private/SAssetView.cpp b/Engine/Source/Editor/ContentBrowser/Private/SAssetView.cpp
--- a/Engine/Source/Editor/ContentBrowser/Private/SAssetView.cpp
+++ b/Engine/Source/Editor/ContentBrowser/Private/SAssetView.cpp
@@ -117,10 +117,17 @@ public:
 			return true;
 		}
 
+		if (AssetView->OnShouldFilterItem.IsBound()
+			&& !AssetView->OnShouldFilterItem.Execute(InItemToFilter->GetItem()))
+		{
+			return true;
+		}
+
 		// If we have OnShouldFilterAsset then it is assumed that we really only want to see true assets and 
 		// nothing else so only include things that have asset data and also pass the query filter
 		FAssetData ItemAssetData;
-		if (InItemToFilter->GetItem().Legacy_TryGetAssetData(ItemAssetData))
+		if (AssetView->OnShouldFilterAsset.IsBound()
+			&& InItemToFilter->GetItem().Legacy_TryGetAssetData(ItemAssetData))
 		{
 			if (!AssetView->OnShouldFilterAsset.Execute(ItemAssetData))
 			{
@@ -273,6 +280,8 @@ void SAssetView::Construct( const FArguments& InArgs )
 	}
 
 	OnShouldFilterAsset = InArgs._OnShouldFilterAsset;
+	OnShouldFilterItem = InArgs._OnShouldFilterItem;
+
 	OnNewItemRequested = InArgs._OnNewItemRequested;
 	OnItemSelectionChanged = InArgs._OnItemSelectionChanged;
 	OnItemsActivated = InArgs._OnItemsActivated;
@@ -1255,7 +1264,7 @@ void SAssetView::ProcessItemsPendingFilter(const double TickStartTime)
 		}
 	};
 
-	const bool bRunQueryFilter = OnShouldFilterAsset.IsBound();
+	const bool bRunQueryFilter = OnShouldFilterAsset.IsBound() || OnShouldFilterItem.IsBound();
 	const bool bFlushAllPendingItems = TickStartTime < 0;
 
 	bool bRefreshList = false;
@@ -1773,7 +1782,8 @@ FContentBrowserDataFilter SAssetView::CreateBackendDataFilter() const
 		|| OnGetCustomAssetToolTip.IsBound()
 		|| OnVisualizeAssetToolTip.IsBound()
 		|| OnAssetToolTipClosing.IsBound()
-		|| OnShouldFilterAsset.IsBound();
+		|| OnShouldFilterAsset.IsBound()
+		|| OnShouldFilterItem.IsBound();
 
 	FContentBrowserDataFilter DataFilter;
 	DataFilter.bRecursivePaths = bRecurse || !bUsingFolders || bHasCollections;
@@ -4530,7 +4540,7 @@ void SAssetView::HandleItemDataUpdated(TArrayView<const FContentBrowserItemDataU
 	if (ItemsPendingInplaceFrontendFilter.Num() > 0)
 	{
 		FAssetViewFrontendFilterHelper FrontendFilterHelper(this);
-		const bool bRunQueryFilter = OnShouldFilterAsset.IsBound();
+		const bool bRunQueryFilter = OnShouldFilterAsset.IsBound() || OnShouldFilterItem.IsBound();
 
 		for (auto It = FilteredAssetItems.CreateIterator(); It && ItemsPendingInplaceFrontendFilter.Num() > 0; ++It)
 		{
diff --git a/Engine/Source/Editor/ContentBrowser/Public/ContentBrowserDelegates.h b/Engine/Source/Editor/ContentBrowser/Public/ContentBrowserDelegates.h
--- a/Engine/Source/Editor/ContentBrowser/Public/ContentBrowserDelegates.h
+++ b/Engine/Source/Editor/ContentBrowser/Public/ContentBrowserDelegates.h
@@ -37,6 +37,9 @@ DECLARE_DELEGATE_RetVal_OneParam( TSharedRef< SToolTip >, FConstructToolTipForAs
 /** Called to check if an asset should be filtered out by external code. Return true to exclude the asset from the view. */
 DECLARE_DELEGATE_RetVal_OneParam(bool, FOnShouldFilterAsset, const FAssetData& /*AssetData*/);
 
+/** Called to check if an item should be filtered out by external code. Return true to exclude the item from the view. */
+DECLARE_DELEGATE_RetVal_OneParam(bool, FOnShouldFilterItem, const FContentBrowserItem& /*AssetItem*/);
+
 /** Called to check if an asset tag should be display in details view. Return false to exclude the asset from the view. */
 DECLARE_DELEGATE_RetVal_TwoParams(bool, FOnShouldDisplayAssetTag, FName /*AssetType*/, FName /*TagName*/);
 
@@ -122,6 +125,9 @@ DECLARE_DELEGATE_OneParam(FOnAssetsChosenForOpen, const TArray<FAssetData>& /*Se
 /** Called from the Asset Dialog when an asset name is chosen in non-modal Save dialogs */
 DECLARE_DELEGATE_OneParam(FOnObjectPathChosenForSave, const FString& /*ObjectPath*/);
 
+/** Called when custom behavior is needed for allowing folders to toggle their private content edit state */
+DECLARE_DELEGATE_RetVal_OneParam(bool, FIsFolderShowPrivateContentToggleableDelegate, const FStringView /*VirtualFolderPath*/);
+
 /** Contains the delegates used to handle a custom drag-and-drop in the asset view */
 struct FAssetViewDragAndDropExtender
 {
diff --git a/Engine/Source/Editor/ContentBrowser/Public/SAssetView.h b/Engine/Source/Editor/ContentBrowser/Public/SAssetView.h
--- a/Engine/Source/Editor/ContentBrowser/Public/SAssetView.h
+++ b/Engine/Source/Editor/ContentBrowser/Public/SAssetView.h
@@ -86,6 +86,9 @@ public:
 		/** Called to check if an asset should be filtered out by external code */
 		SLATE_EVENT( FOnShouldFilterAsset, OnShouldFilterAsset )
 
+		/** Called to check if an item should be filtered out by external code */
+		SLATE_EVENT(FOnShouldFilterItem, OnShouldFilterItem)
+
 		/** Called when the asset view is asked to start to create a temporary item */
 		SLATE_EVENT( FOnAssetViewNewItemRequested, OnNewItemRequested )
 
@@ -878,6 +881,9 @@ private:
 	/** Called to check if an asset should be filtered out by external code */
 	FOnShouldFilterAsset OnShouldFilterAsset;
 
+	/** Called to check if an item should be filtered out by external code */
+	FOnShouldFilterItem OnShouldFilterItem;
+
 	/** Called when the asset view is asked to start to create a temporary item */
 	FOnAssetViewNewItemRequested OnNewItemRequested;
 
