Updated font and padding for Public Asset tooltip UI

diff --git a/Engine/Source/Editor/ContentBrowser/Private/AssetViewWidgets.cpp b/Engine/Source/Editor/ContentBrowser/Private/AssetViewWidgets.cpp
--- a/Engine/Source/Editor/ContentBrowser/Private/AssetViewWidgets.cpp
+++ b/Engine/Source/Editor/ContentBrowser/Private/AssetViewWidgets.cpp
@@ -858,7 +858,8 @@ TSharedRef<SWidget> SAssetViewItem::CreateToolTipWidget() const
 								.HighlightText(HighlightText)
 							]
 							+ SHorizontalBox::Slot()
-							.FillWidth(1.0)
+							.FillWidth(1.0f)
+							.Padding(10.0f, 4.0f)
 							.VAlign(VAlign_Center)
 							.HAlign(HAlign_Right) 
 							[
@@ -870,7 +871,6 @@ TSharedRef<SWidget> SAssetViewItem::CreateToolTipWidget() const
 									SNew(STextBlock)
 									.Text(PublicStateText)
 									.HighlightText(HighlightText)
-								    .Font(FEditorStyle::GetFontStyle("ContentBrowser.TileViewTooltip.NameFont"))
 								]
 								
 							]
