Phase 2 Public Assets
- Added ability to toggle Private Content Edit mode at Folder level in Content Browser
- You cannot edit an asset's Private/Public state unless it's within a Private Content Editable folder
- Private Assets do not show up unless their within a Private Content Editable folder

diff --git a/Engine/Source/Editor/ContentBrowser/Private/ContentBrowserSingleton.cpp b/Engine/Source/Editor/ContentBrowser/Private/ContentBrowserSingleton.cpp
--- a/Engine/Source/Editor/ContentBrowser/Private/ContentBrowserSingleton.cpp
+++ b/Engine/Source/Editor/ContentBrowser/Private/ContentBrowserSingleton.cpp
@@ -85,6 +85,9 @@ FContentBrowserSingleton::FContentBrowserSingleton()
 	FContentBrowserCommands::Register();
 
 	PopulateConfigValues();
+
+	ShowPrivateContentState.InvariantPaths = MakeShared<FPathPermissionList>();
+	ShowPrivateContentState.CachedVirtualPaths = MakeShared<FPathPermissionList>();
 }
 
 FContentBrowserSingleton::~FContentBrowserSingleton()
@@ -1040,6 +1043,66 @@ void FContentBrowserSingleton::RefreshPathView(TSharedPtr<SWidget> Widget)
 	}
 }
 
+bool FContentBrowserSingleton::IsShowingPrivateContent(const FStringView VirtualFolderPath)
+{
+	if (!ShowPrivateContentState.InvariantPaths->HasFiltering())
+	{
+		return false;
+	}
+
+	if (!ShowPrivateContentState.CachedVirtualPaths)
+	{
+		RebuildPrivateContentStateCache();
+	}
+
+	return ShowPrivateContentState.CachedVirtualPaths->PassesStartsWithFilter(VirtualFolderPath);
+}
+
+void FContentBrowserSingleton::RebuildPrivateContentStateCache()
+{
+	ShowPrivateContentState.CachedVirtualPaths.Reset();
+	ShowPrivateContentState.CachedVirtualPaths = MakeShared<FPathPermissionList>();
+
+	const auto& InvariantAllowList = ShowPrivateContentState.InvariantPaths->GetAllowList();
+	for (const TPair<FString, FPermissionListOwners>& InvariantPathOwnerPair : InvariantAllowList)
+	{
+		FName VirtualPath;
+		IContentBrowserDataModule::Get().GetSubsystem()->ConvertInternalPathToVirtual(FStringView(InvariantPathOwnerPair.Key), VirtualPath);
+
+		ShowPrivateContentState.CachedVirtualPaths->AddAllowListItem(TEXT("ContentBrowser"), VirtualPath);
+	}
+}
+
+bool FContentBrowserSingleton::IsFolderShowPrivateContentToggleable(const FStringView VirtualFolderPath)
+{
+	if (IsFolderShowPrivateContentToggleableDelegate.IsBound())
+	{
+		return IsFolderShowPrivateContentToggleableDelegate.Execute(VirtualFolderPath);
+	}
+
+	return true;
+}
+
+const TSharedPtr<FPathPermissionList>& FContentBrowserSingleton::GetShowPrivateContentPermissionList()
+{
+	return ShowPrivateContentState.InvariantPaths;
+}
+
+void FContentBrowserSingleton::SetPrivateContentPermissionListDirty()
+{
+	ShowPrivateContentState.CachedVirtualPaths.Reset();
+}
+
+void FContentBrowserSingleton::RegisterIsFolderShowPrivateContentToggleableDelegate(FIsFolderShowPrivateContentToggleableDelegate InIsFolderShowPrivateContentToggleableDelegate)
+{
+	IsFolderShowPrivateContentToggleableDelegate = InIsFolderShowPrivateContentToggleableDelegate;
+}
+
+void FContentBrowserSingleton::UnregisterIsFolderShowPrivateContentToggleableDelegate()
+{
+	IsFolderShowPrivateContentToggleableDelegate = FIsFolderShowPrivateContentToggleableDelegate();
+}
+
 void FContentBrowserSingleton::PopulateConfigValues()
 {
 	const FString ContentBrowserSection = TEXT("ContentBrowser");
diff --git a/Engine/Source/Editor/ContentBrowser/Private/ContentBrowserSingleton.h b/Engine/Source/Editor/ContentBrowser/Private/ContentBrowserSingleton.h
--- a/Engine/Source/Editor/ContentBrowser/Private/ContentBrowserSingleton.h
+++ b/Engine/Source/Editor/ContentBrowser/Private/ContentBrowserSingleton.h
@@ -37,6 +37,12 @@ struct FContentBrowserPluginSettings
 	{}
 };
 
+struct FShowPrivateContentState
+{
+	TSharedPtr<FPathPermissionList> InvariantPaths;
+	TSharedPtr<FPathPermissionList> CachedVirtualPaths;
+};
+
 /**
  * Content browser module singleton implementation class
  */
@@ -81,6 +87,12 @@ PRAGMA_ENABLE_DEPRECATION_WARNINGS
 	virtual void ExecuteRename(TSharedPtr<SWidget> PickerWidget) override;
 	virtual void ExecuteAddFolder(TSharedPtr<SWidget> PathPickerWidget) override;
 	virtual void RefreshPathView(TSharedPtr<SWidget> PathPickerWidget) override;
+	virtual bool IsShowingPrivateContent(const FStringView VirtualFolderPath) override;
+	virtual bool IsFolderShowPrivateContentToggleable(const FStringView VirtualFolderPath) override;
+	virtual const TSharedPtr<FPathPermissionList>& GetShowPrivateContentPermissionList() override;
+	virtual void SetPrivateContentPermissionListDirty() override;
+	virtual void RegisterIsFolderShowPrivateContentToggleableDelegate(FIsFolderShowPrivateContentToggleableDelegate InIsFolderShowPrivateContentToggleableDelegate) override;
+	virtual void UnregisterIsFolderShowPrivateContentToggleableDelegate() override;
 
 	/** Gets the content browser singleton as a FContentBrowserSingleton */
 	static FContentBrowserSingleton& Get();
@@ -142,11 +154,17 @@ private:
 
 	void ExtendContentBrowserTabContextMenu(FMenuBuilder& InMenuBuilder);
 
+	/** Rebuilds the private content state cache based off of the currently registered Invariant Path permission list */
+	void RebuildPrivateContentStateCache();
+
 public:
 	/** The tab identifier/instance name for content browser tabs */
 	FName ContentBrowserTabIDs[MAX_CONTENT_BROWSERS];
 
 private:
+
+	FIsFolderShowPrivateContentToggleableDelegate IsFolderShowPrivateContentToggleableDelegate;
+
 	TArray<TWeakPtr<SContentBrowser>> AllContentBrowsers;
 
 	TWeakPtr<SContentBrowser> ContentBrowserDrawer;
@@ -165,4 +183,6 @@ private:
 
 	/** An incrementing int32 which is used when making unique settings strings */
 	int32 SettingsStringID;
+
+	FShowPrivateContentState ShowPrivateContentState;
 };
diff --git a/Engine/Source/Editor/ContentBrowser/Private/ContentBrowserUtils.cpp b/Engine/Source/Editor/ContentBrowser/Private/ContentBrowserUtils.cpp
--- a/Engine/Source/Editor/ContentBrowser/Private/ContentBrowserUtils.cpp
+++ b/Engine/Source/Editor/ContentBrowser/Private/ContentBrowserUtils.cpp
@@ -717,4 +717,42 @@ const TArray<FString>& ContentBrowserUtils::GetFavoriteFolders()
 	return FContentBrowserSingleton::Get().FavoriteFolderPaths;
 }
 
+void ContentBrowserUtils::AddShowPrivateContentFolder(const FStringView VirtualFolderPath, const FName Owner)
+{
+	FContentBrowserSingleton& ContentBrowserSingleton = FContentBrowserSingleton::Get();
+
+	if (!ContentBrowserSingleton.IsFolderShowPrivateContentToggleable(VirtualFolderPath))
+	{
+		return;
+	}
+
+	FName InvariantPath;
+	IContentBrowserDataModule::Get().GetSubsystem()->TryConvertVirtualPath(VirtualFolderPath, InvariantPath);
+
+	const TSharedPtr<FPathPermissionList>& ShowPrivateContentPermissionList = ContentBrowserSingleton.GetShowPrivateContentPermissionList();
+
+	ShowPrivateContentPermissionList->AddAllowListItem(Owner, InvariantPath);
+
+	ContentBrowserSingleton.SetPrivateContentPermissionListDirty();
+}
+
+void ContentBrowserUtils::RemoveShowPrivateContentFolder(const FStringView VirtualFolderPath, const FName Owner)
+{
+	FContentBrowserSingleton& ContentBrowserSingleton = FContentBrowserSingleton::Get();
+
+	if (!ContentBrowserSingleton.IsFolderShowPrivateContentToggleable(VirtualFolderPath))
+	{
+		return;
+	}
+
+	FName InvariantPath;
+	IContentBrowserDataModule::Get().GetSubsystem()->TryConvertVirtualPath(VirtualFolderPath, InvariantPath);
+
+	const TSharedPtr<FPathPermissionList>& ShowPrivateContentPermissionList = ContentBrowserSingleton.GetShowPrivateContentPermissionList();
+
+	ShowPrivateContentPermissionList->RemoveAllowListItem(Owner, InvariantPath);
+
+	ContentBrowserSingleton.SetPrivateContentPermissionListDirty();
+}
+
 #undef LOCTEXT_NAMESPACE
diff --git a/Engine/Source/Editor/ContentBrowser/Private/ContentBrowserUtils.h b/Engine/Source/Editor/ContentBrowser/Private/ContentBrowserUtils.h
--- a/Engine/Source/Editor/ContentBrowser/Private/ContentBrowserUtils.h
+++ b/Engine/Source/Editor/ContentBrowser/Private/ContentBrowserUtils.h
@@ -97,4 +97,10 @@ namespace ContentBrowserUtils
 	void RemoveFavoriteFolder(const FString& FolderPath, bool bFlushConfig = true);
 
 	const TArray<FString>& GetFavoriteFolders();
+
+	/** Adds FolderPath as a private content edit folder if it's allowed to be toggled as such */
+	void AddShowPrivateContentFolder(const FStringView VirtualFolderPath, const FName Owner);
+
+	/** Removes FolderPath as a private content edit folder if it's allowed to be toggled as such */
+	void RemoveShowPrivateContentFolder(const FStringView VirtualFolderPath, const FName Owner);
 }
diff --git a/Engine/Source/Editor/ContentBrowser/Private/PathContextMenu.cpp b/Engine/Source/Editor/ContentBrowser/Private/PathContextMenu.cpp
--- a/Engine/Source/Editor/ContentBrowser/Private/PathContextMenu.cpp
+++ b/Engine/Source/Editor/ContentBrowser/Private/PathContextMenu.cpp
@@ -65,6 +65,11 @@ void FPathContextMenu::SetOnFolderFavoriteToggled(const FOnFolderFavoriteToggled
 	OnFolderFavoriteToggled = InOnFolderFavoriteToggled;
 }
 
+void FPathContextMenu::SetOnPrivateContentEditToggled(const FOnPrivateContentEditToggled& InOnPrivateContentEditToggled)
+{
+	OnPrivateContentEditToggled = InOnPrivateContentEditToggled;
+}
+
 const TArray<FContentBrowserItem>& FPathContextMenu::GetSelectedFolders() const
 {
 	return SelectedFolders;
@@ -179,8 +184,9 @@ void FPathContextMenu::MakePathViewContextMenu(UToolMenu* Menu)
 					);
 			}			
 
+			FString SelectedFolderPath = SelectedFolders[0].GetVirtualPath().ToString();
 			// If this folder is already favorited, show the option to remove from favorites
-			if (ContentBrowserUtils::IsFavoriteFolder(SelectedFolders[0].GetVirtualPath().ToString()))
+			if (ContentBrowserUtils::IsFavoriteFolder(SelectedFolderPath))
 			{
 				// Remove from favorites
 				Section.AddMenuEntry(
@@ -202,6 +208,39 @@ void FPathContextMenu::MakePathViewContextMenu(UToolMenu* Menu)
 					FUIAction(FExecuteAction::CreateSP(this, &FPathContextMenu::ExecuteFavorite))
 				);
 			}
+
+			static const auto PublicAssetUIEnabledCVar = IConsoleManager::Get().FindConsoleVariable(TEXT("ContentBrowser.PublicAsset.EnablePublicAssetFeature"));
+			bool bIsPublicAssetUIEnabled = false;
+
+			if (PublicAssetUIEnabledCVar)
+			{
+				bIsPublicAssetUIEnabled = PublicAssetUIEnabledCVar->GetBool();
+			}
+
+			FStringView SelectedFolderPathView(SelectedFolderPath);
+			if (bIsPublicAssetUIEnabled && FContentBrowserSingleton::Get().IsFolderShowPrivateContentToggleable(SelectedFolderPathView))
+			{
+				if (FContentBrowserSingleton::Get().IsShowingPrivateContent(SelectedFolderPathView))
+				{
+					Section.AddMenuEntry(
+						"DisallowPrivateContentEditing",
+						LOCTEXT("DisallowPrivateContentEditing", "Disallow Private Content Editing"),
+						LOCTEXT("DisallowPrivateContentEditingTooltip", "Hides Private Content and prevents editing the Public/Private state of content in this folder"),
+						FSlateIcon(FAppStyle::GetAppStyleSetName(), "ContentBrowser.PrivateContentEdit"),
+						FUIAction(FExecuteAction::CreateSP(this, &FPathContextMenu::ExecutePrivateContentEdit))
+					);
+				}
+				else
+				{
+					Section.AddMenuEntry(
+						"AllowPrivateContentEditing",
+						LOCTEXT("AllowPrivateContentEditing", "Allow Private Content Editing"),
+						LOCTEXT("AllowPrivateContentEditingTooltip", "Reveals Private Content and allows editing the Public/Private state of content in this folder"),
+						FSlateIcon(FAppStyle::GetAppStyleSetName(), "ContentBrowser.PrivateContentEdit"),
+						FUIAction(FExecuteAction::CreateSP(this, &FPathContextMenu::ExecutePrivateContentEdit))
+					);
+				}
+			}
 		}
 
 		if(Context->bCanBeModified)
@@ -364,6 +403,17 @@ void FPathContextMenu::ExecuteFavorite()
 	OnFolderFavoriteToggled.ExecuteIfBound(PathsToUpdate);
 }
 
+void FPathContextMenu::ExecutePrivateContentEdit()
+{
+	TArray<FString> PathsToUpdate;
+	for (const FContentBrowserItem& SelectedItem : SelectedFolders)
+	{
+		PathsToUpdate.Add(SelectedItem.GetVirtualPath().ToString());
+	}
+
+	OnPrivateContentEditToggled.ExecuteIfBound(PathsToUpdate);
+}
+
 void FPathContextMenu::NewColorComplete(const TSharedRef<SWindow>& Window)
 {
 	// Save the colors back in the config (ptr should have already updated by the widget)
diff --git a/Engine/Source/Editor/ContentBrowser/Private/PathContextMenu.h b/Engine/Source/Editor/ContentBrowser/Private/PathContextMenu.h
--- a/Engine/Source/Editor/ContentBrowser/Private/PathContextMenu.h
+++ b/Engine/Source/Editor/ContentBrowser/Private/PathContextMenu.h
@@ -32,6 +32,10 @@ public:
 	DECLARE_DELEGATE_OneParam(FOnFolderFavoriteToggled, const TArray<FString>& /*FoldersToToggle*/)
 	void SetOnFolderFavoriteToggled(const FOnFolderFavoriteToggled& InOnFolderFavoriteToggled);
 
+	/* Delegate for when the context menu has succesfully toggled the private content edit mode of a folder*/
+	DECLARE_DELEGATE_OneParam(FOnPrivateContentEditToggled, const TArray<FString>& /*FolderVirtualPaths*/)
+	void SetOnPrivateContentEditToggled(const FOnPrivateContentEditToggled& InOnPrivateContentEditableToggled);
+
 	/** Gets the currently selected folders */
 	const TArray<FContentBrowserItem>& GetSelectedFolders() const;
 
@@ -74,6 +78,9 @@ public:
 	/** Handler for favoriting */
 	void ExecuteFavorite();
 
+	/* Handler for enabling private content editing*/
+	void ExecutePrivateContentEdit();
+
 	/** Handler for when "Save" is selected */
 	void ExecuteSaveFolder();
 
@@ -107,4 +114,5 @@ private:
 	FOnRenameFolderRequested OnRenameFolderRequested;
 	FOnFolderDeleted OnFolderDeleted;
 	FOnFolderFavoriteToggled OnFolderFavoriteToggled;
+	FOnPrivateContentEditToggled OnPrivateContentEditToggled;
 };
diff --git a/Engine/Source/Editor/ContentBrowser/Private/SContentBrowser.cpp b/Engine/Source/Editor/ContentBrowser/Private/SContentBrowser.cpp
--- a/Engine/Source/Editor/ContentBrowser/Private/SContentBrowser.cpp
+++ b/Engine/Source/Editor/ContentBrowser/Private/SContentBrowser.cpp
@@ -142,6 +142,7 @@ void SContentBrowser::Construct( const FArguments& InArgs, const FName& InInstan
 	PathContextMenu->SetOnRenameFolderRequested(FPathContextMenu::FOnRenameFolderRequested::CreateSP(this, &SContentBrowser::OnRenameRequested));
 	PathContextMenu->SetOnFolderDeleted(FPathContextMenu::FOnFolderDeleted::CreateSP(this, &SContentBrowser::OnOpenedFolderDeleted));
 	PathContextMenu->SetOnFolderFavoriteToggled(FPathContextMenu::FOnFolderFavoriteToggled::CreateSP(this, &SContentBrowser::ToggleFolderFavorite));
+	PathContextMenu->SetOnPrivateContentEditToggled(FPathContextMenu::FOnPrivateContentEditToggled::CreateSP(this, &SContentBrowser::TogglePrivateContentEdit));
 	FrontendFilters = MakeShareable(new FAssetFilterCollectionType());
 	TextFilter = MakeShareable( new FFrontendFilter_Text() );
 
@@ -184,6 +185,7 @@ void SContentBrowser::Construct( const FArguments& InArgs, const FName& InInstan
 		.OnItemsActivated(this, &SContentBrowser::OnItemsActivated)
 		.OnGetItemContextMenu(this, &SContentBrowser::GetItemContextMenu, EContentBrowserViewContext::AssetView)
 		.OnItemRenameCommitted(this, &SContentBrowser::OnItemRenameCommitted)
+		.OnShouldFilterItem(this, &SContentBrowser::HandlePrivateContentFilter)
 		.FrontendFilters(FrontendFilters)
 		.HighlightedText(this, &SContentBrowser::GetHighlightedText)
 		.ShowBottomToolbar(bShowBottomToolbar)
@@ -626,6 +628,25 @@ void SContentBrowser::ToggleFolderFavorite(const TArray<FString>& FolderPaths)
 	}
 }
 
+void SContentBrowser::TogglePrivateContentEdit(const TArray<FString>& FolderPaths)
+{
+	for (const FString& FolderPath : FolderPaths)
+	{
+		ensure(FContentBrowserSingleton::Get().IsFolderShowPrivateContentToggleable(FolderPath));
+
+		if (FContentBrowserSingleton::Get().IsShowingPrivateContent(FolderPath))
+		{
+			ContentBrowserUtils::RemoveShowPrivateContentFolder(FolderPath, TEXT("ContentBrowser"));
+		}
+		else
+		{
+			ContentBrowserUtils::AddShowPrivateContentFolder(FolderPath, TEXT("ContentBrowser"));
+		}
+	}
+
+	OnAssetViewRefreshRequested();
+}
+
 void SContentBrowser::HandleAssetViewSearchOptionsChanged()
 {
 	TextFilter->SetIncludeClassName(AssetViewPtr->IsIncludingClassNames());
@@ -3416,6 +3437,32 @@ void SContentBrowser::HandleItemDataUpdated(TArrayView<const FContentBrowserItem
 	}
 }
 
+bool SContentBrowser::HandlePrivateContentFilter(const FContentBrowserItem& AssetItem)
+{
+	static const auto PublicAssetUIEnabledCVar = IConsoleManager::Get().FindConsoleVariable(TEXT("ContentBrowser.PublicAsset.EnablePublicAssetFeature"));
+	bool bIsPublicAssetUIEnabled = false;
+
+	if (PublicAssetUIEnabledCVar)
+	{
+		bIsPublicAssetUIEnabled = PublicAssetUIEnabledCVar->GetBool();
+	}
+
+	if (!bIsPublicAssetUIEnabled)
+	{
+		return false;
+	}
+
+	FAssetData ItemAssetData;
+	if ((AssetItem.Legacy_TryGetAssetData(ItemAssetData) && (ItemAssetData.PackageFlags & PKG_NotExternallyReferenceable)))
+	{
+		const FNameBuilder AssetItemFolderPath(AssetItem.GetVirtualPath());
+
+		return !FContentBrowserSingleton::Get().IsShowingPrivateContent(AssetItemFolderPath);
+	}
+
+	return false;
+}
+
 FText SContentBrowser::GetSearchAssetsHintText() const
 {
 	if (PathViewPtr.IsValid())
diff --git a/Engine/Source/Editor/ContentBrowser/Private/SContentBrowser.h b/Engine/Source/Editor/ContentBrowser/Private/SContentBrowser.h
--- a/Engine/Source/Editor/ContentBrowser/Private/SContentBrowser.h
+++ b/Engine/Source/Editor/ContentBrowser/Private/SContentBrowser.h
@@ -467,6 +467,9 @@ private:
 	/** Handles content items being updated */
 	void HandleItemDataUpdated(TArrayView<const FContentBrowserItemDataUpdate> InUpdatedItems);
 
+	/* Handles hiding private content if folder is not showing private content*/
+	bool HandlePrivateContentFilter(const FContentBrowserItem& AssetItem);
+
 	/** Gets all suggestions for the asset search box */
 	void OnAssetSearchSuggestionFilter(const FText& SearchText, TArray<FAssetSearchBoxSuggestion>& PossibleSuggestions, FText& SuggestionHighlightText) const;
 
@@ -509,6 +512,9 @@ private:
 	/** Toggles the favorite status of an array of folders*/
 	void ToggleFolderFavorite(const TArray<FString>& FolderPaths);
 
+	/* Toggles the private show private content state of an array of folders*/
+	void TogglePrivateContentEdit(const TArray<FString>& FolderPaths);
+
 	/** Called when Asset View Options "Search" options change */
 	void HandleAssetViewSearchOptionsChanged();
 
diff --git a/Engine/Source/Editor/ContentBrowser/Public/IContentBrowserSingleton.h b/Engine/Source/Editor/ContentBrowser/Public/IContentBrowserSingleton.h
--- a/Engine/Source/Editor/ContentBrowser/Public/IContentBrowserSingleton.h
+++ b/Engine/Source/Editor/ContentBrowser/Public/IContentBrowserSingleton.h
@@ -684,4 +684,22 @@ public:
 
 	/** Returns InPath if can be written to, otherwise picks a default path that can be written to */
 	virtual FContentBrowserItemPath GetInitialPathToSaveAsset(const FContentBrowserItemPath& InPath) = 0;
+
+	/** Returns true if FolderPath is a private content edit folder */
+	virtual bool IsShowingPrivateContent(const FStringView VirtualFolderPath) = 0;
+
+	/** Returns true if FolderPath's private content edit mode is allowed to be toggled */
+	virtual bool IsFolderShowPrivateContentToggleable(const FStringView VirtualFolderPath) = 0;
+
+	/** Returns the Private Content Permission List */
+	virtual const TSharedPtr<FPathPermissionList>& GetShowPrivateContentPermissionList() = 0;
+
+	/** Declares the Private Content Permission List dirty */
+	virtual void SetPrivateContentPermissionListDirty() = 0;
+
+	/** Registers the delegate for custom handling of if a Folder allows private content edits */
+	virtual void RegisterIsFolderShowPrivateContentToggleableDelegate(FIsFolderShowPrivateContentToggleableDelegate InIsFolderShowPrivateContentToggleableDelegate) = 0;
+
+	/** Unregisters the delegate for custom handling of if a Folder allows private content edits */
+	virtual void UnregisterIsFolderShowPrivateContentToggleableDelegate() = 0;
 };
diff --git a/Engine/Source/Editor/EditorStyle/Private/SlateEditorStyle.cpp b/Engine/Source/Editor/EditorStyle/Private/SlateEditorStyle.cpp
--- a/Engine/Source/Editor/EditorStyle/Private/SlateEditorStyle.cpp
+++ b/Engine/Source/Editor/EditorStyle/Private/SlateEditorStyle.cpp
@@ -6592,6 +6592,7 @@ void FSlateEditorStyle::FStyle::SetupContentBrowserStyle()
 	{
 		// Tab and menu icon
 		Set( "ContentBrowser.TabIcon", new IMAGE_BRUSH( "Icons/icon_tab_ContentBrowser_16x", Icon16x16 ) );
+		Set( "ContentBrowser.PrivateContentEdit", new IMAGE_BRUSH("Icons/hiererchy_16x", Icon16x16));
 
 		// Sources View
 		Set( "ContentBrowser.SourceTitleFont", DEFAULT_FONT( "Regular", 12 ) );
diff --git a/Engine/Source/Editor/EditorStyle/Private/StarshipStyle.cpp b/Engine/Source/Editor/EditorStyle/Private/StarshipStyle.cpp
--- a/Engine/Source/Editor/EditorStyle/Private/StarshipStyle.cpp
+++ b/Engine/Source/Editor/EditorStyle/Private/StarshipStyle.cpp
@@ -5953,6 +5953,7 @@ void FStarshipEditorStyle::FStyle::SetupContentBrowserStyle()
 	{
 		// Tab and menu icon
 		Set("ContentBrowser.TabIcon", new IMAGE_BRUSH_SVG("Starship/Common/ContentBrowser", Icon16x16));
+		Set("ContentBrowser.PrivateContentEdit", new IMAGE_BRUSH("Icons/hiererchy_16x", Icon16x16));
 
 		// Sources View
 		Set("ContentBrowser.SourceTitleFont", DEFAULT_FONT( "Regular", 12 ) );
