Fix issue with multi select in Content Browser when toggling Private Public state of an asset
- Added in bulk set Public set Private when more than one asset is selected in the Content Browser
- Additionally added in caching mechanism for selected file Public asset editability

diff --git a/Engine/Source/Editor/ContentBrowser/Private/AssetContextMenu.cpp b/Engine/Source/Editor/ContentBrowser/Private/AssetContextMenu.cpp
--- a/Engine/Source/Editor/ContentBrowser/Private/AssetContextMenu.cpp
+++ b/Engine/Source/Editor/ContentBrowser/Private/AssetContextMenu.cpp
@@ -12,6 +12,7 @@
 #include "ContentBrowserUtils.h"
 #include "SAssetView.h"
 #include "ContentBrowserModule.h"
+#include "ContentBrowserSingleton.h"
 #include "EditorStyleSet.h"
 #include "HAL/FileManager.h"
 
@@ -22,6 +23,8 @@
 #include "Toolkits/GlobalEditorCommonCommands.h"
 #include "Framework/Commands/GenericCommands.h"
 #include "ContentBrowserCommands.h"
+#include "Misc/PathViews.h"
+
 
 namespace ContentBrowserConsoleVariables
 {
@@ -459,18 +462,46 @@ bool FAssetContextMenu::AddReferenceMenuOptions(UToolMenu* Menu)
 
 		if (ContentBrowserConsoleVariables::ContentBrowser_EnablePublicAssetFeature)
 		{
-			Section.AddMenuEntry(
-				"PublicAsset",
-				LOCTEXT("PublicAssetToggle", "Public Asset"),
-				LOCTEXT("PublicAssetToggleTooltip", "Sets whether or not an asset is publicly available for reference by other plugins"),
-				FSlateIcon(FEditorStyle::GetStyleSetName(), "ContentBrowser.AssetActions.PublicAssetToggle"),
-				FUIAction(
-					FExecuteAction::CreateSP(this, &FAssetContextMenu::ExecutePublicAssetToggle),
-					FCanExecuteAction::CreateSP(this, &FAssetContextMenu::CanExecutePublicAssetToggle),
-					FGetActionCheckState::CreateSP(this, &FAssetContextMenu::GetPublicAssetCheckState)
-				),
-				EUserInterfaceActionType::ToggleButton
-			);
+			if (SelectedFiles.Num() == 1)
+			{
+				Section.AddMenuEntry(
+					"PublicAsset",
+					LOCTEXT("PublicAssetToggle", "Public Asset"),
+					LOCTEXT("PublicAssetToggleTooltip", "Sets whether or not an asset is publicly available for reference by other plugins"),
+					FSlateIcon(FEditorStyle::GetStyleSetName(), "ContentBrowser.AssetActions.PublicAssetToggle"),
+					FUIAction(
+						FExecuteAction::CreateSP(this, &FAssetContextMenu::ExecutePublicAssetToggle),
+						FCanExecuteAction::CreateSP(this, &FAssetContextMenu::CanExecutePublicAssetToggle),
+						FGetActionCheckState::CreateSP(this, &FAssetContextMenu::GetPublicAssetCheckState)
+					),
+					EUserInterfaceActionType::ToggleButton
+				);
+			}
+			else if (SelectedFiles.Num() > 1)
+			{
+				Section.AddMenuEntry(
+					"MarkSelectedAsPublic",
+					LOCTEXT("MarkSelectedAsPublic", "Mark Selected As Public"),
+					LOCTEXT("MarkSelectedAsPublicTooltip", "Sets all selected assets to be publicly available for reference by other plugins"),
+					FSlateIcon(FEditorStyle::GetStyleSetName(), "ContentBrowser.AssetActions.PublicAssetToggle"),
+					FUIAction(
+						FExecuteAction::CreateSP(this, &FAssetContextMenu::ExecuteBulkSetPublicAsset),
+						FCanExecuteAction::CreateSP(this, &FAssetContextMenu::CanExecuteBulkSetPublicAsset)
+					)
+				);
+
+				Section.AddMenuEntry(
+					"MarkSelectedAsPrivate",
+					LOCTEXT("MarkSelectedAsPrivate", "Mark Selected As Private"),
+					LOCTEXT("MarkSelectedAsPrivateTooltip", "Sets all selected assets to be private and unavailable for reference by other plugins"),
+					FSlateIcon(FEditorStyle::GetStyleSetName(), "ContentBrowser.AssetActions.PublicAssetToggle"),
+					FUIAction(
+						FExecuteAction::CreateSP(this, &FAssetContextMenu::ExecuteBulkUnsetPublicAsset),
+						FCanExecuteAction::CreateSP(this, &FAssetContextMenu::CanExecuteBulkSetPublicAsset)
+					)
+				);
+			}
+
 		}
 
 		Section.AddMenuEntry(
@@ -871,14 +902,73 @@ void FAssetContextMenu::ExecutePublicAssetToggle()
 
 			ItemAssetPackage->SetIsExternallyReferenceable(!ItemAssetPackage->IsExternallyReferenceable());
 
-			ItemAssetPackage->SetDirtyFlag(true);
+			ItemAssetPackage->Modify();
+
+			OnAssetViewRefreshRequested.ExecuteIfBound();
 		}
 	}
 }
 
 bool FAssetContextMenu::CanExecutePublicAssetToggle()
 {
-	return SelectedFiles.Num() == 1 && SelectedFiles[0].CanEdit();
+	return bCanExecutePublicAssetToggle;
+}
+
+void FAssetContextMenu::ExecuteBulkSetPublicAsset()
+{
+	for (const FContentBrowserItem& SelectedItem : SelectedFiles)
+	{
+		FAssetData ItemAssetData;
+		if (SelectedItem.Legacy_TryGetAssetData(ItemAssetData))
+		{
+			UPackage* ItemAssetPackage = ItemAssetData.GetPackage();
+
+			if (!ItemAssetPackage)
+			{
+				continue;
+			}
+
+			if (!ItemAssetPackage->IsExternallyReferenceable())
+			{
+				ItemAssetPackage->SetIsExternallyReferenceable(true);
+
+				ItemAssetPackage->Modify();
+			}
+		}
+	}
+
+	OnAssetViewRefreshRequested.ExecuteIfBound();
+}
+
+void FAssetContextMenu::ExecuteBulkUnsetPublicAsset()
+{
+	for (const FContentBrowserItem& SelectedItem : SelectedFiles)
+	{
+		FAssetData ItemAssetData;
+		if (SelectedItem.Legacy_TryGetAssetData(ItemAssetData))
+		{
+			UPackage* ItemAssetPackage = ItemAssetData.GetPackage();
+
+			if (!ItemAssetPackage)
+			{
+				continue;
+			}
+
+			if (ItemAssetPackage->IsExternallyReferenceable())
+			{
+				ItemAssetPackage->SetIsExternallyReferenceable(false);
+
+				ItemAssetPackage->Modify();
+			}
+		}
+	}
+
+	OnAssetViewRefreshRequested.ExecuteIfBound();
+}
+
+bool FAssetContextMenu::CanExecuteBulkSetPublicAsset()
+{
+	return bCanExecuteBulkSetPublicAsset;
 }
 
 ECheckBoxState FAssetContextMenu::GetPublicAssetCheckState()
@@ -1000,15 +1090,26 @@ bool FAssetContextMenu::CanExecuteSaveAsset() const
 void FAssetContextMenu::CacheCanExecuteVars()
 {
 	bCanExecuteFindInExplorer = false;
+	bCanExecutePublicAssetToggle = false;
+	bCanExecuteBulkSetPublicAsset = true;
 
 	// Selection must contain at least one file that has exists on disk
 	for (const FContentBrowserItem& SelectedItem : SelectedFiles)
 	{
 		FString ItemFilename;
-		if (SelectedItem.GetItemPhysicalPath(ItemFilename) && FPaths::FileExists(ItemFilename))
+		if (!bCanExecuteFindInExplorer && SelectedItem.GetItemPhysicalPath(ItemFilename) && FPaths::FileExists(ItemFilename))
 		{
 			bCanExecuteFindInExplorer = true;
-			break;
+		}
+
+		if(SelectedItem.CanEdit() &&
+			FContentBrowserSingleton::Get().IsShowingPrivateContent(FPathViews::GetPath(FNameBuilder(SelectedItem.GetVirtualPath()))))
+		{
+			bCanExecutePublicAssetToggle = true;
+		}
+		else
+		{
+			bCanExecuteBulkSetPublicAsset = false;
 		}
 	}
 }
diff --git a/Engine/Source/Editor/ContentBrowser/Private/AssetContextMenu.h b/Engine/Source/Editor/ContentBrowser/Private/AssetContextMenu.h
--- a/Engine/Source/Editor/ContentBrowser/Private/AssetContextMenu.h
+++ b/Engine/Source/Editor/ContentBrowser/Private/AssetContextMenu.h
@@ -121,6 +121,15 @@ private:
 	/** Handler to check to see if a Public Asset toggle is allowed */
 	bool CanExecutePublicAssetToggle();
 
+	/** Handler for setting all selected assets to Public */
+	void ExecuteBulkSetPublicAsset();
+
+	/** Handler for setting all selected assets to Private */
+	void ExecuteBulkUnsetPublicAsset();
+
+	/** Handler to check if all selected assets can have their Public state changed */
+	bool CanExecuteBulkSetPublicAsset();
+
 	/** Handler for getting the Public Asset check state */
 	ECheckBoxState GetPublicAssetCheckState();
 
@@ -165,4 +174,6 @@ private:
 
 	/** Cached CanExecute vars */
 	bool bCanExecuteFindInExplorer = false;
+	bool bCanExecutePublicAssetToggle = false;
+	bool bCanExecuteBulkSetPublicAsset = false;
 };
