Content Browser: Unlock a Browser on closing its tab, and add a toast for when all browsers are locked and the user tries to summon one

diff --git a/Engine/Source/Editor/ContentBrowser/Private/ContentBrowserSingleton.cpp b/Engine/Source/Editor/ContentBrowser/Private/ContentBrowserSingleton.cpp
--- a/Engine/Source/Editor/ContentBrowser/Private/ContentBrowserSingleton.cpp
+++ b/Engine/Source/Editor/ContentBrowser/Private/ContentBrowserSingleton.cpp
@@ -35,6 +35,7 @@
 #include "Misc/NamePermissionList.h"
 #include "StatusBarSubsystem.h"
 #include "ToolMenus.h"
+#include "Widgets/Notifications/SNotificationList.h"
 
 #define LOCTEXT_NAMESPACE "ContentBrowser"
 
@@ -396,7 +397,12 @@ TSharedPtr<SContentBrowser> FContentBrowserSingleton::FindContentBrowserToSync(b
 
 	if ( !ContentBrowserToSync.IsValid() )
 	{
-		UE_LOG( LogContentBrowser, Log, TEXT( "Unable to sync content browser, all browsers appear to be locked" ) );
+		FText NotificationText = FText::FromString(TEXT("Unable to browse to the requested asset. All Content Browsers are locked."));
+		FNotificationInfo Notification(NotificationText);
+		Notification.ExpireDuration = 3.0f;
+		FSlateNotificationManager::Get().AddNotification(Notification);
+
+		UE_LOG(LogNet, Log, TEXT("%s"), *NotificationText.ToString());
 	}
 
 	return ContentBrowserToSync;
diff --git a/Engine/Source/Editor/ContentBrowser/Private/SContentBrowser.cpp b/Engine/Source/Editor/ContentBrowser/Private/SContentBrowser.cpp
--- a/Engine/Source/Editor/ContentBrowser/Private/SContentBrowser.cpp
+++ b/Engine/Source/Editor/ContentBrowser/Private/SContentBrowser.cpp
@@ -1460,7 +1460,6 @@ void SContentBrowser::SaveSettings() const
 	const FString& SettingsString = InstanceName.ToString();
 
 	GConfig->SetBool(*SettingsIniSection, *(SettingsString + TEXT(".SourcesExpanded")), bSourcesViewExpanded, GEditorPerProjectIni);
-	GConfig->SetBool(*SettingsIniSection, *(SettingsString + TEXT(".Locked")), bIsLocked, GEditorPerProjectIni);
 
 	GConfig->SetBool(*SettingsIniSection, *(SettingsString + TEXT(".FavoritesAreaExpanded")), FavoritesArea->IsExpanded(), GEditorPerProjectIni);
 	GConfig->SetBool(*SettingsIniSection, *(SettingsString + TEXT(".PathAreaExpanded")), PathArea->IsExpanded(), GEditorPerProjectIni);
@@ -1594,7 +1593,6 @@ void SContentBrowser::LoadSettings(const FName& InInstanceName)
 	// Now that we have determined the appropriate settings string, actually load the settings
 	bSourcesViewExpanded = true;
 	GConfig->GetBool(*SettingsIniSection, *(SettingsString + TEXT(".SourcesExpanded")), bSourcesViewExpanded, GEditorPerProjectIni);
-	GConfig->GetBool(*SettingsIniSection, *(SettingsString + TEXT(".Locked")), bIsLocked, GEditorPerProjectIni);
 
 	bool bFavoritesAreaExpanded = false;
 	GConfig->GetBool(*SettingsIniSection, *(SettingsString + TEXT(".FavoritesAreaExpanded")), bFavoritesAreaExpanded, GEditorPerProjectIni);
