ContentBrowser: Fix not being able to hide favorites

diff --git a/Engine/Source/Editor/ContentBrowser/Private/SContentBrowser.cpp b/Engine/Source/Editor/ContentBrowser/Private/SContentBrowser.cpp
--- a/Engine/Source/Editor/ContentBrowser/Private/SContentBrowser.cpp
+++ b/Engine/Source/Editor/ContentBrowser/Private/SContentBrowser.cpp
@@ -819,6 +819,7 @@ TSharedRef<SWidget> SContentBrowser::CreateFavoritesView(const FContentBrowserCo
 		.BorderImage(FAppStyle::Get().GetBrush("Brushes.Header"))
 		.BodyBorderImage(FAppStyle::Get().GetBrush("Brushes.Recessed"))
 		.HeaderPadding(FMargin(4.0f, 2.0f))
+		.Visibility(this, &SContentBrowser::GetFavoriteFolderVisibility)
 		.Padding(0)
 		.AllowAnimatedTransition(false)
 		.OnAreaExpansionChanged_Lambda([this](bool bIsExpanded) { if (!bIsExpanded) FavoritesSearchToggleButton->SetExpanded(false); })
