Content Browser: Deactivate Add and Import buttons in content browser when path selected does not support those features (virtual paths such as /All)

diff --git a/Engine/Plugins/Editor/ContentBrowser/ContentBrowserAssetDataSource/Source/ContentBrowserAssetDataSource/Private/ContentBrowserAssetDataSource.cpp b/Engine/Plugins/Editor/ContentBrowser/ContentBrowserAssetDataSource/Source/ContentBrowserAssetDataSource/Private/ContentBrowserAssetDataSource.cpp
--- a/Engine/Plugins/Editor/ContentBrowser/ContentBrowserAssetDataSource/Source/ContentBrowserAssetDataSource/Private/ContentBrowserAssetDataSource.cpp
+++ b/Engine/Plugins/Editor/ContentBrowser/ContentBrowserAssetDataSource/Source/ContentBrowserAssetDataSource/Private/ContentBrowserAssetDataSource.cpp
@@ -2130,7 +2130,7 @@ FReply UContentBrowserAssetDataSource::OnImportClicked(const UContentBrowserTool
 
 bool UContentBrowserAssetDataSource::IsImportEnabled(const UContentBrowserToolbarMenuContext* ContextObject) const
 {
-	return ContextObject->GetCurrentPath() != NAME_None;
+	return ContextObject->CanWriteToCurrentPath();
 }
 
 FContentBrowserItemData UContentBrowserAssetDataSource::OnFinalizeCreateFolder(const FContentBrowserItemData& InItemData, const FString& InProposedName, FText* OutErrorMsg)
diff --git a/Engine/Source/Editor/ContentBrowser/Private/ContentBrowserMenuContexts.cpp b/Engine/Source/Editor/ContentBrowser/Private/ContentBrowserMenuContexts.cpp
--- a/Engine/Source/Editor/ContentBrowser/Private/ContentBrowserMenuContexts.cpp
+++ b/Engine/Source/Editor/ContentBrowser/Private/ContentBrowserMenuContexts.cpp
@@ -12,3 +12,13 @@ FName UContentBrowserToolbarMenuContext::GetCurrentPath() const
 
 	return NAME_None;
 }
+
+bool UContentBrowserToolbarMenuContext::CanWriteToCurrentPath() const
+{
+	if (TSharedPtr<SContentBrowser> Browser = ContentBrowser.Pin())
+	{
+		return Browser->CanWriteToCurrentPath();
+	}
+
+	return false;
+}
diff --git a/Engine/Source/Editor/ContentBrowser/Private/SContentBrowser.cpp b/Engine/Source/Editor/ContentBrowser/Private/SContentBrowser.cpp
--- a/Engine/Source/Editor/ContentBrowser/Private/SContentBrowser.cpp
+++ b/Engine/Source/Editor/ContentBrowser/Private/SContentBrowser.cpp
@@ -2484,10 +2484,59 @@ void SContentBrowser::PopulateAddNewContextMenu(class UToolMenu* Menu)
 		);
 }
 
+bool SContentBrowser::CanWriteToCurrentPath() const
+{
+	if (AssetViewPtr.IsValid())
+	{
+		const FSourcesData& SourcesData = AssetViewPtr->GetSourcesData();
+		if (SourcesData.VirtualPaths.Num() == 1)
+		{
+			if (!CachedCanWriteToCurrentPath.IsSet() || CachedCanWriteToCurrentPath.GetValue() != SourcesData.VirtualPaths[0])
+			{
+				CachedCanWriteToCurrentPath = SourcesData.VirtualPaths[0];
+				bCachedCanWriteToCurrentPath = CanWriteToPath(FContentBrowserItemPath(SourcesData.VirtualPaths[0], EContentBrowserPathType::Virtual));
+			}
+
+			return bCachedCanWriteToCurrentPath;
+		}
+		else
+		{
+			CachedCanWriteToCurrentPath.Reset();
+			bCachedCanWriteToCurrentPath = false;
+		}
+	}
+
+	return false;
+}
+
+bool SContentBrowser::CanWriteToPath(const FContentBrowserItemPath InPath) const
+{
+	// Reject if only virtual
+	if (!InPath.HasInternalPath())
+	{
+		return false;
+	}
+
+	// Reject if path not inside a mount point
+	if (!FPackageName::IsValidPath(InPath.GetInternalPathString()))
+	{
+		return false;
+	}
+
+	// Reject if folder writes blocked to path
+	FAssetToolsModule& AssetToolsModule = FModuleManager::LoadModuleChecked<FAssetToolsModule>("AssetTools");
+	const TSharedRef<FPathPermissionList>& WritableFolderFilter = AssetToolsModule.Get().GetWritableFolderPermissionList();
+	if (!WritableFolderFilter->PassesStartsWithFilter(InPath.GetInternalPathName()))
+	{
+		return false;
+	}
+
+	return true;
+}
+
 bool SContentBrowser::IsAddNewEnabled() const
 {
-	const FSourcesData& SourcesData = AssetViewPtr->GetSourcesData();
-	return SourcesData.VirtualPaths.Num() == 1;
+	return CanWriteToCurrentPath();
 }
 
 FText SContentBrowser::GetAddNewToolTipText() const
@@ -2497,6 +2546,12 @@ FText SContentBrowser::GetAddNewToolTipText() const
 	if ( SourcesData.VirtualPaths.Num() == 1 )
 	{
 		const FString CurrentPath = SourcesData.VirtualPaths[0].ToString();
+
+		if (!CanWriteToCurrentPath())
+		{
+			return FText::Format(LOCTEXT("AddNewToolTip_CannotWrite", "Cannot write to path {0}..."), FText::FromString(CurrentPath));
+		}
+
 		return FText::Format( LOCTEXT("AddNewToolTip_AddNewContent", "Create a new content in {0}..."), FText::FromString(CurrentPath) );
 	}
 	else if ( SourcesData.VirtualPaths.Num() > 1 )
@@ -3190,6 +3245,8 @@ void SContentBrowser::UpdatePath()
 		ActiveSourcesWidgetIndex = NewSourcesWidgetIndex;
 		SourcesWidgetSwitcher->SetActiveWidgetIndex(ActiveSourcesWidgetIndex);
 	}
+
+	CachedCanWriteToCurrentPath.Reset();
 }
 
 void SContentBrowser::OnFilterChanged()
diff --git a/Engine/Source/Editor/ContentBrowser/Private/SContentBrowser.h b/Engine/Source/Editor/ContentBrowser/Private/SContentBrowser.h
--- a/Engine/Source/Editor/ContentBrowser/Private/SContentBrowser.h
+++ b/Engine/Source/Editor/ContentBrowser/Private/SContentBrowser.h
@@ -193,6 +193,12 @@ public:
 	virtual FReply OnPreviewMouseButtonDown( const FGeometry& MyGeometry, const FPointerEvent& MouseEvent ) override;
 	virtual FReply OnMouseButtonDoubleClick( const FGeometry& InMyGeometry, const FPointerEvent& InMouseEvent ) override;
 
+	/** Returns true if current path can be written to */
+	bool CanWriteToCurrentPath() const;
+
+	/** Returns true if path can be written to */
+	bool CanWriteToPath(const FContentBrowserItemPath InPath) const;
+
 private:
 
 	/** Called prior to syncing the selection in this Content Browser */
@@ -633,6 +639,12 @@ private:
 	/** True if source should not be changed from an outside source */
 	bool bIsLocked;
 
+	/** Cached result of CanWriteToPath to avoid recalculating it every frame */
+	mutable bool bCachedCanWriteToCurrentPath = false;
+
+	/** Path that was last used to determine bCachedCanWriteToCurrentPath */
+	mutable TOptional<FName> CachedCanWriteToCurrentPath;
+
 	/** The list of FrontendFilters currently applied to the asset view */
 	TSharedPtr<FAssetFilterCollectionType> FrontendFilters;
 
diff --git a/Engine/Source/Editor/ContentBrowser/Public/ContentBrowserMenuContexts.h b/Engine/Source/Editor/ContentBrowser/Public/ContentBrowserMenuContexts.h
--- a/Engine/Source/Editor/ContentBrowser/Public/ContentBrowserMenuContexts.h
+++ b/Engine/Source/Editor/ContentBrowser/Public/ContentBrowserMenuContexts.h
@@ -120,7 +120,8 @@ class CONTENTBROWSER_API UContentBrowserToolbarMenuContext : public UObject
 
 public:
 	FName GetCurrentPath() const;
-public:
+
+	bool CanWriteToCurrentPath() const;
 
 	TWeakPtr<SContentBrowser> ContentBrowser;
 };
