Fixed Python file being filtered out by the Content Browser.
  - The OnShouldFilterAsset is not a legacy delegate binding, but the one added when the ContentBrowser was decoupled from the AssetRegistry. So it should not be used to set bHasLegacyDelegateBindings.

diff --git a/Engine/Source/Editor/ContentBrowser/Private/SAssetView.cpp b/Engine/Source/Editor/ContentBrowser/Private/SAssetView.cpp
--- a/Engine/Source/Editor/ContentBrowser/Private/SAssetView.cpp
+++ b/Engine/Source/Editor/ContentBrowser/Private/SAssetView.cpp
@@ -1782,8 +1782,7 @@ FContentBrowserDataFilter SAssetView::CreateBackendDataFilter() const
 		|| OnGetCustomAssetToolTip.IsBound()
 		|| OnVisualizeAssetToolTip.IsBound()
 		|| OnAssetToolTipClosing.IsBound()
-		|| OnShouldFilterAsset.IsBound()
-		|| OnShouldFilterItem.IsBound();
+		|| OnShouldFilterAsset.IsBound();
 
 	FContentBrowserDataFilter DataFilter;
 	DataFilter.bRecursivePaths = bRecurse || !bUsingFolders || bHasCollections;
