Content Browser: Filter checkboxes now show indeterminate state when some sub-menu items are disabled

diff --git a/Engine/Source/Editor/ContentBrowser/Private/SFilterList.cpp b/Engine/Source/Editor/ContentBrowser/Private/SFilterList.cpp
--- a/Engine/Source/Editor/ContentBrowser/Private/SFilterList.cpp
+++ b/Engine/Source/Editor/ContentBrowser/Private/SFilterList.cpp
@@ -1424,7 +1424,7 @@ void SFilterList::PopulateAddFilterMenu(UToolMenu* Menu)
 				FUIAction(
 				FExecuteAction::CreateSP( this, &SFilterList::FilterByTypeCategoryClicked, MenuExpansion ),
 				FCanExecuteAction(),
-				FIsActionChecked::CreateSP(this, &SFilterList::IsAssetTypeCategoryInUse, MenuExpansion ) ),
+				FGetActionCheckState::CreateSP(this, &SFilterList::IsAssetTypeCategoryChecked, MenuExpansion ) ),
 				EUserInterfaceActionType::ToggleButton
 				);
 		}
@@ -1458,7 +1458,7 @@ void SFilterList::PopulateAddFilterMenu(UToolMenu* Menu)
 					FUIAction(
 					FExecuteAction::CreateSP(this, &SFilterList::FilterByTypeCategoryClicked, CategoryMenuPair.Key),
 					FCanExecuteAction(),
-					FIsActionChecked::CreateSP(this, &SFilterList::IsAssetTypeCategoryInUse, CategoryMenuPair.Key)),
+					FGetActionCheckState::CreateSP(this, &SFilterList::IsAssetTypeCategoryChecked, CategoryMenuPair.Key)),
 					EUserInterfaceActionType::ToggleButton
 					);
 			}
@@ -1593,22 +1593,54 @@ void SFilterList::FilterByTypeCategoryClicked(EAssetTypeCategories::Type Categor
 }
 
 bool SFilterList::IsAssetTypeCategoryInUse(EAssetTypeCategories::Type Category) const
+{
+	ECheckBoxState AssetTypeCategoryCheckState = IsAssetTypeCategoryChecked(Category);
+
+	if (AssetTypeCategoryCheckState == ECheckBoxState::Unchecked)
+	{
+		return false;
+	}
+
+	// An asset type category is in use if any of its type actions are in use (ECheckBoxState::Checked or ECheckBoxState::Undetermined)
+	return true;
+}
+
+ECheckBoxState SFilterList::IsAssetTypeCategoryChecked(EAssetTypeCategories::Type Category) const
 {
 	TArray<TWeakPtr<IAssetTypeActions>> TypeActionsList;
 	GetTypeActionsForCategory(Category, TypeActionsList);
 
+	bool bIsAnyActionInUse = false;
+	bool bIsAnyActionNotInUse = false;
+
 	for (const TWeakPtr<IAssetTypeActions>& AssetTypeActions : TypeActionsList)
 	{
 		if (AssetTypeActions.IsValid())
 		{
-			if (!IsAssetTypeActionsInUse(AssetTypeActions))
+			if (IsAssetTypeActionsInUse(AssetTypeActions))
 			{
-				return false;
+				bIsAnyActionInUse = true;
+			}
+			else
+			{
+				bIsAnyActionNotInUse = true;
+			}
+
+			if (bIsAnyActionInUse && bIsAnyActionNotInUse)
+			{
+				return ECheckBoxState::Undetermined;
 			}
 		}
 	}
 
-	return true;
+	if (bIsAnyActionInUse)
+	{
+		return ECheckBoxState::Checked;
+	}
+	else
+	{
+		return ECheckBoxState::Unchecked;
+	}
 }
 
 void SFilterList::GetTypeActionsForCategory(EAssetTypeCategories::Type Category, TArray< TWeakPtr<IAssetTypeActions> >& TypeActions) const
diff --git a/Engine/Source/Editor/ContentBrowser/Private/SFilterList.h b/Engine/Source/Editor/ContentBrowser/Private/SFilterList.h
--- a/Engine/Source/Editor/ContentBrowser/Private/SFilterList.h
+++ b/Engine/Source/Editor/ContentBrowser/Private/SFilterList.h
@@ -140,11 +140,14 @@ private:
 	void FilterByTypeCategoryClicked(EAssetTypeCategories::Type Category);
 
 	/** Handler to determine the "checked" state of an asset type category in the filter dropdown */
-	bool IsAssetTypeCategoryInUse(EAssetTypeCategories::Type Category) const;
+	ECheckBoxState IsAssetTypeCategoryChecked(EAssetTypeCategories::Type Category) const;
 
 	/** Returns all the asset type actions objects for the specified category */
 	void GetTypeActionsForCategory(EAssetTypeCategories::Type Category, TArray< TWeakPtr<IAssetTypeActions> >& TypeActions) const;
 
+	/** Function to check if a given asset type category is in use */
+	bool IsAssetTypeCategoryInUse(EAssetTypeCategories::Type Category) const;
+
 	void FrontendFilterClicked(TSharedRef<FFrontendFilter> FrontendFilter);
 	bool IsFrontendFilterInUse(TSharedRef<FFrontendFilter> FrontendFilter) const;
 	void FrontendFilterCategoryClicked(TSharedPtr<FFrontendFilterCategory> MenuCategory);
