Dirty folder paths showing Private Content when All/ folder is hidden

diff --git a/Engine/Source/Editor/ContentBrowser/Private/ContentBrowserModule.cpp b/Engine/Source/Editor/ContentBrowser/Private/ContentBrowserModule.cpp
--- a/Engine/Source/Editor/ContentBrowser/Private/ContentBrowserModule.cpp
+++ b/Engine/Source/Editor/ContentBrowser/Private/ContentBrowserModule.cpp
@@ -61,6 +61,8 @@ void FContentBrowserModule::ContentBrowserSettingChanged(FName InName)
 		ContentBrowserData->RefreshVirtualPathTreeIfNeeded();
 	}
 
+	ContentBrowserSingleton->SetPrivateContentPermissionListDirty();
+
 	// Resize the recently opened asset list
 	if (InName == NumberOfRecentAssetsName)
 	{
