Setup delegate to dirty cache of Private Content Folder List when it's edited

diff --git a/Engine/Source/Editor/ContentBrowser/Private/ContentBrowserSingleton.cpp b/Engine/Source/Editor/ContentBrowser/Private/ContentBrowserSingleton.cpp
--- a/Engine/Source/Editor/ContentBrowser/Private/ContentBrowserSingleton.cpp
+++ b/Engine/Source/Editor/ContentBrowser/Private/ContentBrowserSingleton.cpp
@@ -88,6 +88,7 @@ FContentBrowserSingleton::FContentBrowserSingleton()
 	PopulateConfigValues();
 
 	ShowPrivateContentState.InvariantPaths = MakeShared<FPathPermissionList>();
+	ShowPrivateContentState.InvariantPaths->OnFilterChanged().AddRaw(this, &FContentBrowserSingleton::SetPrivateContentPermissionListDirty);
 	ShowPrivateContentState.CachedVirtualPaths = MakeShared<FPathPermissionList>();
 }
 
