SContentBrowser: Fix missing icon causing warning in Lyra

diff --git a/Engine/Source/Editor/ContentBrowser/Private/SContentBrowser.cpp b/Engine/Source/Editor/ContentBrowser/Private/SContentBrowser.cpp
--- a/Engine/Source/Editor/ContentBrowser/Private/SContentBrowser.cpp
+++ b/Engine/Source/Editor/ContentBrowser/Private/SContentBrowser.cpp
@@ -1030,7 +1030,7 @@ TSharedRef<SWidget> SContentBrowser::CreateDrawerDockButton(const FContentBrowse
 				[
 					SNew(SImage)
 					.ColorAndOpacity(FSlateColor::UseForeground())
-					.Image(FAppStyle::Get().GetBrush("EditorViewport.SubMenu.Layouts"))
+					.Image(FAppStyle::Get().GetBrush("Icons.Layout"))
 				]
 				+ SHorizontalBox::Slot()
 				.VAlign(VAlign_Center)
