Fix content browser doing partial deselections on syncing selection set

diff --git a/Engine/Source/Editor/ContentBrowser/Private/SContentBrowser.cpp b/Engine/Source/Editor/ContentBrowser/Private/SContentBrowser.cpp
--- a/Engine/Source/Editor/ContentBrowser/Private/SContentBrowser.cpp
+++ b/Engine/Source/Editor/ContentBrowser/Private/SContentBrowser.cpp
@@ -3161,18 +3161,15 @@ void SContentBrowser::SyncGlobalSelectionSet()
 				}
 			}
 		}
-
-
-		// Now we'll build a list of objects that need to be removed from the global selection set
-		for( int32 CurEditorObjectIndex = 0; CurEditorObjectIndex < EditorSelection->Num(); ++CurEditorObjectIndex )
+		
+		// List of objects that need to be removed from the global selection set
+		TArray<UObject*> EditorSelectedObjects;
+		EditorSelection->GetSelectedObjects(EditorSelectedObjects);
+		for (UObject* CurEditorObject : EditorSelectedObjects)
 		{
-			UObject* CurEditorObject = EditorSelection->GetSelectedObject( CurEditorObjectIndex );
-			if( CurEditorObject != NULL ) 
+			if (CurEditorObject && !SelectedObjects.Contains(CurEditorObject))
 			{
-				if( !SelectedObjects.Contains( CurEditorObject ) )
-				{
-					EditorSelection->Deselect( CurEditorObject );
-				}
+				EditorSelection->Deselect(CurEditorObject);
 			}
 		}
 	}
