[VK] Update UX to toggle Public/Private asset state to use radio buttons

diff --git a/Engine/Source/Editor/ContentBrowser/Private/AssetContextMenu.cpp b/Engine/Source/Editor/ContentBrowser/Private/AssetContextMenu.cpp
--- a/Engine/Source/Editor/ContentBrowser/Private/AssetContextMenu.cpp
+++ b/Engine/Source/Editor/ContentBrowser/Private/AssetContextMenu.cpp
@@ -318,6 +318,11 @@ void FAssetContextMenu::AddMenuOptions(UToolMenu* InMenu)
 	// Add quick access to view commands
 	AddExploreMenuOptions(InMenu);
 
+	if (ContentBrowserConsoleVariables::ContentBrowser_EnablePublicAssetFeature)
+	{
+		AddPublicStateMenuOptions(InMenu);
+	}
+
 	// Add reference options
 	AddReferenceMenuOptions(InMenu);
 
@@ -453,56 +458,79 @@ void FAssetContextMenu::AddExploreMenuOptions(UToolMenu* Menu)
 	}
 }
 
-bool FAssetContextMenu::AddReferenceMenuOptions(UToolMenu* Menu)
+bool FAssetContextMenu::AddPublicStateMenuOptions(UToolMenu* Menu)
 {
-	UContentBrowserDataMenuContext_FileMenu* Context = Menu->FindContext<UContentBrowserDataMenuContext_FileMenu>();
+	if (!bCanExecutePublicAssetToggle && !bCanExecuteBulkSetPublicAsset)
+	{
+		return false;
+	}
 
+	UContentBrowserDataMenuContext_FileMenu* Context = Menu->FindContext<UContentBrowserDataMenuContext_FileMenu>();
 	{
-		FToolMenuSection& Section = Menu->AddSection("AssetContextReferences", LOCTEXT("ReferencesMenuHeading", "References"));
+		FToolMenuSection& Section = Menu->AddSection("AssetPublicState", LOCTEXT("PublicStateHandling", "Asset State"));
 
-		if (ContentBrowserConsoleVariables::ContentBrowser_EnablePublicAssetFeature)
+		if (SelectedFiles.Num() == 1 && bCanExecutePublicAssetToggle)
 		{
-			if (SelectedFiles.Num() == 1)
-			{
-				Section.AddMenuEntry(
-					"PublicAsset",
-					LOCTEXT("PublicAssetToggle", "Public Asset"),
-					LOCTEXT("PublicAssetToggleTooltip", "Sets whether or not an asset is publicly available for reference by other plugins"),
-					FSlateIcon(FEditorStyle::GetStyleSetName(), "ContentBrowser.AssetActions.PublicAssetToggle"),
-					FUIAction(
-						FExecuteAction::CreateSP(this, &FAssetContextMenu::ExecutePublicAssetToggle),
-						FCanExecuteAction::CreateSP(this, &FAssetContextMenu::CanExecutePublicAssetToggle),
-						FGetActionCheckState::CreateSP(this, &FAssetContextMenu::GetPublicAssetCheckState)
-					),
-					EUserInterfaceActionType::ToggleButton
-				);
-			}
-			else if (SelectedFiles.Num() > 1)
-			{
-				Section.AddMenuEntry(
-					"MarkSelectedAsPublic",
-					LOCTEXT("MarkSelectedAsPublic", "Mark Selected As Public"),
-					LOCTEXT("MarkSelectedAsPublicTooltip", "Sets all selected assets to be publicly available for reference by other plugins"),
-					FSlateIcon(FEditorStyle::GetStyleSetName(), "ContentBrowser.AssetActions.PublicAssetToggle"),
-					FUIAction(
-						FExecuteAction::CreateSP(this, &FAssetContextMenu::ExecuteBulkSetPublicAsset),
-						FCanExecuteAction::CreateSP(this, &FAssetContextMenu::CanExecuteBulkSetPublicAsset)
-					)
-				);
-
-				Section.AddMenuEntry(
-					"MarkSelectedAsPrivate",
-					LOCTEXT("MarkSelectedAsPrivate", "Mark Selected As Private"),
-					LOCTEXT("MarkSelectedAsPrivateTooltip", "Sets all selected assets to be private and unavailable for reference by other plugins"),
-					FSlateIcon(FEditorStyle::GetStyleSetName(), "ContentBrowser.AssetActions.PublicAssetToggle"),
-					FUIAction(
-						FExecuteAction::CreateSP(this, &FAssetContextMenu::ExecuteBulkUnsetPublicAsset),
-						FCanExecuteAction::CreateSP(this, &FAssetContextMenu::CanExecuteBulkSetPublicAsset)
-					)
-				);
-			}
+			Section.AddMenuEntry(
+				"PublicAsset",
+				LOCTEXT("PublicAssetToggle", "Public Asset"),
+				LOCTEXT("PublicAssetToggleTooltip", "Sets the asset to be referencable by other Plugins"),
+				FSlateIcon(),
+				FUIAction(
+					FExecuteAction::CreateSP(this, &FAssetContextMenu::ExecutePublicAssetToggle),
+					FCanExecuteAction::CreateSP(this, &FAssetContextMenu::CanExecutePublicAssetToggle),
+					FIsActionChecked::CreateSP(this, &FAssetContextMenu::IsSelectedAssetPublic)
+				),
+				EUserInterfaceActionType::RadioButton
+			);
+
+			Section.AddMenuEntry(
+				"PrivateAsset",
+				LOCTEXT("SetPrivateAsset", "Private Asset"),
+				LOCTEXT("SetAssetPrivateTooltip", "Sets the asset so it can't be referenced by other Plugins"),
+				FSlateIcon(),
+				FUIAction(
+					FExecuteAction::CreateSP(this, &FAssetContextMenu::ExecutePublicAssetToggle),
+					FCanExecuteAction::CreateSP(this, &FAssetContextMenu::CanExecutePublicAssetToggle),
+					FIsActionChecked::CreateSP(this, &FAssetContextMenu::IsSelectedAssetPrivate)
+				),
+				EUserInterfaceActionType::RadioButton
+			);
+		}
+		else if (SelectedFiles.Num() > 1 && bCanExecuteBulkSetPublicAsset)
+		{
+			Section.AddMenuEntry(
+				"MarkSelectedAsPublic",
+				LOCTEXT("MarkSelectedAsPublic", "Mark Selected As Public"),
+				LOCTEXT("MarkSelectedAsPublicTooltip", "Sets all selected assets to be publicly available for reference by other plugins"),
+				FSlateIcon(FEditorStyle::GetStyleSetName(), "ContentBrowser.AssetActions.PublicAssetToggle"),
+				FUIAction(
+					FExecuteAction::CreateSP(this, &FAssetContextMenu::ExecuteBulkSetPublicAsset),
+					FCanExecuteAction::CreateSP(this, &FAssetContextMenu::CanExecuteBulkSetPublicAsset)
+				)
+			);
 
+			Section.AddMenuEntry(
+				"MarkSelectedAsPrivate",
+				LOCTEXT("MarkSelectedAsPrivate", "Mark Selected As Private"),
+				LOCTEXT("MarkSelectedAsPrivateTooltip", "Sets all selected assets to be private and unavailable for reference by other plugins"),
+				FSlateIcon(FEditorStyle::GetStyleSetName(), "ContentBrowser.AssetActions.PublicAssetToggle"),
+				FUIAction(
+					FExecuteAction::CreateSP(this, &FAssetContextMenu::ExecuteBulkUnsetPublicAsset),
+					FCanExecuteAction::CreateSP(this, &FAssetContextMenu::CanExecuteBulkSetPublicAsset)
+				)
+			);
 		}
+	}
+
+	return true;
+}
+
+bool FAssetContextMenu::AddReferenceMenuOptions(UToolMenu* Menu)
+{
+	UContentBrowserDataMenuContext_FileMenu* Context = Menu->FindContext<UContentBrowserDataMenuContext_FileMenu>();
+	{
+		FToolMenuSection& Section = Menu->AddSection("AssetContextReferences", LOCTEXT("ReferencesMenuHeading", "References"));
 
 		Section.AddMenuEntry(
 			"CopyReference",
@@ -971,7 +999,7 @@ bool FAssetContextMenu::CanExecuteBulkSetPublicAsset()
 	return bCanExecuteBulkSetPublicAsset;
 }
 
-ECheckBoxState FAssetContextMenu::GetPublicAssetCheckState()
+bool FAssetContextMenu::IsSelectedAssetPublic()
 {
 	if (ensure(SelectedFiles.Num() == 1))
 	{
@@ -982,12 +1010,17 @@ ECheckBoxState FAssetContextMenu::GetPublicAssetCheckState()
 
 			if (ItemAssetPackage)
 			{
-				return ItemAssetPackage->IsExternallyReferenceable() ? ECheckBoxState::Checked : ECheckBoxState::Unchecked;
+				return ItemAssetPackage->IsExternallyReferenceable();
 			}
 		}
 	}
 
-	return ECheckBoxState::Undetermined;
+	return true;
+}
+
+bool FAssetContextMenu::IsSelectedAssetPrivate()
+{
+	return !IsSelectedAssetPublic();
 }
 
 void FAssetContextMenu::ExecuteCopyReference()
diff --git a/Engine/Source/Editor/ContentBrowser/Private/AssetContextMenu.h b/Engine/Source/Editor/ContentBrowser/Private/AssetContextMenu.h
--- a/Engine/Source/Editor/ContentBrowser/Private/AssetContextMenu.h
+++ b/Engine/Source/Editor/ContentBrowser/Private/AssetContextMenu.h
@@ -94,6 +94,8 @@ private:
 	/** Adds asset reference menu options to a menu builder. Returns true if any options were added. */
 	bool AddReferenceMenuOptions(UToolMenu* Menu);
 
+	bool AddPublicStateMenuOptions(UToolMenu* Menu);
+
 	/** Adds menu options related to working with collections */
 	bool AddCollectionMenuOptions(UToolMenu* Menu);
 
@@ -130,8 +132,11 @@ private:
 	/** Handler to check if all selected assets can have their Public state changed */
 	bool CanExecuteBulkSetPublicAsset();
 
-	/** Handler for getting the Public Asset check state */
-	ECheckBoxState GetPublicAssetCheckState();
+	/** Handler for determining the selected asset's Public state */
+	bool IsSelectedAssetPublic();
+
+	/** Handler for determining the selected asset's Private state */
+	bool IsSelectedAssetPrivate();
 
 	/** Handler for CopyReference */
 	void ExecuteCopyReference();
