Refactor SFilterList in the Content Browser into more globally usable Filter Widgets

SBasicFilterBar in ToolWidgets is the most generic filter bar, usable in non editor buillds as well
SAssetFilterBar in EditorWidgets is a complex filter bar with support for the built in asset filters from the content browser
SFilterBar in EditorWidgets is the most convinient filter bar to use in editor if you want asset filters and custom filters

diff --git a/Engine/Source/Developer/ToolWidgets/Public/Filters/FilterBase.h b/Engine/Source/Developer/ToolWidgets/Public/Filters/FilterBase.h
new file mode 100644
--- /dev/null
+++ b/Engine/Source/Developer/ToolWidgets/Public/Filters/FilterBase.h
@@ -0,0 +1,90 @@
+// Copyright Epic Games, Inc. All Rights Reserved.
+
+#pragma once
+
+#include "Misc/IFilter.h"
+#include "Internationalization/Text.h"
+#include "Core/Public/Math/Color.h"
+#include "Core/Public/UObject/NameTypes.h"
+
+class FMenuBuilder;
+
+class FFilterCategory
+{
+public:
+	FFilterCategory(const FText& InTitle, const FText& InTooltip) : Title(InTitle), Tooltip(InTooltip) {}
+
+	/** The title of this category, used for the menu heading */
+	const FText Title;
+
+	/** The menu tooltip for this category */
+	const FText Tooltip;
+};
+
+template<typename FilterType>
+struct FFrontendFilterExternalActivationHelper;
+
+/* The base class for all Filters that can be used with FilterBar Widgets
+ * @see FGenericFilter for a subclass that allows easy creation of filters
+ */
+template<typename FilterType>
+class FFilterBase : public IFilter<FilterType>
+{
+public:
+
+	FFilterBase(TSharedPtr<FFilterCategory> InCategory) : FilterCategory(InCategory) {}
+
+	/** Returns the system name for this filter */
+	virtual FString GetName() const = 0;
+
+	/** Returns the human readable name for this filter */
+	virtual FText GetDisplayName() const = 0;
+
+	/** Returns the tooltip for this filter, shown in the filters menu */
+	virtual FText GetToolTipText() const = 0;
+
+	/** Returns the color this filter button will be when displayed as a button */
+	virtual FLinearColor GetColor() const = 0;
+
+	/** Returns the name of the icon to use in menu entries */
+	virtual FName GetIconName() const = 0;
+
+	/** If true, the filter will be active in the FilterBar when it is inactive in the UI (i.e the filter pill is grayed out)
+	 * @See: FFrontendFilter_ShowOtherDevelopers in Content Browser
+	 */
+	virtual bool IsInverseFilter() const = 0;
+
+	/** Notification that the filter became active or inactive */
+	virtual void ActiveStateChanged(bool bActive) = 0;
+
+	/** Called when the right-click context menu is being built for this filter */
+	virtual void ModifyContextMenu(FMenuBuilder& MenuBuilder) = 0;
+	
+	/** Can be overriden for custom FilterBar subclasses to save settings, currently not implemented in any gneeric Filter Bar */
+	virtual void SaveSettings(const FString& IniFilename, const FString& IniSection, const FString& SettingsString) const = 0;
+
+	/** Can be overriden for custom FilterBar subclasses to load settings, currently not implemented in any gneeric Filter Bar */
+	virtual void LoadSettings(const FString& IniFilename, const FString& IniSection, const FString& SettingsString) = 0;
+
+	/** Set this filter as active/inactive */
+	void SetActive(bool bInActive) { SetActiveEvent.Broadcast(bInActive); }
+
+	// IFilter implementation
+	DECLARE_DERIVED_EVENT( FFilterBase, IFilter<FilterType>::FChangedEvent, FChangedEvent );
+	virtual FChangedEvent& OnChanged() override { return ChangedEvent; }
+
+	TSharedPtr<FFilterCategory> GetCategory() { return FilterCategory; }
+
+protected:
+	void BroadcastChangedEvent() const { ChangedEvent.Broadcast(); }
+
+	FChangedEvent ChangedEvent;
+
+	/** This event is broadcast to when this filter is set Active */
+	DECLARE_EVENT_OneParam(FFilterBase, FSetActiveEvent, bool);
+	FSetActiveEvent SetActiveEvent;
+
+	TSharedPtr<FFilterCategory> FilterCategory;
+
+	friend struct FFrontendFilterExternalActivationHelper<FilterType>;
+};
\ No newline at end of file
diff --git a/Engine/Source/Developer/ToolWidgets/Public/Filters/GenericFilter.h b/Engine/Source/Developer/ToolWidgets/Public/Filters/GenericFilter.h
new file mode 100644
--- /dev/null
+++ b/Engine/Source/Developer/ToolWidgets/Public/Filters/GenericFilter.h
@@ -0,0 +1,128 @@
+// Copyright Epic Games, Inc. All Rights Reserved.
+
+#pragma once
+
+#include "Filters/FilterBase.h"
+#include "Internationalization/Text.h"
+#include "Core/Public/Math/Color.h"
+#include "Core/Public/UObject/NameTypes.h"
+
+
+/* A generic filter that can be used with the various FilterBar Widgets
+ * Takes in a Category, Name, Display Name and Delegate which specifies how to filter an item
+ * You can optionally specify a Tooltip, Icon Name or Color
+ */
+template<typename FilterType>
+class FGenericFilter : public FFilterBase<FilterType>
+{
+public:
+
+	DECLARE_DELEGATE_RetVal_OneParam(bool, FOnItemFiltered, FilterType);
+	
+	FGenericFilter(TSharedPtr<FFilterCategory> InCategory, const FString& InName, const FText &DisplayName, FOnItemFiltered InFilterDelegate) :
+		FFilterBase<FilterType>(InCategory), Name(InName), DisplayName(DisplayName), FilterDelegate(InFilterDelegate), Color(FLinearColor(1.0f, 1.0f, 1.0f, 1.0f))
+	{
+		
+	}
+
+	/** Returns the system name for this filter */
+	virtual FString GetName() const override
+	{
+		return Name;
+	}
+
+	/** Returns the human readable name for this filter */
+	virtual FText GetDisplayName() const override
+	{
+		return DisplayName;
+	}
+
+	void SetToolTipText(const FText& InToolTip)
+	{
+		ToolTip = InToolTip;
+	}
+
+	/** Returns the tooltip for this filter, shown in the filters menu */
+	virtual FText GetToolTipText() const override
+	{
+		return ToolTip;
+	}
+
+	void SetColor(const FLinearColor& InColor)
+	{
+		Color = InColor;
+	}
+
+	/** Returns the color this filter button will be when displayed as a button */
+	virtual FLinearColor GetColor() const override
+	{
+		return Color;
+	}
+
+	void SetIconName(const FName& InIconName)
+	{
+		IconName = InIconName;
+	}
+
+	/** Returns the name of the icon to use in menu entries */
+	virtual FName GetIconName() const override
+	{
+		return IconName;
+	}
+
+	/** Returns true if the filter should be in the list when disabled and not in the list when enabled */
+	virtual bool IsInverseFilter() const override
+	{
+		return false; // Generic Filters are meant to be simple
+	}
+
+	virtual bool PassesFilter( FilterType InItem ) const override
+	{
+		if(FilterDelegate.IsBound())
+		{
+			return FilterDelegate.Execute(InItem);
+		}
+
+		return false;
+	}
+
+	// Functionality that is not needed by generic filters or currently not implemented
+	
+	/** Notification that the filter became active or inactive */
+	virtual void ActiveStateChanged(bool bActive) override
+	{
+		
+	}
+
+	/** Called when the right-click context menu is being built for this filter */
+	virtual void ModifyContextMenu(FMenuBuilder& MenuBuilder) override
+	{
+		
+	}
+	
+	/** Can be overriden for custom FilterBar subclasses to save settings, currently not implemented in any generic Filter Bar */
+	virtual void SaveSettings(const FString& IniFilename, const FString& IniSection, const FString& SettingsString) const override
+	{
+		
+	}
+
+	/** Can be overriden for custom FilterBar subclasses to load settings, currently not implemented in any generic Filter Bar */
+	virtual void LoadSettings(const FString& IniFilename, const FString& IniSection, const FString& SettingsString) override
+	{
+		
+	}
+	
+protected:
+
+	// Required members
+	FString Name;
+	FText DisplayName;
+	FOnItemFiltered FilterDelegate;
+
+	// Optional members
+	FText ToolTip;
+	FLinearColor Color;
+	FName IconName;
+	
+};
+
diff --git a/Engine/Source/Developer/ToolWidgets/Public/Filters/SBasicFilterBar.h b/Engine/Source/Developer/ToolWidgets/Public/Filters/SBasicFilterBar.h
new file mode 100644
--- /dev/null
+++ b/Engine/Source/Developer/ToolWidgets/Public/Filters/SBasicFilterBar.h
@@ -0,0 +1,1146 @@
+//Copyright Epic Games, Inc. All Rights Reserved.
+
+#pragma once
+
+#include "Input/Reply.h"
+#include "Widgets/DeclarativeSyntaxSupport.h"
+#include "Widgets/SWidget.h"
+#include "Widgets/SCompoundWidget.h"
+#include "Filters/FilterBase.h"
+#include "Widgets/Input/SCheckBox.h"
+#include "Widgets/Input/SComboButton.h"
+#include "Textures/SlateIcon.h"
+#include "Widgets/Layout/SBorder.h"
+#include "ToolMenus.h"
+#include "Widgets/Images/SImage.h"
+#include "Widgets/SBoxPanel.h"
+#include "Misc/FilterCollection.h"
+#include "UObject/Object.h"
+
+#include "SBasicFilterBar.generated.h"
+
+#define LOCTEXT_NAMESPACE "FilterBar"
+
+/* Delegate used by SBasicFilterBar to populate the add filter menu */
+DECLARE_DELEGATE_OneParam(FOnPopulateAddFilterMenu, UToolMenu*)
+
+/** ToolMenuContext that is used to create the Add Filter Menu */
+UCLASS()
+class TOOLWIDGETS_API UFilterBarContext : public UObject
+{
+	GENERATED_BODY()
+
+public:
+	FOnPopulateAddFilterMenu PopulateFilterMenu;
+};
+
+/** Forward declaration for a helper struct used to activate external filters */
+template<typename FilterType>
+struct FFrontendFilterExternalActivationHelper;
+
+/**
+* A Basic Filter Bar widget, which can be used to filter items of type [FilterType] given a list of custom filters
+* @see SFilterBar in EditorWidgets if you want a Filter Bar that includes Asset Type Filters
+* NOTE: The filter functions create copies, so you want to use a reference or pointer as the template type when possible
+* Sample Usage:
+*		SAssignNew(MyFilterBar, SBasicFilterBar<FText&>)
+*		.OnFilterChanged() // A delegate for when the list of filters changes
+*		.CustomFilters() // An array of filters available to this FilterBar (@see FGenericFilter to create simple delegate based filters)
+*
+* Use the GetAllActiveFilters() function to get the FilterCollection of Active Filters in this FilterBar, that can be used to filter your items
+* Use MakeAddFilterButton() to make the button that summons the dropdown showing all the filters
+*/
+template<typename FilterType>
+class SBasicFilterBar : public SCompoundWidget
+{
+public:
+	/** Delegate for when filters have changed */
+	DECLARE_DELEGATE( FOnFilterChanged );
+	
+ 	SLATE_BEGIN_ARGS( SBasicFilterBar<FilterType> ){}
+
+ 		/** Delegate for when filters have changed */
+ 		SLATE_EVENT( FOnFilterChanged, OnFilterChanged )
+
+		/** Initial List of Custom Filters that will be added to the AddFilter Menu */
+		SLATE_ARGUMENT( TArray<TSharedRef<FFilterBase<FilterType>>>, CustomFilters)
+	
+ 	SLATE_END_ARGS()
+
+ 	void Construct( const FArguments& InArgs )
+ 	{
+		OnFilterChanged = InArgs._OnFilterChanged;
+
+ 		// A subclass could be using an external filter collection (SFilterList in ContentBrowser)
+ 		if(!ActiveFilters)
+ 		{
+ 			ActiveFilters = MakeShareable(new TFilterCollection<FilterType>());
+ 		}
+
+ 		// Add the custom filters
+ 		for (TSharedRef<FFilterBase<FilterType>> Filter : InArgs._CustomFilters)
+ 		{
+ 			AddFilter(Filter);
+ 		}
+ 	}
+
+	/** Add a custom filter to the FilterBar.
+	 * NOTE: This only adds it to the Add Filter Menu, and does not automatically add it to the filter bar or activate it
+	 */
+	void AddFilter(TSharedRef<FFilterBase<FilterType>> InFilter)
+ 	{
+		AllFrontendFilters.Add(InFilter);
+ 		
+ 		// Add the category for the filter
+ 		if (TSharedPtr<FFilterCategory> Category = InFilter->GetCategory())
+ 		{
+ 			AllFilterCategories.AddUnique(Category);
+ 		}
+ 			
+ 		// Bind external activation event
+ 		FFrontendFilterExternalActivationHelper<FilterType>::BindToFilter(SharedThis(this), InFilter);
+
+ 		// Auto add if it is an inverse filters
+ 		SetFrontendFilterActive(InFilter, false);
+ 	}
+
+	/** Makes the button that summons the Add Filter dropdown on click */
+	static TSharedRef<SWidget> MakeAddFilterButton(TSharedRef<SBasicFilterBar<FilterType>> InFilterBar)
+ 	{
+ 		return SNew(SComboButton)
+				.ComboButtonStyle(&FAppStyle::Get().GetWidgetStyle<FComboButtonStyle>("SimpleComboButtonWithIcon"))
+				.ForegroundColor(FSlateColor::UseStyle())
+				.ToolTipText(LOCTEXT("AddFilterToolTip", "Open the Add Filter Menu to add or manage filters."))
+				.OnGetMenuContent(InFilterBar, &SBasicFilterBar<FilterType>::MakeAddFilterMenu)
+				.ContentPadding(FMargin(1, 0))
+				.ButtonContent()
+				[
+					SNew(SImage)
+					.Image(FAppStyle::Get().GetBrush("Icons.Filter"))
+					.ColorAndOpacity(FSlateColor::UseForeground())
+				];
+ 	}
+
+private:
+
+	/** A class for check boxes in the filter list. If you double click a filter checkbox, you will enable it and disable all others */
+	class SFilterCheckBox : public SCheckBox
+	{
+	public:
+		
+		void SetOnFilterCtrlClicked(const FOnClicked& NewFilterCtrlClicked)
+		{
+			OnFilterCtrlClicked = NewFilterCtrlClicked;
+		}
+
+		void SetOnFilterAltClicked(const FOnClicked& NewFilteAltClicked)
+		{
+			OnFilterAltClicked = NewFilteAltClicked;
+		}
+
+		void SetOnFilterDoubleClicked( const FOnClicked& NewFilterDoubleClicked )
+		{
+			OnFilterDoubleClicked = NewFilterDoubleClicked;
+		}
+
+		void SetOnFilterMiddleButtonClicked( const FOnClicked& NewFilterMiddleButtonClicked )
+		{
+			OnFilterMiddleButtonClicked = NewFilterMiddleButtonClicked;
+		}
+
+		virtual FReply OnMouseButtonDoubleClick( const FGeometry& InMyGeometry, const FPointerEvent& InMouseEvent ) override
+		{
+			if ( InMouseEvent.GetEffectingButton() == EKeys::LeftMouseButton && OnFilterDoubleClicked.IsBound() )
+			{
+				return OnFilterDoubleClicked.Execute();
+			}
+			else
+			{
+				return SCheckBox::OnMouseButtonDoubleClick(InMyGeometry, InMouseEvent);
+			}
+		}
+
+		virtual FReply OnMouseButtonUp( const FGeometry& InMyGeometry, const FPointerEvent& InMouseEvent ) override
+		{
+			if (InMouseEvent.IsControlDown() && OnFilterCtrlClicked.IsBound())
+			{
+				return OnFilterCtrlClicked.Execute();
+			}
+			else if (InMouseEvent.IsAltDown() && OnFilterAltClicked.IsBound())
+			{
+				return OnFilterAltClicked.Execute();
+			}
+			else if( InMouseEvent.GetEffectingButton() == EKeys::MiddleMouseButton && OnFilterMiddleButtonClicked.IsBound() )
+			{
+				return OnFilterMiddleButtonClicked.Execute();
+			}
+			else
+			{
+				SCheckBox::OnMouseButtonUp(InMyGeometry, InMouseEvent);
+				return FReply::Handled().ReleaseMouseCapture();
+			}
+		}
+
+	private:
+		FOnClicked OnFilterCtrlClicked;
+		FOnClicked OnFilterAltClicked;
+		FOnClicked OnFilterDoubleClicked;
+		FOnClicked OnFilterMiddleButtonClicked;
+	};
+
+protected:
+
+	/**
+	 * A single filter in the filter list. Can be removed by clicking the remove button on it.
+	 */
+	class SFilter : public SCompoundWidget
+	{
+	public:
+		DECLARE_DELEGATE_OneParam( FOnRequestRemove, const TSharedRef<SFilter>& /*FilterToRemove*/ );
+		DECLARE_DELEGATE_OneParam( FOnRequestRemoveAllButThis, const TSharedRef<SFilter>& /*FilterToKeep*/ );
+		DECLARE_DELEGATE_OneParam( FOnRequestEnableOnly, const TSharedRef<SFilter>& /*FilterToEnable*/ );
+		DECLARE_DELEGATE( FOnRequestEnableAll );
+		DECLARE_DELEGATE( FOnRequestDisableAll );
+		DECLARE_DELEGATE( FOnRequestRemoveAll );
+
+		SLATE_BEGIN_ARGS( SFilter ){}
+
+			/** If this is an front end filter, this is the filter object */
+			SLATE_ARGUMENT( TSharedPtr<FFilterBase<FilterType>>, FrontendFilter )
+
+			/** Invoked when the filter toggled */
+			SLATE_EVENT( SBasicFilterBar<FilterType>::FOnFilterChanged, OnFilterChanged )
+
+			/** Invoked when a request to remove this filter originated from within this filter */
+			SLATE_EVENT( FOnRequestRemove, OnRequestRemove )
+
+			/** Invoked when a request to enable only this filter originated from within this filter */
+			SLATE_EVENT( FOnRequestEnableOnly, OnRequestEnableOnly )
+
+			/** Invoked when a request to enable all filters originated from within this filter */
+			SLATE_EVENT(FOnRequestEnableAll, OnRequestEnableAll)
+
+			/** Invoked when a request to disable all filters originated from within this filter */
+			SLATE_EVENT( FOnRequestDisableAll, OnRequestDisableAll )
+
+			/** Invoked when a request to remove all filters originated from within this filter */
+			SLATE_EVENT( FOnRequestRemoveAll, OnRequestRemoveAll )
+
+			/** Invoked when a request to remove all filters originated from within this filter */
+			SLATE_EVENT( FOnRequestRemoveAllButThis, OnRequestRemoveAllButThis )
+
+		SLATE_END_ARGS()
+
+		/** Constructs this widget with InArgs */
+		void Construct( const FArguments& InArgs )
+		{
+			bEnabled = false;
+			OnFilterChanged = InArgs._OnFilterChanged;
+			OnRequestRemove = InArgs._OnRequestRemove;
+			OnRequestEnableOnly = InArgs._OnRequestEnableOnly;
+			OnRequestEnableAll = InArgs._OnRequestEnableAll;
+			OnRequestDisableAll = InArgs._OnRequestDisableAll;
+			OnRequestRemoveAll = InArgs._OnRequestRemoveAll;
+			OnRequestRemoveAllButThis = InArgs._OnRequestRemoveAllButThis;
+			FrontendFilter = InArgs._FrontendFilter;
+
+			// Get the tooltip and color of the type represented by this filter
+			FilterColor = FLinearColor::White;
+			if ( FrontendFilter.IsValid() )
+			{
+				FilterColor = FrontendFilter->GetColor();
+				FilterToolTip = TAttribute<FText>::Create(TAttribute<FText>::FGetter::CreateSP(FrontendFilter.ToSharedRef(), &FFilterBase<FilterType>::GetToolTipText));
+			}
+
+			Construct_Internal();
+		}
+
+		/** Sets whether or not this filter is applied to the combined filter */
+		void SetEnabled(bool InEnabled, bool InExecuteOnFilterChanged = true)
+		{
+			if ( InEnabled != bEnabled)
+			{
+				bEnabled = InEnabled;
+				if (InExecuteOnFilterChanged)
+				{
+					OnFilterChanged.ExecuteIfBound();
+				}
+			}
+		}
+
+		/** Returns true if this filter contributes to the combined filter */
+		bool IsEnabled() const
+		{
+			return bEnabled;
+		}
+		
+
+		/** If this is an front end filter, this is the filter object */
+		const TSharedPtr<FFilterBase<FilterType>>& GetFrontendFilter() const
+		{
+			return FrontendFilter;
+		}
+
+		/** Returns the display name for this filter */
+		virtual FText GetFilterDisplayName() const
+		{
+			FText FilterName;
+			if (FrontendFilter.IsValid())
+			{
+				FilterName = FrontendFilter->GetDisplayName();
+			}
+
+			if (FilterName.IsEmpty())
+			{
+				FilterName = LOCTEXT("UnknownFilter", "???");
+			}
+
+			return FilterName;
+		}
+		
+		virtual FString GetFilterName() const
+		{
+			FString FilterName;
+			if (FrontendFilter.IsValid())
+			{
+				FilterName = FrontendFilter->GetName();
+			}
+			return FilterName;
+		}
+
+	protected:
+
+		/** Function that constructs the actual widget for subclasses to call */
+		void Construct_Internal()
+		{
+			ChildSlot
+			[
+				SNew(SBorder)
+				 .Padding(1.0f)
+				.BorderImage(FAppStyle::Get().GetBrush("ContentBrowser.FilterBackground"))
+				 [
+					SAssignNew( ToggleButtonPtr, SFilterCheckBox )
+					.Style(FAppStyle::Get(), "ContentBrowser.FilterButton")
+					.ToolTipText(FilterToolTip)
+					.Padding(0.0f)
+					.IsChecked(this, &SFilter::IsChecked)
+					.OnCheckStateChanged(this, &SFilter::FilterToggled)
+					.OnGetMenuContent(this, &SFilter::GetRightClickMenuContent)
+					[
+						SNew(SHorizontalBox)
+						+SHorizontalBox::Slot()
+						.VAlign(VAlign_Center)
+						.AutoWidth()
+						[
+							SNew(SImage)
+							.Image(FAppStyle::Get().GetBrush("ContentBrowser.FilterImage"))
+							.ColorAndOpacity(this, &SFilter::GetFilterImageColorAndOpacity)
+						]
+						+SHorizontalBox::Slot()
+						.Padding(TAttribute<FMargin>(this, &SFilter::GetFilterNamePadding))
+						.VAlign(VAlign_Center)
+						[
+							SNew(STextBlock)
+							.Text(this, &SFilter::GetFilterDisplayName)
+							.IsEnabled_Lambda([this] {return bEnabled;})
+						]
+					]
+				]
+			];
+
+			ToggleButtonPtr->SetOnFilterCtrlClicked(FOnClicked::CreateSP(this, &SFilter::FilterCtrlClicked));
+			ToggleButtonPtr->SetOnFilterAltClicked(FOnClicked::CreateSP(this, &SFilter::FilterAltClicked));
+			ToggleButtonPtr->SetOnFilterDoubleClicked( FOnClicked::CreateSP(this, &SFilter::FilterDoubleClicked) );
+			ToggleButtonPtr->SetOnFilterMiddleButtonClicked( FOnClicked::CreateSP(this, &SFilter::FilterMiddleButtonClicked) );
+		}
+		
+		/** Handler for when the filter checkbox is clicked */
+		void FilterToggled(ECheckBoxState NewState)
+		{
+			bEnabled = NewState == ECheckBoxState::Checked;
+			OnFilterChanged.ExecuteIfBound();
+		}
+
+		/** Handler for when the filter checkbox is clicked and a control key is pressed */
+		FReply FilterCtrlClicked()
+		{
+			OnRequestEnableAll.ExecuteIfBound();
+			return FReply::Handled();
+		}
+
+		/** Handler for when the filter checkbox is clicked and an alt key is pressed */
+		FReply FilterAltClicked()
+		{
+			OnRequestDisableAll.ExecuteIfBound();
+			return FReply::Handled();
+		}
+
+		/** Handler for when the filter checkbox is double clicked */
+		FReply FilterDoubleClicked()
+		{
+			// Disable all other filters and enable this one.
+			OnRequestDisableAll.ExecuteIfBound();
+			bEnabled = true;
+			OnFilterChanged.ExecuteIfBound();
+
+			return FReply::Handled();
+		}
+
+		/** Handler for when the filter checkbox is middle button clicked */
+		FReply FilterMiddleButtonClicked()
+		{
+			RemoveFilter();
+			return FReply::Handled();
+		}
+
+		/** Handler to create a right click menu */
+		TSharedRef<SWidget> GetRightClickMenuContent()
+		{
+			FMenuBuilder MenuBuilder(/*bInShouldCloseWindowAfterMenuSelection=*/true, NULL);
+
+			MenuBuilder.BeginSection("FilterOptions", LOCTEXT("FilterContextHeading", "Filter Options"));
+			{
+				MenuBuilder.AddMenuEntry(
+					FText::Format( LOCTEXT("RemoveFilter", "Remove: {0}"), GetFilterDisplayName() ),
+					LOCTEXT("RemoveFilterTooltip", "Remove this filter from the list. It can be added again in the filters menu."),
+					FSlateIcon(),
+					FUIAction( FExecuteAction::CreateSP(this, &SFilter::RemoveFilter) )
+					);
+
+				MenuBuilder.AddMenuEntry(
+					FText::Format( LOCTEXT("EnableOnlyThisFilter", "Enable Only This: {0}"), GetFilterDisplayName() ),
+					LOCTEXT("EnableOnlyThisFilterTooltip", "Enable only this filter from the list."),
+					FSlateIcon(),
+					FUIAction( FExecuteAction::CreateSP(this, &SFilter::EnableOnly) )
+					);
+
+			}
+			MenuBuilder.EndSection();
+
+			MenuBuilder.BeginSection("FilterBulkOptions", LOCTEXT("BulkFilterContextHeading", "Bulk Filter Options"));
+			{
+				MenuBuilder.AddMenuEntry(
+					LOCTEXT("EnableAllFilters", "Enable All Filters"),
+					LOCTEXT("EnableAllFiltersTooltip", "Enables all filters."),
+					FSlateIcon(),
+					FUIAction(FExecuteAction::CreateSP(this, &SFilter::EnableAllFilters))
+				);
+
+				MenuBuilder.AddMenuEntry(
+					LOCTEXT("DisableAllFilters", "Disable All Filters"),
+					LOCTEXT("DisableAllFiltersTooltip", "Disables all active filters."),
+					FSlateIcon(),
+					FUIAction( FExecuteAction::CreateSP(this, &SFilter::DisableAllFilters) )
+					);
+
+				MenuBuilder.AddMenuEntry(
+					LOCTEXT("RemoveAllFilters", "Remove All Filters"),
+					LOCTEXT("RemoveAllFiltersTooltip", "Removes all filters from the list."),
+					FSlateIcon(),
+					FUIAction( FExecuteAction::CreateSP(this, &SFilter::RemoveAllFilters) )
+					);
+
+				MenuBuilder.AddMenuEntry(
+					FText::Format( LOCTEXT("RemoveAllButThisFilter", "Remove All But This: {0}"), GetFilterDisplayName() ),
+					LOCTEXT("RemoveAllButThisFilterTooltip", "Remove all other filters except this one from the list."),
+					FSlateIcon(),
+					FUIAction( FExecuteAction::CreateSP(this, &SFilter::RemoveAllButThis) )
+					);
+			}
+			MenuBuilder.EndSection();
+
+			if (FrontendFilter.IsValid())
+			{
+				FrontendFilter->ModifyContextMenu(MenuBuilder);
+			}
+
+			return MenuBuilder.MakeWidget();
+		}
+
+		/** Removes this filter from the filter list */
+		void RemoveFilter()
+		{
+			TSharedRef<SFilter> Self = SharedThis(this);
+			OnRequestRemove.ExecuteIfBound( Self );
+		}
+
+		/** Remove all but this filter from the filter list. */
+		void RemoveAllButThis()
+		{
+			TSharedRef<SFilter> Self = SharedThis(this);
+			OnRequestRemoveAllButThis.ExecuteIfBound(Self);
+		}
+
+		/** Enables only this filter from the filter list */
+		void EnableOnly()
+		{
+			TSharedRef<SFilter> Self = SharedThis(this);
+			OnRequestEnableOnly.ExecuteIfBound( Self );
+		}
+
+		/** Enables all filters in the list */
+		void EnableAllFilters()
+		{
+			OnRequestEnableAll.ExecuteIfBound();
+		}
+
+		/** Disables all active filters in the list */
+		void DisableAllFilters()
+		{
+			OnRequestDisableAll.ExecuteIfBound();
+		}
+
+		/** Removes all filters in the list */
+		void RemoveAllFilters()
+		{
+			OnRequestRemoveAll.ExecuteIfBound();
+		}
+
+		/** Handler to determine the "checked" state of the filter checkbox */
+		ECheckBoxState IsChecked() const
+		{
+			return bEnabled ? ECheckBoxState::Checked : ECheckBoxState::Unchecked;
+		}
+
+		/** Handler to determine the color of the checkbox when it is checked */
+		FSlateColor GetFilterImageColorAndOpacity() const
+		{
+			return bEnabled ? FilterColor : FAppStyle::Get().GetSlateColor("Colors.Recessed");
+		}
+
+		EVisibility GetFilterOverlayVisibility() const
+		{
+			return bEnabled ? EVisibility::Collapsed : EVisibility::HitTestInvisible;
+		}
+
+		/** Handler to determine the padding of the checkbox text when it is pressed */
+		FMargin GetFilterNamePadding() const
+		{
+			return ToggleButtonPtr->IsPressed() ? FMargin(4,2,4,0) : FMargin(4,1,4,1);
+		}
+
+	protected:
+		/** Invoked when the filter toggled */
+		FOnFilterChanged OnFilterChanged;
+
+		/** Invoked when a request to remove this filter originated from within this filter */
+		FOnRequestRemove OnRequestRemove;
+
+		/** Invoked when a request to enable only this filter originated from within this filter */
+		FOnRequestEnableOnly OnRequestEnableOnly;
+
+		/** Invoked when a request to enable all filters originated from within this filter */
+		FOnRequestEnableAll OnRequestEnableAll;
+
+		/** Invoked when a request to disable all filters originated from within this filter */
+		FOnRequestDisableAll OnRequestDisableAll;
+
+		/** Invoked when a request to remove all filters originated from within this filter */
+		FOnRequestDisableAll OnRequestRemoveAll;
+
+		/** Invoked when a request to remove all filters except this one originated from within this filter */
+		FOnRequestRemoveAllButThis OnRequestRemoveAllButThis;
+
+		/** true when this filter should be applied to the search */
+		bool bEnabled;
+
+		/** If this is an front end filter, this is the filter object */
+		TSharedPtr<FFilterBase<FilterType>> FrontendFilter;
+
+		/** The button to toggle the filter on or off */
+		TSharedPtr<SFilterCheckBox> ToggleButtonPtr;
+
+		/** The color of the checkbox for this filter */
+		FLinearColor FilterColor;
+
+		/** The tooltip for this filter */
+		TAttribute<FText> FilterToolTip;
+	};
+
+public:
+
+ 	/** Returns true if any filters are applied */
+ 	bool HasAnyFilters() const
+ 	{
+ 		return Filters.Num() > 0;
+ 	}
+
+ 	/** Retrieve a specific filter */
+ 	TSharedPtr<FFilterBase<FilterType>> GetFilter(const FString& InName) const
+ 	{
+ 		for (const TSharedRef<FFilterBase<FilterType>>& Filter : AllFrontendFilters)
+ 		{
+ 			if (Filter->GetName() == InName)
+ 			{
+ 				return Filter;
+ 			}
+ 		}
+ 		return TSharedPtr<FFilterBase<FilterType>>();
+ 	}
+
+	/** Use this function to get all currently active filters (to filter your items)
+	 * Not const on purpose: Subclasses might need to update the filter collection before getting it
+	 */
+	virtual TSharedPtr< TFilterCollection<FilterType> > GetAllActiveFilters()
+ 	{
+ 		return ActiveFilters;	
+ 	}
+
+	/** Enable all the filters that are currently visible on the filter bar */
+	void EnableAllFilters()
+ 	{
+ 		for (const TSharedRef<SFilter>& Filter : Filters)
+ 		{
+ 			Filter->SetEnabled(true, false);
+ 			if (const TSharedPtr<FFilterBase<FilterType>>& FrontendFilter = Filter->GetFrontendFilter())
+ 			{
+ 				SetFrontendFilterActive(FrontendFilter.ToSharedRef(), true);
+ 			}
+ 		}
+
+ 		OnFilterChanged.ExecuteIfBound();
+ 	}
+
+	/** Disable all the filters that are currently visible on the filter bar */
+	void DisableAllFilters()
+ 	{
+ 		for (const TSharedRef<SFilter>& Filter : Filters)
+ 		{
+ 			Filter->SetEnabled(false, false);
+ 			if (const TSharedPtr<FFilterBase<FilterType>>& FrontendFilter = Filter->GetFrontendFilter())
+ 			{
+ 				SetFrontendFilterActive(FrontendFilter.ToSharedRef(), false);
+ 			}
+ 		}
+
+ 		OnFilterChanged.ExecuteIfBound();
+ 	}
+
+	/** Remove all filters from the filter bar, while disabling any active ones */
+	virtual void RemoveAllFilters()
+ 	{
+ 		if (HasAnyFilters())
+ 		{
+ 			// Update the frontend filters collection
+ 			for (const TSharedRef<SFilter>& FilterToRemove : Filters)
+ 			{
+ 				if (const TSharedPtr<FFilterBase<FilterType>>& FrontendFilter = FilterToRemove->GetFrontendFilter())
+ 				{
+ 					SetFrontendFilterActive(FrontendFilter.ToSharedRef(), false); // Deactivate.
+ 				}
+ 			}
+		
+ 			ChildSlot
+			 [
+				 SNullWidget::NullWidget
+			 ];
+
+ 			Filters.Empty();
+
+ 			// Notify that a filter has changed
+ 			OnFilterChanged.ExecuteIfBound();
+ 		}
+ 	}
+	
+ 	/** Set the check box state of the specified filter (in the filter drop down) and pin/unpin a filter widget on/from the filter bar. When a filter is pinned (was not already pinned), it is activated and deactivated when unpinned. */
+ 	void SetFilterCheckState(const TSharedPtr<FFilterBase<FilterType>>& InFilter, ECheckBoxState InCheckState)
+ 	{
+ 		if (!InFilter || InCheckState == ECheckBoxState::Undetermined)
+ 		{
+ 			return;
+ 		}
+
+ 		// Check if the filter is already checked.
+ 		TSharedRef<FFilterBase<FilterType>> Filter = InFilter.ToSharedRef();
+ 		bool FrontendFilterChecked = IsFrontendFilterInUse(Filter);
+
+ 		if (InCheckState == ECheckBoxState::Checked && !FrontendFilterChecked)
+ 		{
+ 			AddFilterToBar(Filter)->SetEnabled(true); // Pin a filter widget on the UI and activate the filter. Same behaviour as FrontendFilterClicked()
+ 		}
+ 		else if (InCheckState == ECheckBoxState::Unchecked && FrontendFilterChecked)
+ 		{
+ 			RemoveFilter(Filter); // Unpin the filter widget and deactivate the filter.
+ 		}
+ 		// else -> Already in the desired 'check' state.
+ 	}
+
+ 	/** Returns the check box state of the specified filter (in the filter drop down). This tells whether the filter is pinned or not on the filter bar, but not if filter is active or not. @see IsFrontendFilterActive(). */
+ 	ECheckBoxState GetFilterCheckState(const TSharedPtr<FFilterBase<FilterType>>& InFilter) const
+ 	{
+ 		return InFilter && IsFrontendFilterInUse(InFilter.ToSharedRef()) ? ECheckBoxState::Checked : ECheckBoxState::Unchecked;
+ 	}
+
+ 	/** Returns true if the specified filter is both checked (pinned on the filter bar) and active (contributing to filter the result). */
+ 	bool IsFilterActive(const TSharedPtr<FFilterBase<FilterType>>& InFilter) const
+ 	{
+ 		if (InFilter.IsValid())
+ 		{
+ 			for (const TSharedRef<SFilter>& Filter : Filters)
+ 			{
+ 				if (InFilter == Filter->GetFrontendFilter())
+ 				{
+ 					return Filter->IsEnabled(); // Is active or not?
+ 				}
+ 			}
+ 		}
+ 		return false;
+ 	}
+
+protected:
+	
+	/** Remove all filters except the specified one */
+	virtual void RemoveAllButThis(const TSharedRef<SFilter>& FilterToKeep)
+	{
+		for (const TSharedRef<SFilter>& Filter : Filters)
+		{
+			if (Filter == FilterToKeep)
+			{
+				continue;
+			}
+
+			if (const TSharedPtr<FFilterBase<FilterType>>& FrontendFilter = Filter->GetFrontendFilter())
+			{
+				SetFrontendFilterActive(FrontendFilter.ToSharedRef(), false);
+			}
+		}
+
+		Filters.Empty();
+
+		AddFilterToBar(FilterToKeep);
+
+		OnFilterChanged.ExecuteIfBound();
+	}
+
+	/** Helper that creates a toolbar with all the given SFilter's as toolbar items. Filters that don't fit appear in the overflow menu as toggles. */
+	static TSharedRef<SWidget> MakeFilterToolBarWidget(const TArray<TSharedRef<SFilter>>& Filters)
+	{
+		FSlimHorizontalToolBarBuilder ToolbarBuilder(TSharedPtr<const FUICommandList>(), FMultiBoxCustomization::None, TSharedPtr<FExtender>(), true);
+		ToolbarBuilder.SetLabelVisibility(EVisibility::Collapsed);
+		ToolbarBuilder.SetStyle(&FAppStyle::Get(), "ContentBrowser.FilterToolBar");
+
+		for (const TSharedRef<SFilter>& Filter : Filters)
+		{
+			ToolbarBuilder.AddWidget(Filter, NAME_None, true, EHorizontalAlignment::HAlign_Fill, FNewMenuDelegate::CreateLambda([Filter](FMenuBuilder& MenuBuilder)
+			{
+				FUIAction Action;
+				Action.GetActionCheckState = FGetActionCheckState::CreateLambda([Filter]()
+				{
+					return Filter->IsEnabled() ? ECheckBoxState::Checked : ECheckBoxState::Unchecked;
+				});
+				Action.ExecuteAction = FExecuteAction::CreateLambda([Filter]()
+				{
+					Filter->SetEnabled(!Filter->IsEnabled());
+				});
+
+				MenuBuilder.AddMenuEntry(Filter->GetFilterDisplayName(), FText::GetEmpty(), FSlateIcon(), Action, NAME_None, EUserInterfaceActionType::ToggleButton);
+			}));
+		}
+
+		return ToolbarBuilder.MakeWidget();
+	}
+	
+ 	/** Sets the active state of a frontend filter. */
+ 	void SetFrontendFilterActive(const TSharedRef<FFilterBase<FilterType>>& Filter, bool bActive)
+	{
+		if(Filter->IsInverseFilter())
+		{
+			//Inverse filters are active when they are "disabled"
+			bActive = !bActive;
+		}
+		Filter->ActiveStateChanged(bActive);
+
+		if ( bActive )
+		{
+			ActiveFilters->Add(Filter);
+		}
+		else
+		{
+			ActiveFilters->Remove(Filter);
+		}
+	}
+
+	/* 'Activate' A filter by adding it to the filter bar, does not turn it on */
+	TSharedRef<SFilter> AddFilterToBar(const TSharedRef<FFilterBase<FilterType>>& Filter)
+	{
+		TSharedRef<SFilter> NewFilter =
+			SNew(SFilter)
+			.FrontendFilter(Filter)
+			.OnFilterChanged( this, &SBasicFilterBar<FilterType>::FrontendFilterChanged, Filter )
+			.OnRequestRemove(this, &SBasicFilterBar<FilterType>::RemoveFilterAndUpdate)
+			.OnRequestEnableOnly(this, &SBasicFilterBar<FilterType>::EnableOnlyThisFilter)
+			.OnRequestEnableAll(this, &SBasicFilterBar<FilterType>::EnableAllFilters)
+			.OnRequestDisableAll(this, &SBasicFilterBar<FilterType>::DisableAllFilters)
+			.OnRequestRemoveAll(this, &SBasicFilterBar<FilterType>::RemoveAllFilters)
+			.OnRequestRemoveAllButThis(this, &SBasicFilterBar<FilterType>::RemoveAllButThis);
+
+		AddFilterToBar( NewFilter );
+
+		return NewFilter;
+	}
+
+	/* 'Activate' A filter by adding it to the filter bar, does not turn it on */
+	void AddFilterToBar(const TSharedRef<SFilter>& FilterToAdd)
+	{
+		Filters.Add(FilterToAdd);
+	
+		ChildSlot
+		[
+			MakeFilterToolBarWidget(Filters)
+		];
+	}
+
+	/** Handler for when the enable only this button was clicked on a single filter */
+	void EnableOnlyThisFilter(const TSharedRef<SFilter>& FilterToEnable)
+	{
+		for (const TSharedRef<SFilter>& Filter : Filters)
+		{
+			bool bEnable = Filter == FilterToEnable;
+			Filter->SetEnabled(bEnable, /*ExecuteOnFilterChange*/false);
+			if (const TSharedPtr<FFilterBase<FilterType>>& FrontendFilter = Filter->GetFrontendFilter())
+			{
+				SetFrontendFilterActive(FrontendFilter.ToSharedRef(), bEnable);
+			}
+		}
+
+		OnFilterChanged.ExecuteIfBound();
+	}
+
+	/** Remove a filter from the filter bar */
+	void RemoveFilter(const TSharedRef<FFilterBase<FilterType>>& InFilterToRemove, bool ExecuteOnFilterChanged= true)
+	{
+		TSharedPtr<SFilter> FilterToRemove;
+		for (const TSharedRef<SFilter>& Filter : Filters)
+		{
+			if (Filter->GetFrontendFilter() == InFilterToRemove)
+			{
+				FilterToRemove = Filter;
+				break;
+			}
+		}
+
+		if (FilterToRemove.IsValid())
+		{
+			if (ExecuteOnFilterChanged)
+			{
+				RemoveFilterAndUpdate(FilterToRemove.ToSharedRef());
+			}
+			else
+			{
+				RemoveFilter(FilterToRemove.ToSharedRef());
+			}
+		}
+	}
+
+	/** Remove a filter from the filter bar */
+	virtual void RemoveFilter(const TSharedRef<SFilter>& FilterToRemove)
+	{
+		Filters.Remove(FilterToRemove);
+
+		if (const TSharedPtr<FFilterBase<FilterType>>& FrontendFilter = FilterToRemove->GetFrontendFilter()) // Is valid?
+		{
+			// Update the frontend filters collection
+			SetFrontendFilterActive(FrontendFilter.ToSharedRef(), false);
+			OnFilterChanged.ExecuteIfBound();
+		}
+
+		ChildSlot
+		[
+			MakeFilterToolBarWidget(Filters)
+		];
+	}
+
+	/** Remove a filter from the filter bar */
+	void RemoveFilterAndUpdate(const TSharedRef<SFilter>& FilterToRemove)
+	{
+		RemoveFilter(FilterToRemove);
+
+		// Notify that a filter has changed
+		OnFilterChanged.ExecuteIfBound();
+	}
+
+ 	/** Handler for when a frontend filter state has changed */
+ 	void FrontendFilterChanged(TSharedRef<FFilterBase<FilterType>> FrontendFilter)
+	{
+		TSharedPtr<SFilter> FilterToUpdate;
+		for (const TSharedRef<SFilter>& Filter : Filters)
+		{
+			if (Filter->GetFrontendFilter() == FrontendFilter)
+			{
+				FilterToUpdate = Filter;
+				break;
+			}
+		}
+
+		if (FilterToUpdate.IsValid())
+		{
+			SetFrontendFilterActive(FrontendFilter, FilterToUpdate->IsEnabled());
+			OnFilterChanged.ExecuteIfBound();
+		}
+	}
+
+ 	/** Handler for when the add filter menu is populated by a non-category */
+ 	void CreateOtherFiltersMenuCategory(FToolMenuSection& Section, TSharedPtr<FFilterCategory> MenuCategory) const
+	{
+		for (const TSharedRef<FFilterBase<FilterType>>& FrontendFilter : AllFrontendFilters)
+		{
+			if(FrontendFilter->GetCategory() == MenuCategory)
+			{
+				Section.AddMenuEntry(
+					NAME_None,
+					FrontendFilter->GetDisplayName(),
+					FrontendFilter->GetToolTipText(),
+					FSlateIcon(FAppStyle::GetAppStyleSetName(), FrontendFilter->GetIconName()),
+					FUIAction(
+					FExecuteAction::CreateSP(const_cast< SBasicFilterBar<FilterType>* >(this), &SBasicFilterBar<FilterType>::FrontendFilterClicked, FrontendFilter),
+					FCanExecuteAction(),
+					FIsActionChecked::CreateSP(this, &SBasicFilterBar<FilterType>::IsFrontendFilterInUse, FrontendFilter)),
+					EUserInterfaceActionType::ToggleButton
+					);
+			}
+		}
+	}
+
+	/** Handler for when the add filter menu is populated by a non-category */
+ 	void CreateOtherFiltersMenuCategory(UToolMenu* InMenu, TSharedPtr<FFilterCategory> MenuCategory) const
+	{
+		CreateOtherFiltersMenuCategory(InMenu->AddSection("Section"), MenuCategory);
+	}
+
+	/** Handler for a frontend filter is clicked */
+ 	void FrontendFilterClicked(TSharedRef<FFilterBase<FilterType>> FrontendFilter)
+	{
+		if (IsFrontendFilterInUse(FrontendFilter))
+		{
+			RemoveFilter(FrontendFilter);
+		}
+		else
+		{
+			TSharedRef<SFilter> NewFilter = AddFilterToBar(FrontendFilter);
+			NewFilter->SetEnabled(true);
+		}
+	}
+
+	/** Handler to check if a frontend filter is in use */
+ 	bool IsFrontendFilterInUse(TSharedRef<FFilterBase<FilterType>> FrontendFilter) const
+	{
+		for (const TSharedRef<SFilter>& Filter : Filters)
+		{
+			if (Filter->GetFrontendFilter() == FrontendFilter)
+			{
+				return true;
+			}
+		}
+
+		return false;
+	}
+
+	/** Handler for when a filter category is clicked */
+ 	void FrontendFilterCategoryClicked(TSharedPtr<FFilterCategory> MenuCategory)
+	{
+		bool bFullCategoryInUse = IsFrontendFilterCategoryInUse(MenuCategory);
+		bool ExecuteOnFilterChanged = false;
+
+		for (const TSharedRef<FFilterBase<FilterType>>& FrontendFilter : AllFrontendFilters)
+		{
+			if (FrontendFilter->GetCategory() == MenuCategory)
+			{
+				if (bFullCategoryInUse)
+				{
+					RemoveFilter(FrontendFilter, false);
+					ExecuteOnFilterChanged = true;
+				}
+				else if (!IsFrontendFilterInUse(FrontendFilter))
+				{
+					TSharedRef<SFilter> NewFilter = AddFilterToBar(FrontendFilter);
+					NewFilter->SetEnabled(true, false);
+					SetFrontendFilterActive(FrontendFilter, NewFilter->IsEnabled());
+					ExecuteOnFilterChanged = true;
+				}
+			}
+		}
+
+		if (ExecuteOnFilterChanged)
+		{
+			OnFilterChanged.ExecuteIfBound();
+		}
+	}
+	
+	/** Handler to check if a filter category is in use */
+ 	bool IsFrontendFilterCategoryInUse(TSharedPtr<FFilterCategory> MenuCategory) const
+	{
+		for (const TSharedRef<FFilterBase<FilterType>>& FrontendFilter : AllFrontendFilters)
+		{
+			if (FrontendFilter->GetCategory() == MenuCategory && !IsFrontendFilterInUse(FrontendFilter))
+			{
+				return false;
+			}
+		}
+
+		return true;
+	}
+	
+	/** Handler to determine the "checked" state of a frontend filter category in the filter dropdown */
+    ECheckBoxState IsFrontendFilterCategoryChecked(TSharedPtr<FFilterCategory> MenuCategory) const
+    {
+    	bool bIsAnyActionInUse = false;
+    	bool bIsAnyActionNotInUse = false;
+		
+		for (const TSharedRef<FFilterBase<FilterType>>& FrontendFilter : AllFrontendFilters)
+		{
+			if (FrontendFilter->GetCategory() == MenuCategory)
+			{
+				if (IsFrontendFilterInUse(FrontendFilter))
+				{
+					bIsAnyActionInUse = true;
+				}
+				else
+				{
+					bIsAnyActionNotInUse = true;
+				}
+
+				if (bIsAnyActionInUse && bIsAnyActionNotInUse)
+				{
+					return ECheckBoxState::Undetermined;
+				}
+			}
+		}
+
+    	if (bIsAnyActionInUse)
+    	{
+    		return ECheckBoxState::Checked;
+    	}
+    	else
+    	{
+    		return ECheckBoxState::Unchecked;
+    	}
+    }
+
+ 	/** Called when reset filters option is pressed */
+ 	void OnResetFilters()
+	{
+		RemoveAllFilters();
+	}
+
+ 	/** Called to set a filter active externally */
+ 	void OnSetFilterActive(bool bInActive, TWeakPtr<FFilterBase<FilterType>> InWeakFilter)
+ 	{
+ 		TSharedPtr<FFilterBase<FilterType>> Filter = InWeakFilter.Pin();
+ 		if (Filter.IsValid())
+ 		{
+ 			if (!IsFrontendFilterInUse(Filter.ToSharedRef()))
+ 			{
+ 				TSharedRef<SFilter> NewFilter = AddFilterToBar(Filter.ToSharedRef());
+ 				NewFilter->SetEnabled(bInActive);
+ 			}
+ 		}
+ 	}
+
+	/** Helper function to add common sections to the Add Filter Menu */
+	void PopulateCommonFilterSections(UToolMenu* Menu)
+	{
+		FToolMenuSection& Section = Menu->AddSection("FilterBarResetFilters");
+		Section.AddMenuEntry(
+			"ResetFilters",
+			LOCTEXT("FilterListResetFilters", "Reset Filters"),
+			LOCTEXT("FilterListResetToolTip", "Resets current filter selection"),
+			FSlateIcon(),
+			FUIAction(FExecuteAction::CreateSP(this, &SBasicFilterBar<FilterType>::OnResetFilters))
+			);
+	}
+
+	/** Helper function to add all custom filters to the Add Filter Menu */
+	void PopulateCustomFilters(UToolMenu* Menu)
+	{
+		FToolMenuSection& Section = Menu->AddSection("BasicFilterBarFiltersMenu");
+		
+		// Add all the filters
+		for (const TSharedPtr<FFilterCategory>& Category : AllFilterCategories)
+		{
+			Section.AddSubMenu(
+				NAME_None,
+				Category->Title,
+				Category->Tooltip,
+				FNewToolMenuDelegate::CreateSP(this, &SBasicFilterBar<FilterType>::CreateOtherFiltersMenuCategory, Category),
+				FUIAction(
+				FExecuteAction::CreateSP( this, &SBasicFilterBar<FilterType>::FrontendFilterCategoryClicked, Category ),
+				FCanExecuteAction(),
+				FGetActionCheckState::CreateSP(this, &SBasicFilterBar<FilterType>::IsFrontendFilterCategoryChecked, Category ) ),
+				EUserInterfaceActionType::ToggleButton
+				);
+		}
+	}
+	
+	/** Virtual Function for subclasses to save their settings */
+	virtual void SaveSettings(const FString& IniFilename, const FString& IniSection, const FString& SettingsString) const {}
+
+	/** Virtual Function for subclasses to load their settings */
+	virtual void LoadSettings(const FString& IniFilename, const FString& IniSection, const FString& SettingsString) {}
+
+private:
+
+	/** Handler for when the add filter button was clicked */
+	virtual TSharedRef<SWidget> MakeAddFilterMenu()
+	{
+		const FName FilterMenuName = "FilterBar.FilterMenu";
+		if (!UToolMenus::Get()->IsMenuRegistered(FilterMenuName))
+		{
+			UToolMenu* Menu = UToolMenus::Get()->RegisterMenu(FilterMenuName);
+			Menu->bShouldCloseWindowAfterMenuSelection = true;
+			Menu->bCloseSelfOnly = true;
+
+			Menu->AddDynamicSection(NAME_None, FNewToolMenuDelegate::CreateLambda([](UToolMenu* InMenu)
+			{
+				if (UFilterBarContext* Context = InMenu->FindContext<UFilterBarContext>())
+				{
+					Context->PopulateFilterMenu.ExecuteIfBound(InMenu);
+				}
+			}));
+		}
+
+		UFilterBarContext* FilterBarContext = NewObject<UFilterBarContext>();
+		FilterBarContext->PopulateFilterMenu = FOnPopulateAddFilterMenu::CreateSP(this, &SBasicFilterBar<FilterType>::PopulateAddFilterMenu);
+		FToolMenuContext ToolMenuContext(FilterBarContext);
+
+		return UToolMenus::Get()->GenerateWidget(FilterMenuName, ToolMenuContext);
+	}
+
+	/** Function to populate the add filter menu */
+	void PopulateAddFilterMenu(UToolMenu* Menu)
+	{
+		PopulateCommonFilterSections(Menu);
+		
+		PopulateCustomFilters(Menu);
+	}
+	
+protected:
+ 	/** All SFilters in the list */
+ 	TArray< TSharedRef<SFilter> > Filters;
+
+ 	/** All possible filter objects */
+ 	TArray< TSharedRef< FFilterBase<FilterType> > > AllFrontendFilters;
+
+	/** Currently active filter objects */
+	TSharedPtr< TFilterCollection<FilterType> > ActiveFilters;
+
+ 	/** All filter categories (for menu construction) */
+ 	TArray< TSharedPtr<FFilterCategory> > AllFilterCategories;
+
+ 	/** Delegate for when filters have changed */
+ 	FOnFilterChanged OnFilterChanged;
+
+	friend struct FFrontendFilterExternalActivationHelper<FilterType>;
+};
+
+/** Helper struct to avoid friending the whole of SBasicFilterBarSBasicFilterBar */
+template<typename FilterType>
+struct FFrontendFilterExternalActivationHelper
+{
+	static void BindToFilter(TSharedRef<SBasicFilterBar<FilterType>> InFilterList, TSharedRef<FFilterBase<FilterType>> InFrontendFilter)
+	{
+		TWeakPtr<FFilterBase<FilterType>> WeakFilter = InFrontendFilter;
+		InFrontendFilter->SetActiveEvent.AddSP(&InFilterList.Get(), &SBasicFilterBar<FilterType>::OnSetFilterActive, WeakFilter);
+	}
+};
+
+#undef LOCTEXT_NAMESPACE
\ No newline at end of file
diff --git a/Engine/Source/Developer/ToolWidgets/ToolWidgets.Build.cs b/Engine/Source/Developer/ToolWidgets/ToolWidgets.Build.cs
--- a/Engine/Source/Developer/ToolWidgets/ToolWidgets.Build.cs
+++ b/Engine/Source/Developer/ToolWidgets/ToolWidgets.Build.cs
@@ -16,6 +16,8 @@ public class ToolWidgets : ModuleRules
 				"Slate",
 				"SlateCore",
 				"InputCore",
+				"ToolMenus",
+				"CoreUObject"
 			}
 		);
 
diff --git a/Engine/Source/Editor/ContentBrowser/Private/SContentBrowser.cpp b/Engine/Source/Editor/ContentBrowser/Private/SContentBrowser.cpp
--- a/Engine/Source/Editor/ContentBrowser/Private/SContentBrowser.cpp
+++ b/Engine/Source/Editor/ContentBrowser/Private/SContentBrowser.cpp
@@ -697,6 +697,21 @@ TSharedRef<SWidget> SContentBrowser::CreateAssetView(const FContentBrowserConfig
 
 	const float SearchBoxDesiredWidth = 500.0f;
 
+	// Create the Filter Bar Widget
+	FilterListPtr = SNew(SFilterList)
+					.OnFilterChanged(this, &SContentBrowser::OnFilterChanged)
+					.OnGetContextMenu(this, &SContentBrowser::GetFilterContextMenu)
+					.Visibility((Config != nullptr ? Config->bCanShowFilters : true) ? EVisibility::Visible : EVisibility::Collapsed)
+					.FrontendFilters(FrontendFilters)
+					.AddMetaData<FTagMetaData>(FTagMetaData(TEXT("ContentBrowserFilters")));
+
+	// Create the Filter Combo Button
+	TSharedPtr<SWidget> FilterComboButton = SFilterList::MakeAddFilterButton(FilterListPtr.ToSharedRef());
+	
+	TSharedPtr<ISlateMetaData> FilterComboButtonMetaData = MakeShared<FTagMetaData>(TEXT("ContentBrowserFiltersCombo"));
+	FilterComboButton->AddMetadata(FilterComboButtonMetaData.ToSharedRef());
+	FilterComboButton->SetVisibility((Config != nullptr ? Config->bCanShowFilters : true) ? EVisibility::Visible : EVisibility::Collapsed);
+
 	return
 		SNew(SVerticalBox)
 		+ SVerticalBox::Slot()
@@ -766,20 +781,7 @@ TSharedRef<SWidget> SContentBrowser::CreateAssetView(const FContentBrowserConfig
 						.HAlign(HAlign_Left)
 						.Padding(5, 0, 0, 0)
 						[
-							SNew(SComboButton)
-							.ComboButtonStyle(&FAppStyle::Get().GetWidgetStyle<FComboButtonStyle>("SimpleComboButtonWithIcon"))
-							.ForegroundColor(FSlateColor::UseStyle())
-							.ToolTipText(LOCTEXT("AddFilterToolTip", "Add an asset filter."))
-							.OnGetMenuContent(this, &SContentBrowser::MakeAddFilterMenu)
-							.ContentPadding(FMargin(1, 0))
-							.AddMetaData<FTagMetaData>(FTagMetaData(TEXT("ContentBrowserFiltersCombo")))
-							.Visibility((Config != nullptr ? Config->bCanShowFilters : true) ? EVisibility::Visible : EVisibility::Collapsed)
-							.ButtonContent()
-							[
-								SNew(SImage)
-								.Image(FAppStyle::Get().GetBrush("Icons.Filter"))
-								.ColorAndOpacity(FSlateColor::UseForeground())
-							]
+							FilterComboButton.ToSharedRef()
 						]
 					]
 				]
@@ -793,12 +795,7 @@ TSharedRef<SWidget> SContentBrowser::CreateAssetView(const FContentBrowserConfig
 					.MinDesiredHeight_Lambda([this, SearchFilterItemHeight]() { return FilterListPtr->HasAnyFilters() ? SearchFilterItemHeight : 0; })
 					.VAlign(VAlign_Center)
 					[
-						SAssignNew(FilterListPtr, SFilterList)
-						.OnFilterChanged(this, &SContentBrowser::OnFilterChanged)
-						.OnGetContextMenu(this, &SContentBrowser::GetFilterContextMenu)
-						.Visibility((Config != nullptr ? Config->bCanShowFilters : true) ? EVisibility::Visible : EVisibility::Collapsed)
-						.FrontendFilters(FrontendFilters)
-						.AddMetaData<FTagMetaData>(FTagMetaData(TEXT("ContentBrowserFilters")))
+						FilterListPtr.ToSharedRef()
 					]
 				]
 			]
diff --git a/Engine/Source/Editor/ContentBrowser/Private/SFilterList.cpp b/Engine/Source/Editor/ContentBrowser/Private/SFilterList.cpp
--- a/Engine/Source/Editor/ContentBrowser/Private/SFilterList.cpp
+++ b/Engine/Source/Editor/ContentBrowser/Private/SFilterList.cpp
@@ -18,7 +18,6 @@
 #include "Framework/MultiBox/MultiBoxBuilder.h"
 #include "Widgets/Input/SCheckBox.h"
 #include "Styling/SlateTypes.h"
-#include "Styling/AppStyle.h"
 #include "IContentBrowserDataModule.h"
 #include "ContentBrowserDataSource.h"
 #include "ContentBrowserDataSubsystem.h"
@@ -31,501 +30,10 @@
 #include "ToolMenus.h"
 #include "ContentBrowserMenuContexts.h"
 #include "Widgets/Images/SImage.h"
-#include "Styling/SlateIconFinder.h"
+#include "Filters/GenericFilter.h"
 
 #define LOCTEXT_NAMESPACE "ContentBrowser"
 
-/** Helper struct to avoid friending the whole of SFilterList */
-struct FFrontendFilterExternalActivationHelper
-{
-	static void BindToFilter(TSharedRef<SFilterList> InFilterList, TSharedRef<FFrontendFilter> InFrontendFilter)
-	{
-		TWeakPtr<FFrontendFilter> WeakFilter = InFrontendFilter;
-		InFrontendFilter->SetActiveEvent.AddSP(&InFilterList.Get(), &SFilterList::OnSetFilterActive, WeakFilter);
-	}
-};
-
-/** A class for check boxes in the filter list. If you double click a filter checkbox, you will enable it and disable all others */
-class SFilterCheckBox : public SCheckBox
-{
-public:
-	void SetOnFilterCtrlClicked(const FOnClicked& NewFilterCtrlClicked)
-	{
-		OnFilterCtrlClicked = NewFilterCtrlClicked;
-	}
-
-	void SetOnFilterAltClicked(const FOnClicked& NewFilteAltClicked)
-	{
-		OnFilterAltClicked = NewFilteAltClicked;
-	}
-
-	void SetOnFilterDoubleClicked( const FOnClicked& NewFilterDoubleClicked )
-	{
-		OnFilterDoubleClicked = NewFilterDoubleClicked;
-	}
-
-	void SetOnFilterMiddleButtonClicked( const FOnClicked& NewFilterMiddleButtonClicked )
-	{
-		OnFilterMiddleButtonClicked = NewFilterMiddleButtonClicked;
-	}
-
-	virtual FReply OnMouseButtonDoubleClick( const FGeometry& InMyGeometry, const FPointerEvent& InMouseEvent ) override
-	{
-		if ( InMouseEvent.GetEffectingButton() == EKeys::LeftMouseButton && OnFilterDoubleClicked.IsBound() )
-		{
-			return OnFilterDoubleClicked.Execute();
-		}
-		else
-		{
-			return SCheckBox::OnMouseButtonDoubleClick(InMyGeometry, InMouseEvent);
-		}
-	}
-
-	virtual FReply OnMouseButtonUp( const FGeometry& InMyGeometry, const FPointerEvent& InMouseEvent ) override
-	{
-		if (InMouseEvent.IsControlDown() && OnFilterCtrlClicked.IsBound())
-		{
-			return OnFilterCtrlClicked.Execute();
-		}
-		else if (InMouseEvent.IsAltDown() && OnFilterAltClicked.IsBound())
-		{
-			return OnFilterAltClicked.Execute();
-		}
-		else if( InMouseEvent.GetEffectingButton() == EKeys::MiddleMouseButton && OnFilterMiddleButtonClicked.IsBound() )
-		{
-			return OnFilterMiddleButtonClicked.Execute();
-		}
-		else
-		{
-			SCheckBox::OnMouseButtonUp(InMyGeometry, InMouseEvent);
-			return FReply::Handled().ReleaseMouseCapture();
-		}
-	}
-
-private:
-	FOnClicked OnFilterCtrlClicked;
-	FOnClicked OnFilterAltClicked;
-	FOnClicked OnFilterDoubleClicked;
-	FOnClicked OnFilterMiddleButtonClicked;
-};
-
-/**
- * A single filter in the filter list. Can be removed by clicking the remove button on it.
- */
-class SFilter : public SCompoundWidget
-{
-public:
-	DECLARE_DELEGATE_OneParam( FOnRequestRemove, const TSharedRef<SFilter>& /*FilterToRemove*/ );
-	DECLARE_DELEGATE_OneParam( FOnRequestRemoveAllButThis, const TSharedRef<SFilter>& /*FilterToKeep*/ );
-	DECLARE_DELEGATE_OneParam( FOnRequestEnableOnly, const TSharedRef<SFilter>& /*FilterToEnable*/ );
-	DECLARE_DELEGATE( FOnRequestEnableAll );
-	DECLARE_DELEGATE( FOnRequestDisableAll );
-	DECLARE_DELEGATE( FOnRequestRemoveAll );
-
-	SLATE_BEGIN_ARGS( SFilter ){}
-
-		/** The asset type actions that are associated with this filter */
-		SLATE_ARGUMENT( TWeakPtr<IAssetTypeActions>, AssetTypeActions )
-
-		/** If this is an front end filter, this is the filter object */
-		SLATE_ARGUMENT( TSharedPtr<FFrontendFilter>, FrontendFilter )
-
-		/** Invoked when the filter toggled */
-		SLATE_EVENT( SFilterList::FOnFilterChanged, OnFilterChanged )
-
-		/** Invoked when a request to remove this filter originated from within this filter */
-		SLATE_EVENT( FOnRequestRemove, OnRequestRemove )
-
-		/** Invoked when a request to enable only this filter originated from within this filter */
-		SLATE_EVENT( FOnRequestEnableOnly, OnRequestEnableOnly )
-
-		/** Invoked when a request to enable all filters originated from within this filter */
-		SLATE_EVENT(FOnRequestEnableAll, OnRequestEnableAll)
-
-		/** Invoked when a request to disable all filters originated from within this filter */
-		SLATE_EVENT( FOnRequestDisableAll, OnRequestDisableAll )
-
-		/** Invoked when a request to remove all filters originated from within this filter */
-		SLATE_EVENT( FOnRequestRemoveAll, OnRequestRemoveAll )
-
-		/** Invoked when a request to remove all filters originated from within this filter */
-		SLATE_EVENT( FOnRequestRemoveAllButThis, OnRequestRemoveAllButThis )
-
-	SLATE_END_ARGS()
-
-	/** Constructs this widget with InArgs */
-	void Construct( const FArguments& InArgs )
-	{
-		bEnabled = false;
-		OnFilterChanged = InArgs._OnFilterChanged;
-		AssetTypeActions = InArgs._AssetTypeActions;
-		OnRequestRemove = InArgs._OnRequestRemove;
-		OnRequestEnableOnly = InArgs._OnRequestEnableOnly;
-		OnRequestEnableAll = InArgs._OnRequestEnableAll;
-		OnRequestDisableAll = InArgs._OnRequestDisableAll;
-		OnRequestRemoveAll = InArgs._OnRequestRemoveAll;
-		OnRequestRemoveAllButThis = InArgs._OnRequestRemoveAllButThis;
-		FrontendFilter = InArgs._FrontendFilter;
-
-		// Get the tooltip and color of the type represented by this filter
-		TAttribute<FText> FilterToolTip;
-		FilterColor = FLinearColor::White;
-		if ( InArgs._AssetTypeActions.IsValid() )
-		{
-			TSharedPtr<IAssetTypeActions> TypeActions = InArgs._AssetTypeActions.Pin();
-			FilterColor = FLinearColor( TypeActions->GetTypeColor() );
-
-			// No tooltip for asset type filters
-		}
-		else if ( FrontendFilter.IsValid() )
-		{
-			FilterColor = FrontendFilter->GetColor();
-			FilterToolTip = TAttribute<FText>::Create(TAttribute<FText>::FGetter::CreateSP(FrontendFilter.ToSharedRef(), &FFrontendFilter::GetToolTipText));
-		}
-
-		ChildSlot
-		[
-			SNew(SBorder)
- 			.Padding(1.0f)
-			.BorderImage(FAppStyle::Get().GetBrush("ContentBrowser.FilterBackground"))
- 			[
-				SAssignNew( ToggleButtonPtr, SFilterCheckBox )
-				.Style(FAppStyle::Get(), "ContentBrowser.FilterButton")
-				.ToolTipText(FilterToolTip)
-				.Padding(0.0f)
-				.IsChecked(this, &SFilter::IsChecked)
-				.OnCheckStateChanged(this, &SFilter::FilterToggled)
-				.OnGetMenuContent(this, &SFilter::GetRightClickMenuContent)
-				[
-					SNew(SHorizontalBox)
-					+SHorizontalBox::Slot()
-					.VAlign(VAlign_Center)
-					.AutoWidth()
-					[
-						SNew(SImage)
-						.Image(FAppStyle::Get().GetBrush("ContentBrowser.FilterImage"))
-						.ColorAndOpacity(this, &SFilter::GetFilterImageColorAndOpacity)
-					]
-					+SHorizontalBox::Slot()
-					.Padding(TAttribute<FMargin>(this, &SFilter::GetFilterNamePadding))
-					.VAlign(VAlign_Center)
-					[
-						SNew(STextBlock)
-						.Text(this, &SFilter::GetFilterName)
-						.IsEnabled_Lambda([this] {return bEnabled;})
-					]
-				]
-			]
-		];
-
-		ToggleButtonPtr->SetOnFilterCtrlClicked(FOnClicked::CreateSP(this, &SFilter::FilterCtrlClicked));
-		ToggleButtonPtr->SetOnFilterAltClicked(FOnClicked::CreateSP(this, &SFilter::FilterAltClicked));
-		ToggleButtonPtr->SetOnFilterDoubleClicked( FOnClicked::CreateSP(this, &SFilter::FilterDoubleClicked) );
-		ToggleButtonPtr->SetOnFilterMiddleButtonClicked( FOnClicked::CreateSP(this, &SFilter::FilterMiddleButtonClicked) );
-	}
-
-	/** Sets whether or not this filter is applied to the combined filter */
-	void SetEnabled(bool InEnabled, bool InExecuteOnFilterChanged = true)
-	{
-		if ( InEnabled != bEnabled)
-		{
-			bEnabled = InEnabled;
-			if (InExecuteOnFilterChanged)
-			{
-				OnFilterChanged.ExecuteIfBound();
-			}
-		}
-	}
-
-	/** Returns true if this filter contributes to the combined filter */
-	bool IsEnabled() const
-	{
-		return bEnabled;
-	}
-
-	/** Returns this widgets contribution to the combined filter */
-	FARFilter GetBackendFilter() const
-	{
-		FARFilter Filter;
-
-		if ( AssetTypeActions.IsValid() )
-		{
-			if (AssetTypeActions.Pin()->CanFilter())
-			{
-				AssetTypeActions.Pin()->BuildBackendFilter(Filter);
-			}
-		}
-
-		return Filter;
-	}
-
-	/** If this is an front end filter, this is the filter object */
-	const TSharedPtr<FFrontendFilter>& GetFrontendFilter() const
-	{
-		return FrontendFilter;
-	}
-
-	/** Gets the asset type actions associated with this filter */
-	const TWeakPtr<IAssetTypeActions>& GetAssetTypeActions() const
-	{
-		return AssetTypeActions;
-	}
-
-	/** Returns the display name for this filter */
-	FText GetFilterName() const
-	{
-		FText FilterName;
-		if (AssetTypeActions.IsValid())
-		{
-			TSharedPtr<IAssetTypeActions> TypeActions = AssetTypeActions.Pin();
-			FilterName = TypeActions->GetName();
-		}
-		else if (FrontendFilter.IsValid())
-		{
-			FilterName = FrontendFilter->GetDisplayName();
-		}
-
-		if (FilterName.IsEmpty())
-		{
-			FilterName = LOCTEXT("UnknownFilter", "???");
-		}
-
-		return FilterName;
-	}
-
-private:
-	/** Handler for when the filter checkbox is clicked */
-	void FilterToggled(ECheckBoxState NewState)
-	{
-		bEnabled = NewState == ECheckBoxState::Checked;
-		OnFilterChanged.ExecuteIfBound();
-	}
-
-	/** Handler for when the filter checkbox is clicked and a control key is pressed */
-	FReply FilterCtrlClicked()
-	{
-		OnRequestEnableAll.ExecuteIfBound();
-		return FReply::Handled();
-	}
-
-	/** Handler for when the filter checkbox is clicked and an alt key is pressed */
-	FReply FilterAltClicked()
-	{
-		OnRequestDisableAll.ExecuteIfBound();
-		return FReply::Handled();
-	}
-
-	/** Handler for when the filter checkbox is double clicked */
-	FReply FilterDoubleClicked()
-	{
-		// Disable all other filters and enable this one.
-		OnRequestDisableAll.ExecuteIfBound();
-		bEnabled = true;
-		OnFilterChanged.ExecuteIfBound();
-
-		return FReply::Handled();
-	}
-
-	/** Handler for when the filter checkbox is middle button clicked */
-	FReply FilterMiddleButtonClicked()
-	{
-		RemoveFilter();
-		return FReply::Handled();
-	}
-
-	/** Handler to create a right click menu */
-	TSharedRef<SWidget> GetRightClickMenuContent()
-	{
-		FMenuBuilder MenuBuilder(/*bInShouldCloseWindowAfterMenuSelection=*/true, NULL);
-
-		MenuBuilder.BeginSection("FilterOptions", LOCTEXT("FilterContextHeading", "Filter Options"));
-		{
-			MenuBuilder.AddMenuEntry(
-				FText::Format( LOCTEXT("RemoveFilter", "Remove: {0}"), GetFilterName() ),
-				LOCTEXT("RemoveFilterTooltip", "Remove this filter from the list. It can be added again in the filters menu."),
-				FSlateIcon(),
-				FUIAction( FExecuteAction::CreateSP(this, &SFilter::RemoveFilter) )
-				);
-
-			MenuBuilder.AddMenuEntry(
-				FText::Format( LOCTEXT("EnableOnlyThisFilter", "Enable Only This: {0}"), GetFilterName() ),
-				LOCTEXT("EnableOnlyThisFilterTooltip", "Enable only this filter from the list."),
-				FSlateIcon(),
-				FUIAction( FExecuteAction::CreateSP(this, &SFilter::EnableOnly) )
-				);
-
-		}
-		MenuBuilder.EndSection();
-
-		MenuBuilder.BeginSection("FilterBulkOptions", LOCTEXT("BulkFilterContextHeading", "Bulk Filter Options"));
-		{
-			MenuBuilder.AddMenuEntry(
-				LOCTEXT("EnableAllFilters", "Enable All Filters"),
-				LOCTEXT("EnableAllFiltersTooltip", "Enables all filters."),
-				FSlateIcon(),
-				FUIAction(FExecuteAction::CreateSP(this, &SFilter::EnableAllFilters))
-			);
-
-			MenuBuilder.AddMenuEntry(
-				LOCTEXT("DisableAllFilters", "Disable All Filters"),
-				LOCTEXT("DisableAllFiltersTooltip", "Disables all active filters."),
-				FSlateIcon(),
-				FUIAction( FExecuteAction::CreateSP(this, &SFilter::DisableAllFilters) )
-				);
-
-			MenuBuilder.AddMenuEntry(
-				LOCTEXT("RemoveAllFilters", "Remove All Filters"),
-				LOCTEXT("RemoveAllFiltersTooltip", "Removes all filters from the list."),
-				FSlateIcon(),
-				FUIAction( FExecuteAction::CreateSP(this, &SFilter::RemoveAllFilters) )
-				);
-
-			MenuBuilder.AddMenuEntry(
-				FText::Format( LOCTEXT("RemoveAllButThisFilter", "Remove All But This: {0}"), GetFilterName() ),
-				LOCTEXT("RemoveAllButThisFilterTooltip", "Remove all other filters except this one from the list."),
-				FSlateIcon(),
-				FUIAction( FExecuteAction::CreateSP(this, &SFilter::RemoveAllButThis) )
-				);
-		}
-		MenuBuilder.EndSection();
-
-		if (FrontendFilter.IsValid())
-		{
-			FrontendFilter->ModifyContextMenu(MenuBuilder);
-		}
-
-		return MenuBuilder.MakeWidget();
-	}
-
-	/** Removes this filter from the filter list */
-	void RemoveFilter()
-	{
-		TSharedRef<SFilter> Self = SharedThis(this);
-		OnRequestRemove.ExecuteIfBound( Self );
-	}
-
-	/** Remove all but this filter from the filter list. */
-	void RemoveAllButThis()
-	{
-		TSharedRef<SFilter> Self = SharedThis(this);
-		OnRequestRemoveAllButThis.ExecuteIfBound(Self);
-	}
-
-	/** Enables only this filter from the filter list */
-	void EnableOnly()
-	{
-		TSharedRef<SFilter> Self = SharedThis(this);
-		OnRequestEnableOnly.ExecuteIfBound( Self );
-	}
-
-	/** Enables all filters in the list */
-	void EnableAllFilters()
-	{
-		OnRequestEnableAll.ExecuteIfBound();
-	}
-
-	/** Disables all active filters in the list */
-	void DisableAllFilters()
-	{
-		OnRequestDisableAll.ExecuteIfBound();
-	}
-
-	/** Removes all filters in the list */
-	void RemoveAllFilters()
-	{
-		OnRequestRemoveAll.ExecuteIfBound();
-	}
-
-	/** Handler to determine the "checked" state of the filter checkbox */
-	ECheckBoxState IsChecked() const
-	{
-		return bEnabled ? ECheckBoxState::Checked : ECheckBoxState::Unchecked;
-	}
-
-	/** Handler to determine the color of the checkbox when it is checked */
-	FSlateColor GetFilterImageColorAndOpacity() const
-	{
-		return bEnabled ? FilterColor : FAppStyle::Get().GetSlateColor("Colors.Recessed");
-	}
-
-	EVisibility GetFilterOverlayVisibility() const
-	{
-		return bEnabled ? EVisibility::Collapsed : EVisibility::HitTestInvisible;
-	}
-
-	/** Handler to determine the padding of the checkbox text when it is pressed */
-	FMargin GetFilterNamePadding() const
-	{
-		return ToggleButtonPtr->IsPressed() ? FMargin(4,2,4,0) : FMargin(4,1,4,1);
-	}
-
-private:
-	/** Invoked when the filter toggled */
-	SFilterList::FOnFilterChanged OnFilterChanged;
-
-	/** Invoked when a request to remove this filter originated from within this filter */
-	FOnRequestRemove OnRequestRemove;
-
-	/** Invoked when a request to enable only this filter originated from within this filter */
-	FOnRequestEnableOnly OnRequestEnableOnly;
-
-	/** Invoked when a request to enable all filters originated from within this filter */
-	FOnRequestEnableAll OnRequestEnableAll;
-
-	/** Invoked when a request to disable all filters originated from within this filter */
-	FOnRequestDisableAll OnRequestDisableAll;
-
-	/** Invoked when a request to remove all filters originated from within this filter */
-	FOnRequestDisableAll OnRequestRemoveAll;
-
-	/** Invoked when a request to remove all filters except this one originated from within this filter */
-	FOnRequestRemoveAllButThis OnRequestRemoveAllButThis;
-
-	/** true when this filter should be applied to the search */
-	bool bEnabled;
-
-	/** The asset type actions that are associated with this filter */
-	TWeakPtr<IAssetTypeActions> AssetTypeActions;
-
-	/** If this is an front end filter, this is the filter object */
-	TSharedPtr<FFrontendFilter> FrontendFilter;
-
-	/** The button to toggle the filter on or off */
-	TSharedPtr<SFilterCheckBox> ToggleButtonPtr;
-
-	/** The color of the checkbox for this filter */
-	FLinearColor FilterColor;
-};
-
-/** Helper that creates a toolbar with all the given SFilter's as toolbar items. Filters that don't fit appear in the overflow menu as toggles. */
-static TSharedRef<SWidget> MakeFilterToolBarWidget(const TArray<TSharedRef<SFilter>>& Filters)
-{
-	FSlimHorizontalToolBarBuilder ToolbarBuilder(TSharedPtr<const FUICommandList>(), FMultiBoxCustomization::None, TSharedPtr<FExtender>(), true);
-	ToolbarBuilder.SetLabelVisibility(EVisibility::Collapsed);
-	ToolbarBuilder.SetStyle(&FAppStyle::Get(), "ContentBrowser.FilterToolBar");
-
-	for (const TSharedRef<SFilter>& Filter : Filters)
-	{
-		ToolbarBuilder.AddWidget(Filter, NAME_None, true, EHorizontalAlignment::HAlign_Fill, FNewMenuDelegate::CreateLambda([Filter](FMenuBuilder& MenuBuilder)
-		{
-			FUIAction Action;
-			Action.GetActionCheckState = FGetActionCheckState::CreateLambda([Filter]()
-			{
-				return Filter->IsEnabled() ? ECheckBoxState::Checked : ECheckBoxState::Unchecked;
-			});
-			Action.ExecuteAction = FExecuteAction::CreateLambda([Filter]()
-			{
-				Filter->SetEnabled(!Filter->IsEnabled());
-			});
-
-			MenuBuilder.AddMenuEntry(Filter->GetFilterName(), FText::GetEmpty(), FSlateIcon(), Action, NAME_None, EUserInterfaceActionType::ToggleButton);
-		}));
-	}
-
-	return ToolbarBuilder.MakeWidget();
-}
-
-
 /////////////////////
 // SFilterList
 /////////////////////
@@ -534,26 +42,25 @@ static TSharedRef<SWidget> MakeFilterToolBarWidget(const TArray<TSharedRef<SFilt
 void SFilterList::Construct( const FArguments& InArgs )
 {
 	OnGetContextMenu = InArgs._OnGetContextMenu;
-	OnFilterChanged = InArgs._OnFilterChanged;
-	FrontendFilters = InArgs._FrontendFilters;
+	this->OnFilterChanged = InArgs._OnFilterChanged;
+	this->ActiveFilters = InArgs._FrontendFilters;
 	InitialClassFilters = InArgs._InitialClassFilters;
 
 	TSharedPtr<FFrontendFilterCategory> DefaultCategory = MakeShareable( new FFrontendFilterCategory(LOCTEXT("FrontendFiltersCategory", "Other Filters"), LOCTEXT("FrontendFiltersCategoryTooltip", "Filter assets by all filters in this category.")) );
-
+	
 	// Add all built-in frontend filters here
-	AllFrontendFilters.Add( MakeShareable(new FFrontendFilter_CheckedOut(DefaultCategory)) );
-	AllFrontendFilters.Add( MakeShareable(new FFrontendFilter_Modified(DefaultCategory)) );
-	AllFrontendFilters.Add( MakeShareable(new FFrontendFilter_Writable(DefaultCategory)) );
-	AllFrontendFilters.Add( MakeShareable(new FFrontendFilter_ShowOtherDevelopers(DefaultCategory)) );
-	AllFrontendFilters.Add( MakeShareable(new FFrontendFilter_ReplicatedBlueprint(DefaultCategory)) );
-	AllFrontendFilters.Add( MakeShareable(new FFrontendFilter_ShowRedirectors(DefaultCategory)) );
-	AllFrontendFilters.Add( MakeShareable(new FFrontendFilter_InUseByLoadedLevels(DefaultCategory)) );
-	AllFrontendFilters.Add( MakeShareable(new FFrontendFilter_UsedInAnyLevel(DefaultCategory)) );
-	AllFrontendFilters.Add( MakeShareable(new FFrontendFilter_NotUsedInAnyLevel(DefaultCategory)) );
-	AllFrontendFilters.Add( MakeShareable(new FFrontendFilter_ArbitraryComparisonOperation(DefaultCategory)) );
-	AllFrontendFilters.Add(MakeShareable(new FFrontendFilter_Recent(DefaultCategory)));
-	AllFrontendFilters.Add( MakeShareable(new FFrontendFilter_NotSourceControlled(DefaultCategory)) );
-	AllFrontendFilters.Add(MakeShareable(new FFrontendFilter_VirtualizedData(DefaultCategory)));
+	AllFrontendFilters_Internal.Add( MakeShareable(new FFrontendFilter_CheckedOut(DefaultCategory)) );
+	AllFrontendFilters_Internal.Add( MakeShareable(new FFrontendFilter_Modified(DefaultCategory)) );
+	AllFrontendFilters_Internal.Add( MakeShareable(new FFrontendFilter_Writable(DefaultCategory)) );
+	AllFrontendFilters_Internal.Add( MakeShareable(new FFrontendFilter_ShowOtherDevelopers(DefaultCategory)) );
+	AllFrontendFilters_Internal.Add( MakeShareable(new FFrontendFilter_ReplicatedBlueprint(DefaultCategory)) );
+	AllFrontendFilters_Internal.Add( MakeShareable(new FFrontendFilter_ShowRedirectors(DefaultCategory)) );
+	AllFrontendFilters_Internal.Add( MakeShareable(new FFrontendFilter_InUseByLoadedLevels(DefaultCategory)) );
+	AllFrontendFilters_Internal.Add( MakeShareable(new FFrontendFilter_UsedInAnyLevel(DefaultCategory)) );
+	AllFrontendFilters_Internal.Add( MakeShareable(new FFrontendFilter_NotUsedInAnyLevel(DefaultCategory)) );
+	AllFrontendFilters_Internal.Add( MakeShareable(new FFrontendFilter_ArbitraryComparisonOperation(DefaultCategory)) );
+	AllFrontendFilters_Internal.Add(MakeShareable(new FFrontendFilter_Recent(DefaultCategory)));
+	AllFrontendFilters_Internal.Add( MakeShareable(new FFrontendFilter_NotSourceControlled(DefaultCategory)) );
 
 	// Add any global user-defined frontend filters
 	for (TObjectIterator<UContentBrowserFrontEndFilterExtension> ExtensionIt(RF_NoFlags); ExtensionIt; ++ExtensionIt)
@@ -565,15 +72,15 @@ void SFilterList::Construct( const FArguments& InArgs )
 				// Grab the filters
 				TArray< TSharedRef<FFrontendFilter> > ExtendedFrontendFilters;
 				PotentialExtension->AddFrontEndFilterExtensions(DefaultCategory, ExtendedFrontendFilters);
-				AllFrontendFilters.Append(ExtendedFrontendFilters);
+				AllFrontendFilters_Internal.Append(ExtendedFrontendFilters);
 
 				// Grab the categories
 				for (const TSharedRef<FFrontendFilter>& FilterRef : ExtendedFrontendFilters)
 				{
-					TSharedPtr<FFrontendFilterCategory> Category = FilterRef->GetCategory();
+					TSharedPtr<FFilterCategory> Category = FilterRef->GetCategory();
 					if (Category.IsValid())
 					{
-						AllFrontendFilterCategories.AddUnique(Category);
+						this->AllFilterCategories.AddUnique(Category);
 					}
 				}
 			}
@@ -583,25 +90,30 @@ void SFilterList::Construct( const FArguments& InArgs )
 	// Add in filters specific to this invocation
 	for (const TSharedRef<FFrontendFilter>& Filter : InArgs._ExtraFrontendFilters)
 	{
-		if (TSharedPtr<FFrontendFilterCategory> Category = Filter->GetCategory())
+		if (TSharedPtr<FFilterCategory> Category = Filter->GetCategory())
 		{
-			AllFrontendFilterCategories.AddUnique(Category);
+			this->AllFilterCategories.AddUnique(Category);
 		}
 
-		AllFrontendFilters.Add(Filter);
+		AllFrontendFilters_Internal.Add(Filter);
 	}
 
-	AllFrontendFilterCategories.AddUnique(DefaultCategory);
+	this->AllFilterCategories.AddUnique(DefaultCategory);
 
-	
-	for (const TSharedRef<FFrontendFilter>& Filter : AllFrontendFilters)
+	// Add the local copy of all filters to SFilterBar's copy of all filters
+	for(TSharedRef<FFrontendFilter> FrontendFilter : AllFrontendFilters_Internal)
 	{
-		// Bind external activation event
-		FFrontendFilterExternalActivationHelper::BindToFilter(SharedThis(this), Filter);
-
-		// Auto add all inverse filters
-		SetFrontendFilterActive(Filter, false);
+		this->AllFrontendFilters.Add(FrontendFilter);
 	}
+	
+	SAssetFilterBar<FAssetFilterType>::FArguments Args;
+
+	/** Explicitly setting this to true as it should ALWAYS be true for SFilterList */
+	Args._UseDefaultAssetFilters = true;
+	Args._OnFilterChanged = this->OnFilterChanged;
+	
+	SAssetFilterBar<FAssetFilterType>::Construct(Args);
+
 }
 
 FReply SFilterList::OnMouseButtonUp( const FGeometry& MyGeometry, const FPointerEvent& MouseEvent )
@@ -634,61 +146,9 @@ const TArray<UClass*>& SFilterList::GetInitialClassFilters()
 	return InitialClassFilters;
 }
 
-bool SFilterList::HasAnyFilters() const
-{
-	return Filters.Num() > 0;
-}
-
-FARFilter SFilterList::GetCombinedBackendFilter() const
-{
-	FARFilter CombinedFilter;
-
-	// Add all selected filters
-	for (int32 FilterIdx = 0; FilterIdx < Filters.Num(); ++FilterIdx)
-	{
-		if ( Filters[FilterIdx]->IsEnabled() )
-		{
-			CombinedFilter.Append(Filters[FilterIdx]->GetBackendFilter());
-		}
-	}
-
-	if ( CombinedFilter.bRecursiveClasses )
-	{
-		// Add exclusions for AssetTypeActions NOT in the filter.
-		// This will prevent assets from showing up that are both derived from an asset in the filter set and derived from an asset not in the filter set
-		// Get the list of all asset type actions
-		FAssetToolsModule& AssetToolsModule = FModuleManager::LoadModuleChecked<FAssetToolsModule>("AssetTools");
-		TArray< TWeakPtr<IAssetTypeActions> > AssetTypeActionsList;
-		AssetToolsModule.Get().GetAssetTypeActionsList(AssetTypeActionsList);
-		for (const TWeakPtr<IAssetTypeActions>& WeakTypeActions : AssetTypeActionsList)
-		{
-			if (const TSharedPtr<IAssetTypeActions> TypeActions = WeakTypeActions.Pin())
-			{
-				if (TypeActions->CanFilter())
-				{
-					const UClass* TypeClass = TypeActions->GetSupportedClass();
-					if (TypeClass && !CombinedFilter.ClassNames.Contains(TypeClass->GetFName()))
-					{
-						CombinedFilter.RecursiveClassesExclusionSet.Add(TypeClass->GetFName());
-					}
-				}
-			}
-		}
-	}
-
-	// HACK: A blueprint can be shown as Blueprint or as BlueprintGeneratedClass, but we don't want to distinguish them while filtering.
-	// This should be removed, once all blueprints are shown as BlueprintGeneratedClass.
-	if(CombinedFilter.ClassNames.Contains(FName(TEXT("Blueprint"))))
-	{
-		CombinedFilter.ClassNames.AddUnique(FName(TEXT("BlueprintGeneratedClass")));
-	}
-
-	return CombinedFilter;
-}
-
 TSharedPtr<FFrontendFilter> SFilterList::GetFrontendFilter(const FString& InName) const
 {
-	for (const TSharedRef<FFrontendFilter>& Filter : AllFrontendFilters)
+	for (const TSharedRef<FFrontendFilter>& Filter : AllFrontendFilters_Internal)
 	{
 		if (Filter->GetName() == InName)
 		{
@@ -703,107 +163,32 @@ TSharedRef<SWidget> SFilterList::ExternalMakeAddFilterMenu(EAssetTypeCategories:
 	return MakeAddFilterMenu(MenuExpansion);
 }
 
-void SFilterList::EnableAllFilters()
+void SFilterList::DisableFiltersThatHideItems(TArrayView<const FContentBrowserItem> ItemList)
 {
-	for (const TSharedRef<SFilter>& Filter : Filters)
+	if (HasAnyFilters() && ItemList.Num() > 0)
 	{
-		Filter->SetEnabled(true, false);
-		if (const TSharedPtr<FFrontendFilter>& FrontendFilter = Filter->GetFrontendFilter())
+		// Determine if we should disable backend filters. If any item fails the combined backend filter, disable them all.
+		bool bDisableAllBackendFilters = false;
 		{
-			SetFrontendFilterActive(FrontendFilter.ToSharedRef(), true);
-		}
-	}
-
-	OnFilterChanged.ExecuteIfBound();
-}
+			FContentBrowserDataCompiledFilter CompiledDataFilter;
+			{
+				static const FName RootPath = "/";
 
-void SFilterList::DisableAllFilters()
-{
-	for (const TSharedRef<SFilter>& Filter : Filters)
-	{
-		Filter->SetEnabled(false, false);
-		if (const TSharedPtr<FFrontendFilter>& FrontendFilter = Filter->GetFrontendFilter())
-		{
-			SetFrontendFilterActive(FrontendFilter.ToSharedRef(), false);
-		}
-	}
+				UContentBrowserDataSubsystem* ContentBrowserData = IContentBrowserDataModule::Get().GetSubsystem();
 
-	OnFilterChanged.ExecuteIfBound();
-}
+				FContentBrowserDataFilter DataFilter;
+				DataFilter.bRecursivePaths = true;
+				ContentBrowserUtils::AppendAssetFilterToContentBrowserFilter(GetCombinedBackendFilter(), nullptr, nullptr, DataFilter);
 
-void SFilterList::RemoveAllFilters()
-{
-	if (HasAnyFilters())
-	{
-		// Update the frontend filters collection
-		for (const TSharedRef<SFilter>& FilterToRemove : Filters)
-		{
-			if (const TSharedPtr<FFrontendFilter>& FrontendFilter = FilterToRemove->GetFrontendFilter())
-			{
-				SetFrontendFilterActive(FrontendFilter.ToSharedRef(), false); // Deactivate.
+				ContentBrowserData->CompileFilter(RootPath, DataFilter, CompiledDataFilter);
 			}
-		}
-		
-		ChildSlot
-		[
-			SNullWidget::NullWidget
-		];
 
-		Filters.Empty();
-
-		// Notify that a filter has changed
-		OnFilterChanged.ExecuteIfBound();
-	}
-}
-
-void SFilterList::RemoveAllButThis(const TSharedRef<SFilter>& FilterToKeep)
-{
-	for (const TSharedRef<SFilter>& Filter : Filters)
-	{
-		if (Filter == FilterToKeep)
-		{
-			continue;
-		}
-
-		if (const TSharedPtr<FFrontendFilter>& FrontendFilter = Filter->GetFrontendFilter())
-		{
-			SetFrontendFilterActive(FrontendFilter.ToSharedRef(), false);
-		}
-	}
-
-	Filters.Empty();
-
-	AddFilter(FilterToKeep);
-
-	OnFilterChanged.ExecuteIfBound();
-}
-
-void SFilterList::DisableFiltersThatHideItems(TArrayView<const FContentBrowserItem> ItemList)
-{
-	if (HasAnyFilters() && ItemList.Num() > 0)
-	{
-		// Determine if we should disable backend filters. If any item fails the combined backend filter, disable them all.
-		bool bDisableAllBackendFilters = false;
-		{
-			FContentBrowserDataCompiledFilter CompiledDataFilter;
-			{
-				static const FName RootPath = "/";
-
-				UContentBrowserDataSubsystem* ContentBrowserData = IContentBrowserDataModule::Get().GetSubsystem();
-
-				FContentBrowserDataFilter DataFilter;
-				DataFilter.bRecursivePaths = true;
-				ContentBrowserUtils::AppendAssetFilterToContentBrowserFilter(GetCombinedBackendFilter(), nullptr, nullptr, DataFilter);
-
-				ContentBrowserData->CompileFilter(RootPath, DataFilter, CompiledDataFilter);
-			}
-
-			for (const FContentBrowserItem& Item : ItemList)
-			{
-				if (!Item.IsFile())
-				{
-					continue;
-				}
+			for (const FContentBrowserItem& Item : ItemList)
+			{
+				if (!Item.IsFile())
+				{
+					continue;
+				}
 
 				FContentBrowserItem::FItemDataArrayView InternalItems = Item.GetInternalItems();
 				for (const FContentBrowserItemData& InternalItem : InternalItems)
@@ -824,13 +209,12 @@ void SFilterList::DisableFiltersThatHideItems(TArrayView<const FContentBrowserIt
 		}
 
 		// Iterate over all enabled filters and disable any frontend filters that would hide any of the supplied assets
-		// and disable all backend filters if it was determined that the combined backend filter hides any of the assets
-		bool ExecuteOnFilteChanged = false;
-		for (const TSharedRef<SFilter>& Filter : Filters)
+		bool ExecuteOnFilterChanged = false;
+		for (const TSharedPtr<SFilter> Filter : Filters)
 		{
 			if (Filter->IsEnabled())
 			{
-				if (const TSharedPtr<FFrontendFilter>& FrontendFilter = Filter->GetFrontendFilter())
+				if (const TSharedPtr<FFilterBase<FAssetFilterType>>& FrontendFilter = Filter->GetFrontendFilter())
 				{
 					for (const FContentBrowserItem& Item : ItemList)
 					{
@@ -839,24 +223,31 @@ void SFilterList::DisableFiltersThatHideItems(TArrayView<const FContentBrowserIt
 							// This is a frontend filter and at least one asset did not pass.
 							Filter->SetEnabled(false, false);
 							SetFrontendFilterActive(FrontendFilter.ToSharedRef(), false);
-							ExecuteOnFilteChanged = true;
+							ExecuteOnFilterChanged = true;
 						}
 					}
 				}
+			}
+		}
 
-				if (bDisableAllBackendFilters)
+		// Disable all backend filters if it was determined that the combined backend filter hides any of the assets
+		if (bDisableAllBackendFilters)
+		{
+			for(const TSharedPtr<SAssetFilter> AssetFilter : AssetFilters)
+			{
+				if(AssetFilter.IsValid())
 				{
-					FARFilter BackendFilter = Filter->GetBackendFilter();
+					FARFilter BackendFilter = AssetFilter->GetBackendFilter();
 					if (!BackendFilter.IsEmpty())
 					{
-						Filter->SetEnabled(false, false);
-						ExecuteOnFilteChanged = true;
+						AssetFilter->SetEnabled(false, false);
+						ExecuteOnFilterChanged = true;
 					}
 				}
 			}
 		}
 
-		if (ExecuteOnFilteChanged)
+		if (ExecuteOnFilterChanged)
 		{
 			OnFilterChanged.ExecuteIfBound();
 		}
@@ -869,53 +260,52 @@ void SFilterList::SaveSettings(const FString& IniFilename, const FString& IniSec
 	FString EnabledTypeFilterString;
 	FString ActiveFrontendFilterString;
 	FString EnabledFrontendFilterString;
-	for ( auto FilterIt = Filters.CreateConstIterator(); FilterIt; ++FilterIt )
+	for ( const TSharedPtr<SFilter> Filter : this->Filters )
 	{
-		const TSharedRef<SFilter>& Filter = *FilterIt;
+		const FString FilterName = Filter->GetFilterName();
 
-		if ( Filter->GetAssetTypeActions().IsValid() )
+		// If it is a FrontendFilter
+		if ( Filter->GetFrontendFilter().IsValid() )
 		{
-			if ( ActiveTypeFilterString.Len() > 0 )
+			if ( ActiveFrontendFilterString.Len() > 0 )
 			{
-				ActiveTypeFilterString += TEXT(",");
+				ActiveFrontendFilterString += TEXT(",");
 			}
-
-			const FString FilterName = Filter->GetAssetTypeActions().Pin()->GetFilterName().ToString();
-			ActiveTypeFilterString += FilterName;
-
+			ActiveFrontendFilterString += FilterName;
+		
 			if ( Filter->IsEnabled() )
 			{
-				if ( EnabledTypeFilterString.Len() > 0 )
+				if ( EnabledFrontendFilterString.Len() > 0 )
 				{
-					EnabledTypeFilterString += TEXT(",");
+					EnabledFrontendFilterString += TEXT(",");
 				}
 
-				EnabledTypeFilterString += FilterName;
+				EnabledFrontendFilterString += FilterName;
 			}
+
+			const TSharedPtr<FFilterBase<FAssetFilterType>>& FrontendFilter = Filter->GetFrontendFilter();
+			const FString CustomSettingsString = FString::Printf(TEXT("%s.CustomSettings.%s"), *SettingsString, *FilterName);
+			FrontendFilter->SaveSettings(IniFilename, IniSection, CustomSettingsString);
+
 		}
-		else if ( Filter->GetFrontendFilter().IsValid() )
+		// Otherwise we assume it is a type filter
+		else
 		{
-			const TSharedPtr<FFrontendFilter>& FrontendFilter = Filter->GetFrontendFilter();
-			if ( ActiveFrontendFilterString.Len() > 0 )
+			if ( ActiveTypeFilterString.Len() > 0 )
 			{
-				ActiveFrontendFilterString += TEXT(",");
+				ActiveTypeFilterString += TEXT(",");
 			}
-
-			const FString FilterName = FrontendFilter->GetName();
-			ActiveFrontendFilterString += FilterName;
-
+			ActiveTypeFilterString += FilterName;
+		
 			if ( Filter->IsEnabled() )
 			{
-				if ( EnabledFrontendFilterString.Len() > 0 )
+				if ( EnabledTypeFilterString.Len() > 0 )
 				{
-					EnabledFrontendFilterString += TEXT(",");
+					EnabledTypeFilterString += TEXT(",");
 				}
 
-				EnabledFrontendFilterString += FilterName;
+				EnabledTypeFilterString += FilterName;
 			}
-
-			const FString CustomSettingsString = FString::Printf(TEXT("%s.CustomSettings.%s"), *SettingsString, *FilterName);
-			FrontendFilter->SaveSettings(IniFilename, IniSection, CustomSettingsString);
 		}
 	}
 
@@ -940,21 +330,14 @@ void SFilterList::LoadSettings(const FString& IniFilename, const FString& IniSec
 		ActiveTypeFilterString.ParseIntoArray(TypeFilterNames, TEXT(","), /*bCullEmpty=*/true);
 		EnabledTypeFilterString.ParseIntoArray(EnabledTypeFilterNames, TEXT(","), /*bCullEmpty=*/true);
 
-		// Get the list of all asset type actions
-		FAssetToolsModule& AssetToolsModule = FModuleManager::LoadModuleChecked<FAssetToolsModule>("AssetTools");
-		TArray< TWeakPtr<IAssetTypeActions> > AssetTypeActionsList;
-		AssetToolsModule.Get().GetAssetTypeActionsList(AssetTypeActionsList);
-
-		// For each TypeActions, add any that were active and enable any that were previously enabled
-		for ( auto TypeActionsIt = AssetTypeActionsList.CreateConstIterator(); TypeActionsIt; ++TypeActionsIt )
+		for(const TSharedRef<FCustomClassFilterData> &CustomClassFilter : CustomClassFilters)
 		{
-			const TWeakPtr<IAssetTypeActions>& TypeActions = *TypeActionsIt;
-			if ( TypeActions.IsValid() && TypeActions.Pin()->CanFilter() && !IsAssetTypeActionsInUse(TypeActions) )
+			if(!this->IsClassTypeInUse(CustomClassFilter))
 			{
-				const FString FilterName = TypeActions.Pin()->GetFilterName().ToString();
+				const FString FilterName = CustomClassFilter->GetFilterName();
 				if ( TypeFilterNames.Contains(FilterName) )
 				{
-					TSharedRef<SFilter> NewFilter = AddFilter(TypeActions);
+					TSharedRef<SFilter> NewFilter = AddAssetFilterToBar(CustomClassFilter);
 
 					if ( EnabledTypeFilterNames.Contains(FilterName) )
 					{
@@ -981,13 +364,13 @@ void SFilterList::LoadSettings(const FString& IniFilename, const FString& IniSec
 		// For each FrontendFilter, add any that were active and enable any that were previously enabled
 		for ( auto FrontendFilterIt = AllFrontendFilters.CreateIterator(); FrontendFilterIt; ++FrontendFilterIt )
 		{
-			TSharedRef<FFrontendFilter>& FrontendFilter = *FrontendFilterIt;
+			TSharedRef<FFilterBase<FAssetFilterType>>& FrontendFilter = *FrontendFilterIt;
 			const FString& FilterName = FrontendFilter->GetName();
 			if (!IsFrontendFilterInUse(FrontendFilter))
 			{
 				if ( FrontendFilterNames.Contains(FilterName) )
 				{
-					TSharedRef<SFilter> NewFilter = AddFilter(FrontendFilter);
+					TSharedRef<SFilter> NewFilter = AddFilterToBar(FrontendFilter);
 
 					if ( EnabledFrontendFilterNames.Contains(FilterName) )
 					{
@@ -1005,279 +388,19 @@ void SFilterList::LoadSettings(const FString& IniFilename, const FString& IniSec
 	OnFilterChanged.ExecuteIfBound();
 }
 
-void SFilterList::SetFrontendFilterCheckState(const TSharedPtr<FFrontendFilter>& InFrontendFilter, ECheckBoxState InCheckState)
+void SFilterList::SetFrontendFilterCheckState(const TSharedPtr<FFrontendFilter>& InFrontendFilter, ECheckBoxState CheckState)
 {
-	if (!InFrontendFilter || InCheckState == ECheckBoxState::Undetermined)
-	{
-		return;
-	}
-
-	// Check if the filter is already checked.
-	TSharedRef<FFrontendFilter> FrontendFilter = InFrontendFilter.ToSharedRef();
-	bool FrontendFilterChecked = IsFrontendFilterInUse(FrontendFilter);
-
-	if (InCheckState == ECheckBoxState::Checked && !FrontendFilterChecked)
-	{
-		AddFilter(FrontendFilter)->SetEnabled(true); // Pin a filter widget on the UI and activate the filter. Same behaviour as FrontendFilterClicked()
-	}
-	else if (InCheckState == ECheckBoxState::Unchecked && FrontendFilterChecked)
-	{
-		RemoveFilter(FrontendFilter); // Unpin the filter widget and deactivate the filter.
-	}
-	// else -> Already in the desired 'check' state.
+	this->SetFilterCheckState(InFrontendFilter, CheckState);
 }
 
 ECheckBoxState SFilterList::GetFrontendFilterCheckState(const TSharedPtr<FFrontendFilter>& InFrontendFilter) const
 {
-	return InFrontendFilter && IsFrontendFilterInUse(InFrontendFilter.ToSharedRef()) ? ECheckBoxState::Checked : ECheckBoxState::Unchecked;
+	return this->GetFilterCheckState(InFrontendFilter);
 }
 
 bool SFilterList::IsFrontendFilterActive(const TSharedPtr<FFrontendFilter>& InFrontendFilter) const
 {
-	if (InFrontendFilter.IsValid())
-	{
-		for (const TSharedRef<SFilter>& Filter : Filters)
-		{
-			if (InFrontendFilter == Filter->GetFrontendFilter())
-			{
-				return Filter->IsEnabled(); // Is active or not?
-			}
-		}
-	}
-	return false;
-}
-
-void SFilterList::SetFrontendFilterActive(const TSharedRef<FFrontendFilter>& Filter, bool bActive)
-{
-	if(Filter->IsInverseFilter())
-	{
-		//Inverse filters are active when they are "disabled"
-		bActive = !bActive;
-	}
-	Filter->ActiveStateChanged(bActive);
-
-	if ( bActive )
-	{
-		FrontendFilters->Add(Filter);
-	}
-	else
-	{
-		FrontendFilters->Remove(Filter);
-	}
-}
-
-TSharedRef<SFilter> SFilterList::AddFilter(const TWeakPtr<IAssetTypeActions>& AssetTypeActions)
-{
-	TSharedRef<SFilter> NewFilter =
-		SNew(SFilter)
-		.AssetTypeActions(AssetTypeActions)
-		.OnFilterChanged(OnFilterChanged)
-		.OnRequestRemove(this, &SFilterList::RemoveFilterAndUpdate)
-		.OnRequestEnableOnly(this, &SFilterList::EnableOnlyThisFilter)
-		.OnRequestEnableAll(this, &SFilterList::EnableAllFilters)
-		.OnRequestDisableAll(this, &SFilterList::DisableAllFilters)
-		.OnRequestRemoveAll(this, &SFilterList::RemoveAllFilters)
-		.OnRequestRemoveAllButThis(this, &SFilterList::RemoveAllButThis);
-
-	AddFilter( NewFilter );
-
-	return NewFilter;
-}
-
-TSharedRef<SFilter> SFilterList::AddFilter(const TSharedRef<FFrontendFilter>& FrontendFilter)
-{
-	TSharedRef<SFilter> NewFilter =
-		SNew(SFilter)
-		.FrontendFilter(FrontendFilter)
-		.OnFilterChanged( this, &SFilterList::FrontendFilterChanged, FrontendFilter )
-		.OnRequestRemove(this, &SFilterList::RemoveFilterAndUpdate)
-		.OnRequestEnableOnly(this, &SFilterList::EnableOnlyThisFilter)
-		.OnRequestEnableAll(this, &SFilterList::EnableAllFilters)
-		.OnRequestDisableAll(this, &SFilterList::DisableAllFilters)
-		.OnRequestRemoveAll(this, &SFilterList::RemoveAllFilters)
-		.OnRequestRemoveAllButThis(this, &SFilterList::RemoveAllButThis);
-
-	AddFilter( NewFilter );
-
-	return NewFilter;
-}
-
-void SFilterList::AddFilter(const TSharedRef<SFilter>& FilterToAdd)
-{
-	Filters.Add(FilterToAdd);
-	
-	ChildSlot
-	[
-		MakeFilterToolBarWidget(Filters)
-	];
-}
-
-void SFilterList::RemoveFilter(const TWeakPtr<IAssetTypeActions>& AssetTypeActions, bool ExecuteOnFilterChanged)
-{
-	TSharedPtr<SFilter> FilterToRemove;
-	for ( auto FilterIt = Filters.CreateConstIterator(); FilterIt; ++FilterIt )
-	{
-		const TWeakPtr<IAssetTypeActions>& Actions = (*FilterIt)->GetAssetTypeActions();
-		if ( Actions.IsValid() && Actions == AssetTypeActions)
-		{
-			FilterToRemove = *FilterIt;
-			break;
-		}
-	}
-
-	if ( FilterToRemove.IsValid() )
-	{
-		if (ExecuteOnFilterChanged)
-		{
-			RemoveFilterAndUpdate(FilterToRemove.ToSharedRef());
-		}
-		else
-		{
-			RemoveFilter(FilterToRemove.ToSharedRef());
-		}
-	}
-}
-
-void SFilterList::EnableOnlyThisFilter(const TSharedRef<SFilter>& FilterToEnable)
-{
-	for (const TSharedRef<SFilter>& Filter : Filters)
-	{
-		bool bEnable = Filter == FilterToEnable;
-		Filter->SetEnabled(bEnable, /*ExecuteOnFilterChange*/false);
-		if (const TSharedPtr<FFrontendFilter>& FrontendFilter = Filter->GetFrontendFilter())
-		{
-			SetFrontendFilterActive(FrontendFilter.ToSharedRef(), bEnable);
-		}
-	}
-
-	OnFilterChanged.ExecuteIfBound();
-}
-
-void SFilterList::RemoveFilter(const TSharedRef<FFrontendFilter>& FrontendFilter, bool ExecuteOnFilterChanged)
-{
-	TSharedPtr<SFilter> FilterToRemove;
-	for (const TSharedRef<SFilter>& Filter : Filters)
-	{
-		if (Filter->GetFrontendFilter() == FrontendFilter)
-		{
-			FilterToRemove = Filter;
-			break;
-		}
-	}
-
-	if (FilterToRemove.IsValid())
-	{
-		if (ExecuteOnFilterChanged)
-		{
-			RemoveFilterAndUpdate(FilterToRemove.ToSharedRef());
-		}
-		else
-		{
-			RemoveFilter(FilterToRemove.ToSharedRef());
-		}
-	}
-}
-
-void SFilterList::RemoveFilter(const TSharedRef<SFilter>& FilterToRemove)
-{
-	Filters.Remove(FilterToRemove);
-
-	if (const TSharedPtr<FFrontendFilter>& FrontendFilter = FilterToRemove->GetFrontendFilter()) // Is valid?
-	{
-		// Update the frontend filters collection
-		SetFrontendFilterActive(FrontendFilter.ToSharedRef(), false);
-		OnFilterChanged.ExecuteIfBound();
-	}
-
-	ChildSlot
-	[
-		MakeFilterToolBarWidget(Filters)
-	];
-}
-
-void SFilterList::RemoveFilterAndUpdate(const TSharedRef<SFilter>& FilterToRemove)
-{
-	RemoveFilter(FilterToRemove);
-
-	// Notify that a filter has changed
-	OnFilterChanged.ExecuteIfBound();
-}
-
-void SFilterList::FrontendFilterChanged(TSharedRef<FFrontendFilter> FrontendFilter)
-{
-	TSharedPtr<SFilter> FilterToUpdate;
-	for (const TSharedRef<SFilter>& Filter : Filters)
-	{
-		if (Filter->GetFrontendFilter() == FrontendFilter)
-		{
-			FilterToUpdate = Filter;
-			break;
-		}
-	}
-
-	if (FilterToUpdate.IsValid())
-	{
-		SetFrontendFilterActive(FrontendFilter, FilterToUpdate->IsEnabled());
-		OnFilterChanged.ExecuteIfBound();
-	}
-}
-
-void SFilterList::CreateFiltersMenuCategory(FToolMenuSection& Section, const TArray<TWeakPtr<IAssetTypeActions>> AssetTypeActionsList) const
-{
-	for (int32 ClassIdx = 0; ClassIdx < AssetTypeActionsList.Num(); ++ClassIdx)
-	{
-		const TWeakPtr<IAssetTypeActions>& WeakTypeActions = AssetTypeActionsList[ClassIdx];
-		if ( WeakTypeActions.IsValid() )
-		{
-			TSharedPtr<IAssetTypeActions> TypeActions = WeakTypeActions.Pin();
-			if ( TypeActions.IsValid() && TypeActions->CanFilter() )
-			{
-				const FText& LabelText = TypeActions->GetName();
-				Section.AddMenuEntry(
-					NAME_None,
-					LabelText,
-					FText::Format( LOCTEXT("FilterByTooltipPrefix", "Filter by {0}"), LabelText ),
-					FSlateIconFinder::FindIconForClass(TypeActions->GetSupportedClass()),
-					FUIAction(
-						FExecuteAction::CreateSP( const_cast<SFilterList*>(this), &SFilterList::FilterByTypeClicked, WeakTypeActions ),
-						FCanExecuteAction(),
-						FIsActionChecked::CreateSP(this, &SFilterList::IsAssetTypeActionsInUse, WeakTypeActions ) ),
-					EUserInterfaceActionType::ToggleButton
-					);
-			}
-		}
-	}
-}
-
-void SFilterList::CreateFiltersMenuCategory(UToolMenu* InMenu, const TArray<TWeakPtr<IAssetTypeActions>> AssetTypeActionsList) const
-{
-	CreateFiltersMenuCategory(InMenu->AddSection("Section"), AssetTypeActionsList);
-}
-
-void SFilterList::CreateOtherFiltersMenuCategory(FToolMenuSection& Section, TSharedPtr<FFrontendFilterCategory> MenuCategory) const
-{
-	for (const TSharedRef<FFrontendFilter>& FrontendFilter : AllFrontendFilters)
-	{
-		if(FrontendFilter->GetCategory() == MenuCategory)
-		{
-			Section.AddMenuEntry(
-				NAME_None,
-				FrontendFilter->GetDisplayName(),
-				FrontendFilter->GetToolTipText(),
-				FSlateIcon(FAppStyle::GetAppStyleSetName(), FrontendFilter->GetIconName()),
-				FUIAction(
-				FExecuteAction::CreateSP(const_cast<SFilterList*>(this), &SFilterList::FrontendFilterClicked, FrontendFilter),
-				FCanExecuteAction(),
-				FIsActionChecked::CreateSP(this, &SFilterList::IsFrontendFilterInUse, FrontendFilter)),
-				EUserInterfaceActionType::ToggleButton
-				);
-		}
-	}
-}
-
-void SFilterList::CreateOtherFiltersMenuCategory(UToolMenu* InMenu, TSharedPtr<FFrontendFilterCategory> MenuCategory) const
-{
-	CreateOtherFiltersMenuCategory(InMenu->AddSection("Section"), MenuCategory);
+	return this->IsFilterActive(InFrontendFilter);
 }
 
 bool IsFilteredByPicker(const TArray<UClass*>& FilterClassList, UClass* TestClass)
@@ -1296,7 +419,7 @@ bool IsFilteredByPicker(const TArray<UClass*>& FilterClassList, UClass* TestClas
 	return true;
 }
 
-void SFilterList::PopulateAddFilterMenu(UToolMenu* Menu)
+void SFilterList::PopulateAddFilterMenu_Internal(UToolMenu* Menu)
 {
 	EAssetTypeCategories::Type MenuExpansion = EAssetTypeCategories::Basic;
 	if (UContentBrowserFilterListContext* Context = Menu->FindContext<UContentBrowserFilterListContext>())
@@ -1304,184 +427,41 @@ void SFilterList::PopulateAddFilterMenu(UToolMenu* Menu)
 		MenuExpansion = Context->MenuExpansion;
 	}
 
-	FAssetToolsModule& AssetToolsModule = FModuleManager::LoadModuleChecked<FAssetToolsModule>(TEXT("AssetTools"));
-
-	// A local struct to describe a category in the filter menu
-	struct FCategoryMenu
-	{
-		FText Name;
-		FText Tooltip;
-		TArray<TWeakPtr<IAssetTypeActions>> Assets;
-
-		//Menu section
-		FName SectionExtensionHook;
-		FText SectionHeading;
-
-		FCategoryMenu(const FText& InName, const FText& InTooltip, const FName& InSectionExtensionHook, const FText& InSectionHeading)
-			: Name(InName)
-			, Tooltip(InTooltip)
-			, Assets()
-			, SectionExtensionHook(InSectionExtensionHook)
-			, SectionHeading(InSectionHeading)
-		{}
-	};
-
-	// Create a map of Categories to Menus
-	TMap<EAssetTypeCategories::Type, FCategoryMenu> CategoryToMenuMap;
-
-	// Add the Basic category
-	CategoryToMenuMap.Add(EAssetTypeCategories::Basic, FCategoryMenu( LOCTEXT("BasicFilter", "Basic"), LOCTEXT("BasicFilterTooltip", "Filter by basic assets."), "ContentBrowserFilterBasicAsset", LOCTEXT("BasicAssetsMenuHeading", "Basic Assets") ) );
-
-	// Add the advanced categories
-	TArray<FAdvancedAssetCategory> AdvancedAssetCategories;
-	AssetToolsModule.Get().GetAllAdvancedAssetCategories(/*out*/ AdvancedAssetCategories);
-
-	for (const FAdvancedAssetCategory& AdvancedAssetCategory : AdvancedAssetCategories)
+	this->PopulateAddFilterMenu(Menu, AssetFilterCategories.FindChecked(EAssetTypeCategories::Basic), FOnFilterAssetType::CreateLambda([this](UClass *TestClass)
 	{
-		const FName ExtensionPoint = NAME_None;
-		const FText SectionHeading = FText::Format(LOCTEXT("WildcardFilterHeadingHeadingTooltip", "{0} Assets."), AdvancedAssetCategory.CategoryName);
-		const FText Tooltip = FText::Format(LOCTEXT("WildcardFilterTooltip", "Filter by {0}."), SectionHeading);
-		CategoryToMenuMap.Add(AdvancedAssetCategory.CategoryType, FCategoryMenu(AdvancedAssetCategory.CategoryName, Tooltip, ExtensionPoint, SectionHeading));
-	}
+		return !IsFilteredByPicker(this->InitialClassFilters, TestClass);
+	}));
 
-	// Get the browser type maps
-	TArray<TWeakPtr<IAssetTypeActions>> AssetTypeActionsList;
-	AssetToolsModule.Get().GetAssetTypeActionsList(AssetTypeActionsList);
+	Menu->AddSection("ContentBrowserFilterMiscAsset", LOCTEXT("MiscAssetsMenuHeading", "Misc Options") );
+}
 
-	// Sort the list
-	struct FCompareIAssetTypeActions
+TSharedRef<SWidget> SFilterList::MakeAddFilterMenu()
+{
+	const FName FilterMenuName = "ContentBrowser.FilterMenu";
+	if (!UToolMenus::Get()->IsMenuRegistered(FilterMenuName))
 	{
-		FORCEINLINE bool operator()( const TWeakPtr<IAssetTypeActions>& A, const TWeakPtr<IAssetTypeActions>& B ) const
-		{
-			return A.Pin()->GetName().CompareTo( B.Pin()->GetName() ) == -1;
-		}
-	};
-	AssetTypeActionsList.Sort( FCompareIAssetTypeActions() );
-
-	const TSharedRef<FNamePermissionList>& AssetClassPermissionList = AssetToolsModule.Get().GetAssetClassPermissionList(EAssetClassAction::CreateAsset);
+		UToolMenu* Menu = UToolMenus::Get()->RegisterMenu(FilterMenuName);
+		Menu->bShouldCloseWindowAfterMenuSelection = true;
+		Menu->bCloseSelfOnly = true;
 
-	// For every asset type, move it into all the categories it should appear in
-	for (int32 ClassIdx = 0; ClassIdx < AssetTypeActionsList.Num(); ++ClassIdx)
-	{
-		const TWeakPtr<IAssetTypeActions>& WeakTypeActions = AssetTypeActionsList[ClassIdx];
-		if ( WeakTypeActions.IsValid() )
+		Menu->AddDynamicSection(NAME_None, FNewToolMenuDelegate::CreateLambda([](UToolMenu* InMenu)
 		{
-			TSharedPtr<IAssetTypeActions> TypeActions = WeakTypeActions.Pin();
-			if ( ensure(TypeActions.IsValid()) && TypeActions->CanFilter() )
+			if (UContentBrowserFilterListContext* Context = InMenu->FindContext<UContentBrowserFilterListContext>())
 			{
-				UClass* SupportedClass = TypeActions->GetSupportedClass();
-				if ((!SupportedClass || AssetClassPermissionList->PassesFilter(SupportedClass->GetFName())) && !IsFilteredByPicker(InitialClassFilters, SupportedClass))
+				if (TSharedPtr<SFilterList> FilterList = Context->FilterList.Pin())
 				{
-					for ( auto MenuIt = CategoryToMenuMap.CreateIterator(); MenuIt; ++MenuIt )
-					{
-						if ( TypeActions->GetCategories() & MenuIt.Key() )
-						{
-							// This is a valid asset type which can be filtered, add it to the correct category
-							FCategoryMenu& CategoryMenu = MenuIt.Value();
-							CategoryMenu.Assets.Add( WeakTypeActions );
-						}
-					}
+					FilterList->PopulateAddFilterMenu_Internal(InMenu);
 				}
 			}
-		}
-	}
-
-	for (auto MenuIt = CategoryToMenuMap.CreateIterator(); MenuIt; ++MenuIt)
-	{
-		if (MenuIt.Value().Assets.Num() == 0)
-		{
-			CategoryToMenuMap.Remove(MenuIt.Key());
-		}
-	}
-
-	{
-		FToolMenuSection& Section = Menu->AddSection("ContentBrowserResetFilters");
-		Section.AddMenuEntry(
-			"ResetFilters",
-			LOCTEXT("FilterListResetFilters", "Reset Filters"),
-			LOCTEXT("FilterListResetToolTip", "Resets current filter selection"),
-			FSlateIcon(FAppStyle::Get().GetStyleSetName(), "PropertyWindow.DiffersFromDefault"),
-			FUIAction(
-				FExecuteAction::CreateSP(this, &SFilterList::OnResetFilters),
-				FCanExecuteAction::CreateLambda([this]() { return HasAnyFilters(); }))
-		);
-	}
-
-	// First add the expanded category, this appears as standard entries in the list (Note: intentionally not using FindChecked here as removing it from the map later would cause the ref to be garbage)
-	FCategoryMenu* ExpandedCategory = CategoryToMenuMap.Find( MenuExpansion );
-	check( ExpandedCategory );
-
-	{
-		FToolMenuSection& Section = Menu->AddSection(ExpandedCategory->SectionExtensionHook, ExpandedCategory->SectionHeading);
-		if(MenuExpansion == EAssetTypeCategories::Basic)
-		{
-			// If we are doing a full menu (i.e expanding basic) we add a menu entry which toggles all other categories
-			Section.AddMenuEntry(
-				NAME_None,
-				ExpandedCategory->Name,
-				ExpandedCategory->Tooltip,
-				FSlateIcon(FAppStyle::Get().GetStyleSetName(), "PlacementBrowser.Icons.Basic"),
-				FUIAction(
-				FExecuteAction::CreateSP( this, &SFilterList::FilterByTypeCategoryClicked, MenuExpansion ),
-				FCanExecuteAction(),
-				FGetActionCheckState::CreateSP(this, &SFilterList::IsAssetTypeCategoryChecked, MenuExpansion ) ),
-				EUserInterfaceActionType::ToggleButton
-				);
-		}
-
-		// Now populate with all the basic assets
-		SFilterList::CreateFiltersMenuCategory( Section, ExpandedCategory->Assets);
+		}));
 	}
 
-	// Remove the basic category from the map now, as this is treated differently and is no longer needed.
-	ExpandedCategory = nullptr;
-	CategoryToMenuMap.Remove(EAssetTypeCategories::Basic);
-
-	// If we have expanded Basic, assume we are in full menu mode and add all the other categories
-	{
-		FToolMenuSection& Section = Menu->AddSection("ContentBrowserFilterAdvancedAsset", LOCTEXT("AdvancedAssetsMenuHeading", "Other Assets"));
-		if(MenuExpansion == EAssetTypeCategories::Basic)
-		{
-			// Sort by category name so that we add the submenus in alphabetical order
-			CategoryToMenuMap.ValueSort([](const FCategoryMenu& A, const FCategoryMenu& B) {
-				return A.Name.CompareTo(B.Name) < 0;
-			});
-
-			// For all the remaining categories, add them as submenus
-			for (const TPair<EAssetTypeCategories::Type, FCategoryMenu>& CategoryMenuPair : CategoryToMenuMap)
-			{
-				Section.AddSubMenu(
-					NAME_None,
-					CategoryMenuPair.Value.Name,
-					CategoryMenuPair.Value.Tooltip,
-					FNewToolMenuDelegate::CreateSP(this, &SFilterList::CreateFiltersMenuCategory, CategoryMenuPair.Value.Assets),
-					FUIAction(
-					FExecuteAction::CreateSP(this, &SFilterList::FilterByTypeCategoryClicked, CategoryMenuPair.Key),
-					FCanExecuteAction(),
-					FGetActionCheckState::CreateSP(this, &SFilterList::IsAssetTypeCategoryChecked, CategoryMenuPair.Key)),
-					EUserInterfaceActionType::ToggleButton
-					);
-			}
-		}
-
-		// Now add the other filter which aren't assets
-		for (const TSharedPtr<FFrontendFilterCategory>& Category : AllFrontendFilterCategories)
-		{
-			Section.AddSubMenu(
-				NAME_None,
-				Category->Title,
-				Category->Tooltip,
-				FNewToolMenuDelegate::CreateSP(this, &SFilterList::CreateOtherFiltersMenuCategory, Category),
-				FUIAction(
-				FExecuteAction::CreateSP( this, &SFilterList::FrontendFilterCategoryClicked, Category ),
-				FCanExecuteAction(),
-				FIsActionChecked::CreateSP(this, &SFilterList::IsFrontendFilterCategoryInUse, Category ) ),
-				EUserInterfaceActionType::ToggleButton
-				);
-		}
-	}
+	UContentBrowserFilterListContext* ContentBrowserFilterListContext = NewObject<UContentBrowserFilterListContext>();
+	ContentBrowserFilterListContext->FilterList = SharedThis(this);
+	ContentBrowserFilterListContext->MenuExpansion = EAssetTypeCategories::Basic;
+	FToolMenuContext ToolMenuContext(ContentBrowserFilterListContext);
 
-	Menu->AddSection("ContentBrowserFilterMiscAsset", LOCTEXT("MiscAssetsMenuHeading", "Misc Options") );
+	return UToolMenus::Get()->GenerateWidget(FilterMenuName, ToolMenuContext);
 }
 
 TSharedRef<SWidget> SFilterList::MakeAddFilterMenu(EAssetTypeCategories::Type MenuExpansion)
@@ -1499,7 +479,7 @@ TSharedRef<SWidget> SFilterList::MakeAddFilterMenu(EAssetTypeCategories::Type Me
 			{
 				if (TSharedPtr<SFilterList> FilterList = Context->FilterList.Pin())
 				{
-					FilterList->PopulateAddFilterMenu(InMenu);
+					FilterList->PopulateAddFilterMenu_Internal(InMenu);
 				}
 			}
 		}));
@@ -1513,245 +493,4 @@ TSharedRef<SWidget> SFilterList::MakeAddFilterMenu(EAssetTypeCategories::Type Me
 	return UToolMenus::Get()->GenerateWidget(FilterMenuName, ToolMenuContext);
 }
 
-void SFilterList::FilterByTypeClicked(TWeakPtr<IAssetTypeActions> AssetTypeActions)
-{
-	if (AssetTypeActions.IsValid())
-	{
-		if (IsAssetTypeActionsInUse(AssetTypeActions))
-		{
-			RemoveFilter(AssetTypeActions);
-		}
-		else
-		{
-			TSharedRef<SFilter> NewFilter = AddFilter(AssetTypeActions);
-			NewFilter->SetEnabled(true);
-		}
-	}
-}
-
-bool SFilterList::IsAssetTypeActionsInUse(TWeakPtr<IAssetTypeActions> AssetTypeActions) const
-{
-	if (!AssetTypeActions.IsValid())
-	{
-		return false;
-	}
-
-	TSharedPtr<IAssetTypeActions> TypeActions = AssetTypeActions.Pin();
-	if (!TypeActions.IsValid())
-	{
-		return false;
-	}
-
-	for (const TSharedRef<SFilter>& Filter : Filters)
-	{
-		if (Filter->GetAssetTypeActions().Pin() == TypeActions)
-		{
-			return true;
-		}
-	}
-
-	return false;
-}
-
-void SFilterList::FilterByTypeCategoryClicked(EAssetTypeCategories::Type Category)
-{
-	TArray<TWeakPtr<IAssetTypeActions>> TypeActionsList;
-	GetTypeActionsForCategory(Category, TypeActionsList);
-
-	// Sort the list of type actions so that we add new filters in alphabetical order
-	TypeActionsList.Sort([](const TWeakPtr<IAssetTypeActions>& A, const TWeakPtr<IAssetTypeActions>& B) {
-		const FText NameA = A.IsValid() ? A.Pin()->GetName() : FText::GetEmpty();
-		const FText NameB = B.IsValid() ? B.Pin()->GetName() : FText::GetEmpty();
-		return NameA.CompareTo(NameB) < 0;
-	});
-
-	bool bFullCategoryInUse = IsAssetTypeCategoryInUse(Category);
-	bool ExecuteOnFilterChanged = false;
-
-	for (const TWeakPtr<IAssetTypeActions>& AssetTypeActions : TypeActionsList)
-	{
-		if (AssetTypeActions.IsValid())
-		{
-			if (bFullCategoryInUse)
-			{
-				RemoveFilter(AssetTypeActions);
-				ExecuteOnFilterChanged = true;
-			}
-			else if (!IsAssetTypeActionsInUse(AssetTypeActions))
-			{
-				TSharedRef<SFilter> NewFilter = AddFilter(AssetTypeActions);
-				NewFilter->SetEnabled(true, false);
-				ExecuteOnFilterChanged = true;
-			}
-		}
-	}
-
-	if (ExecuteOnFilterChanged)
-	{
-		OnFilterChanged.ExecuteIfBound();
-	}
-}
-
-bool SFilterList::IsAssetTypeCategoryInUse(EAssetTypeCategories::Type Category) const
-{
-	ECheckBoxState AssetTypeCategoryCheckState = IsAssetTypeCategoryChecked(Category);
-
-	if (AssetTypeCategoryCheckState == ECheckBoxState::Unchecked)
-	{
-		return false;
-	}
-
-	// An asset type category is in use if any of its type actions are in use (ECheckBoxState::Checked or ECheckBoxState::Undetermined)
-	return true;
-}
-
-ECheckBoxState SFilterList::IsAssetTypeCategoryChecked(EAssetTypeCategories::Type Category) const
-{
-	TArray<TWeakPtr<IAssetTypeActions>> TypeActionsList;
-	GetTypeActionsForCategory(Category, TypeActionsList);
-
-	bool bIsAnyActionInUse = false;
-	bool bIsAnyActionNotInUse = false;
-
-	for (const TWeakPtr<IAssetTypeActions>& AssetTypeActions : TypeActionsList)
-	{
-		if (AssetTypeActions.IsValid())
-		{
-			if (IsAssetTypeActionsInUse(AssetTypeActions))
-			{
-				bIsAnyActionInUse = true;
-			}
-			else
-			{
-				bIsAnyActionNotInUse = true;
-			}
-
-			if (bIsAnyActionInUse && bIsAnyActionNotInUse)
-			{
-				return ECheckBoxState::Undetermined;
-			}
-		}
-	}
-
-	if (bIsAnyActionInUse)
-	{
-		return ECheckBoxState::Checked;
-	}
-	else
-	{
-		return ECheckBoxState::Unchecked;
-	}
-}
-
-void SFilterList::GetTypeActionsForCategory(EAssetTypeCategories::Type Category, TArray< TWeakPtr<IAssetTypeActions> >& TypeActions) const
-{
-	// Load the asset tools module
-	FAssetToolsModule& AssetToolsModule = FModuleManager::LoadModuleChecked<FAssetToolsModule>(TEXT("AssetTools"));
-	TArray<TWeakPtr<IAssetTypeActions>> AssetTypeActionsList;
-	AssetToolsModule.Get().GetAssetTypeActionsList(AssetTypeActionsList);
-	const TSharedRef<FNamePermissionList>& AssetClassPermissionList = AssetToolsModule.Get().GetAssetClassPermissionList(EAssetClassAction::ViewAsset);
-
-	// Find all asset type actions that match the category
-	for (int32 ClassIdx = 0; ClassIdx < AssetTypeActionsList.Num(); ++ClassIdx)
-	{
-		const TWeakPtr<IAssetTypeActions>& WeakTypeActions = AssetTypeActionsList[ClassIdx];
-		TSharedPtr<IAssetTypeActions> AssetTypeActions = WeakTypeActions.Pin();
-
-		if (ensure(AssetTypeActions.IsValid()) && AssetTypeActions->CanFilter() && AssetTypeActions->GetCategories() & Category)
-		{
-			if (AssetTypeActions->GetSupportedClass() == nullptr || AssetClassPermissionList->PassesFilter(AssetTypeActions->GetSupportedClass()->GetFName()))
-			{
-				TypeActions.Add(WeakTypeActions);
-			}
-		}
-	}
-}
-
-void SFilterList::FrontendFilterClicked(TSharedRef<FFrontendFilter> FrontendFilter)
-{
-	if (IsFrontendFilterInUse(FrontendFilter))
-	{
-		RemoveFilter(FrontendFilter);
-	}
-	else
-	{
-		TSharedRef<SFilter> NewFilter = AddFilter(FrontendFilter);
-		NewFilter->SetEnabled(true);
-	}
-}
-
-bool SFilterList::IsFrontendFilterInUse(TSharedRef<FFrontendFilter> FrontendFilter) const
-{
-	for (const TSharedRef<SFilter>& Filter : Filters)
-	{
-		if (Filter->GetFrontendFilter() == FrontendFilter)
-		{
-			return true;
-		}
-	}
-
-	return false;
-}
-
-void SFilterList::FrontendFilterCategoryClicked(TSharedPtr<FFrontendFilterCategory> MenuCategory)
-{
-	bool bFullCategoryInUse = IsFrontendFilterCategoryInUse(MenuCategory);
-	bool ExecuteOnFilterChanged = false;
-
-	for (const TSharedRef<FFrontendFilter>& FrontendFilter : AllFrontendFilters)
-	{
-		if (FrontendFilter->GetCategory() == MenuCategory)
-		{
-			if (bFullCategoryInUse)
-			{
-				RemoveFilter(FrontendFilter, false);
-				ExecuteOnFilterChanged = true;
-			}
-			else if (!IsFrontendFilterInUse(FrontendFilter))
-			{
-				TSharedRef<SFilter> NewFilter = AddFilter(FrontendFilter);
-				NewFilter->SetEnabled(true, false);
-				SetFrontendFilterActive(FrontendFilter, NewFilter->IsEnabled());
-				ExecuteOnFilterChanged = true;
-			}
-		}
-	}
-
-	if (ExecuteOnFilterChanged)
-	{
-		OnFilterChanged.ExecuteIfBound();
-	}
-}
-
-bool SFilterList::IsFrontendFilterCategoryInUse(TSharedPtr<FFrontendFilterCategory> MenuCategory) const
-{
-	for (const TSharedRef<FFrontendFilter>& FrontendFilter : AllFrontendFilters)
-	{
-		if (FrontendFilter->GetCategory() == MenuCategory && !IsFrontendFilterInUse(FrontendFilter))
-		{
-			return false;
-		}
-	}
-
-	return true;
-}
-
-void SFilterList::OnResetFilters()
-{
-	RemoveAllFilters();
-}
-
-void SFilterList::OnSetFilterActive(bool bInActive, TWeakPtr<FFrontendFilter> InWeakFilter)
-{
-	TSharedPtr<FFrontendFilter> Filter = InWeakFilter.Pin();
-	if (Filter.IsValid())
-	{
-		if (!IsFrontendFilterInUse(Filter.ToSharedRef()))
-		{
-			TSharedRef<SFilter> NewFilter = AddFilter(Filter.ToSharedRef());
-			NewFilter->SetEnabled(bInActive);
-		}
-	}
-}
-
 #undef LOCTEXT_NAMESPACE
diff --git a/Engine/Source/Editor/ContentBrowser/Private/SFilterList.h b/Engine/Source/Editor/ContentBrowser/Private/SFilterList.h
--- a/Engine/Source/Editor/ContentBrowser/Private/SFilterList.h
+++ b/Engine/Source/Editor/ContentBrowser/Private/SFilterList.h
@@ -10,6 +10,7 @@
 #include "AssetRegistry/AssetData.h"
 #include "AssetRegistry/ARFilter.h"
 #include "FrontendFilterBase.h"
+#include "Filters/SAssetFilterBar.h"
 
 class UToolMenu;
 struct FToolMenuSection;
@@ -20,14 +21,14 @@ enum class ECheckBoxState : uint8;
 /**
  * A list of filters currently applied to an asset view.
  */
-class SFilterList : public SCompoundWidget
+class SFilterList : public SAssetFilterBar<FAssetFilterType>
 {
 public:
-	/** Delegate for when filters have changed */
-	DECLARE_DELEGATE( FOnFilterChanged );
 
 	DECLARE_DELEGATE_RetVal( TSharedPtr<SWidget>, FOnGetContextMenu );
 
+	using FOnFilterChanged = typename SAssetFilterBar<FAssetFilterType>::FOnFilterChanged;
+
 	SLATE_BEGIN_ARGS( SFilterList ){}
 
 		/** Called when an asset is right clicked */
@@ -50,11 +51,14 @@ public:
 	/** Constructs this widget with InArgs */
 	void Construct( const FArguments& InArgs );
 
-	/** Returns true if any filters are applied */
-	bool HasAnyFilters() const;
+	/** Set the check box state of the specified frontend filter (in the filter drop down) and pin/unpin a filter widget on/from the filter bar. When a filter is pinned (was not already pinned), it is activated and deactivated when unpinned. */
+	void SetFrontendFilterCheckState(const TSharedPtr<FFrontendFilter>& InFrontendFilter, ECheckBoxState CheckState);
 
-	/** Returns all of the filters combined */
-	FARFilter GetCombinedBackendFilter() const;
+	/** Returns the check box state of the specified frontend filter (in the filter drop down). This tells whether the filter is pinned or not on the filter bar, but not if filter is active or not. @see IsFrontendFilterActive(). */
+	ECheckBoxState GetFrontendFilterCheckState(const TSharedPtr<FFrontendFilter>& InFrontendFilter) const;
+
+	/** Returns true if the specified frontend filter is both checked (pinned on the filter bar) and active (contributing to filter the result). */
+	bool IsFrontendFilterActive(const TSharedPtr<FFrontendFilter>& InFrontendFilter) const;
 
 	/** Retrieve a specific frontend filter */
 	TSharedPtr<FFrontendFilter> GetFrontendFilter(const FString& InName) const;
@@ -62,115 +66,33 @@ public:
 	/** Handler for when the floating add filter button was clicked */
 	TSharedRef<SWidget> ExternalMakeAddFilterMenu(EAssetTypeCategories::Type MenuExpansion = EAssetTypeCategories::Basic);
 
-	/** Enables all filters */
-	void EnableAllFilters();
-
-	/** Disables any applied filters */
-	void DisableAllFilters();
-
-	/** Removes all filters in the list */
-	void RemoveAllFilters();
-
-	/** Removes all filters in the list except the given one. */
-	void RemoveAllButThis(const TSharedRef<SFilter>& FilterToKeep);
-
 	/** Disables any active filters that would hide the supplied items */
 	void DisableFiltersThatHideItems(TArrayView<const FContentBrowserItem> ItemList);
 
 	/** Saves any settings to config that should be persistent between editor sessions */
-	void SaveSettings(const FString& IniFilename, const FString& IniSection, const FString& SettingsString) const;
+	void SaveSettings(const FString& IniFilename, const FString& IniSection, const FString& SettingsString) const override;
 
 	/** Loads any settings to config that should be persistent between editor sessions */
-	void LoadSettings(const FString& IniFilename, const FString& IniSection, const FString& SettingsString);
+	void LoadSettings(const FString& IniFilename, const FString& IniSection, const FString& SettingsString) override;
 
 	virtual FReply OnMouseButtonUp( const FGeometry& MyGeometry, const FPointerEvent& MouseEvent ) override;
 
 	/** Returns the class filters specified at construction using argument 'InitialClassFilters'. */
 	const TArray<UClass*>& GetInitialClassFilters();
 
-	/** Set the check box state of the specified frontend filter (in the filter drop down) and pin/unpin a filter widget on/from the filter bar. When a filter is pinned (was not already pinned), it is activated and deactivated when unpinned. */
-	void SetFrontendFilterCheckState(const TSharedPtr<FFrontendFilter>& InFrontendFilter, ECheckBoxState CheckState);
-
-	/** Returns the check box state of the specified frontend filter (in the filter drop down). This tells whether the filter is pinned or not on the filter bar, but not if filter is active or not. @see IsFrontendFilterActive(). */
-	ECheckBoxState GetFrontendFilterCheckState(const TSharedPtr<FFrontendFilter>& InFrontendFilter) const;
-
-	/** Returns true if the specified frontend filter is both checked (pinned on the filter bar) and active (contributing to filter the result). */
-	bool IsFrontendFilterActive(const TSharedPtr<FFrontendFilter>& InFrontendFilter) const;
+protected:
+	
+	/** Handler for when the add filter button was clicked */
+	TSharedRef<SWidget> MakeAddFilterMenu() override;
 
 private:
-	/** Sets the active state of a frontend filter. */
-	void SetFrontendFilterActive(const TSharedRef<FFrontendFilter>& Filter, bool bActive);
-
-	/** Adds a filter to the end of the filter box. */
-	TSharedRef<SFilter> AddFilter(const TWeakPtr<IAssetTypeActions>& AssetTypeActions);
-	TSharedRef<SFilter> AddFilter(const TSharedRef<FFrontendFilter>& FrontendFilter);
-	void AddFilter(const TSharedRef<SFilter>& FilterToAdd);
-
-	/** Handler for when the remove filter button was clicked on a filter */
-	void RemoveFilter(const TWeakPtr<IAssetTypeActions>& AssetTypeActions, bool ExecuteOnFilterChanged = true);
-	void RemoveFilter(const TSharedRef<FFrontendFilter>& FrontendFilter, bool ExecuteOnFilterChanged = true);
-	void RemoveFilter(const TSharedRef<SFilter>& FilterToRemove);
-	void RemoveFilterAndUpdate(const TSharedRef<SFilter>& FilterToRemove);
-
-	/** Handler for when the enable only this button was clicked on a single filter */
-	void EnableOnlyThisFilter(const TSharedRef<SFilter>& FilterToEnable);
-
-	/** Handler for when a frontend filter state has changed */
-	void FrontendFilterChanged(TSharedRef<FFrontendFilter> FrontendFilter);
 
-	/** Handler for when the add filter menu is populated by a category */
-	void CreateFiltersMenuCategory(FToolMenuSection& Section, const TArray<TWeakPtr<IAssetTypeActions>> AssetTypeActionsList) const;
-	void CreateFiltersMenuCategory(UToolMenu* InMenu, const TArray<TWeakPtr<IAssetTypeActions>> AssetTypeActionsList) const;
-
-	/** Handler for when the add filter menu is populated by a non-category */
-	void CreateOtherFiltersMenuCategory(FToolMenuSection& Section, TSharedPtr<FFrontendFilterCategory> MenuCategory) const;
-	void CreateOtherFiltersMenuCategory(UToolMenu* InMenu, TSharedPtr<FFrontendFilterCategory> MenuCategory) const;
-
-	/** Handler for when the add filter button was clicked */
+	// Exists for backwards compatibility with ExternalMakeAddFilterMenu
 	TSharedRef<SWidget> MakeAddFilterMenu(EAssetTypeCategories::Type MenuExpansion = EAssetTypeCategories::Basic);
-	void PopulateAddFilterMenu(UToolMenu* Menu);
-
-	/** Handler for when filter by type is selected */
-	void FilterByTypeClicked(TWeakPtr<IAssetTypeActions> AssetTypeActions);
-
-	/** Handler to determine the "checked" state of an asset type in the filter dropdown */
-	bool IsAssetTypeActionsInUse(TWeakPtr<IAssetTypeActions> AssetTypeActions) const;
-
-	/** Handler for when filter by type category is selected */
-	void FilterByTypeCategoryClicked(EAssetTypeCategories::Type Category);
-
-	/** Handler to determine the "checked" state of an asset type category in the filter dropdown */
-	ECheckBoxState IsAssetTypeCategoryChecked(EAssetTypeCategories::Type Category) const;
-
-	/** Returns all the asset type actions objects for the specified category */
-	void GetTypeActionsForCategory(EAssetTypeCategories::Type Category, TArray< TWeakPtr<IAssetTypeActions> >& TypeActions) const;
-
-	/** Function to check if a given asset type category is in use */
-	bool IsAssetTypeCategoryInUse(EAssetTypeCategories::Type Category) const;
-
-	void FrontendFilterClicked(TSharedRef<FFrontendFilter> FrontendFilter);
-	bool IsFrontendFilterInUse(TSharedRef<FFrontendFilter> FrontendFilter) const;
-	void FrontendFilterCategoryClicked(TSharedPtr<FFrontendFilterCategory> MenuCategory);
-	bool IsFrontendFilterCategoryInUse(TSharedPtr<FFrontendFilterCategory> MenuCategory) const;
-
-	/** Called when reset filters option is pressed */
-	void OnResetFilters();
-
-	/** Called to set a filter active externally */
-	void OnSetFilterActive(bool bInActive, TWeakPtr<FFrontendFilter> InWeakFilter);
 
+	void PopulateAddFilterMenu_Internal(UToolMenu* Menu);
+	
 private:
-	/** All SFilters in the list */
-	TArray<TSharedRef<SFilter>> Filters;
-
-	/** The filter collection used to further filter down assets returned from the backend */
-	TSharedPtr<FAssetFilterCollectionType> FrontendFilters;
-
-	/** All possible frontend filter objects */
-	TArray< TSharedRef<FFrontendFilter> > AllFrontendFilters;
-
-	/** All frontend filter categories (for menu construction) */
-	TArray< TSharedPtr<FFrontendFilterCategory> > AllFrontendFilterCategories;
 
 	/** List of classes that our filters must match */
 	TArray<UClass*> InitialClassFilters;
@@ -181,5 +103,6 @@ private:
 	/** Delegate for when filters have changed */
 	FOnFilterChanged OnFilterChanged;
 
-	friend struct FFrontendFilterExternalActivationHelper;
+	/** A reference to AllFrontEndFilters so we can access the filters as FFrontEndFilter instead of FFilterBase<FAssetFilterType> */
+	TArray< TSharedRef<FFrontendFilter> > AllFrontendFilters_Internal;
 };
diff --git a/Engine/Source/Editor/ContentBrowser/Public/FrontendFilterBase.h b/Engine/Source/Editor/ContentBrowser/Public/FrontendFilterBase.h
--- a/Engine/Source/Editor/ContentBrowser/Public/FrontendFilterBase.h
+++ b/Engine/Source/Editor/ContentBrowser/Public/FrontendFilterBase.h
@@ -5,53 +5,41 @@
 #include "CoreMinimal.h"
 #include "Misc/IFilter.h"
 #include "IContentBrowserSingleton.h"
+#include "Filters/FilterBase.h"
 
 class FMenuBuilder;
 struct FContentBrowserDataFilter;
 
-class FFrontendFilterCategory
+class FFrontendFilterCategory : public FFilterCategory
 {
 public:
-	FFrontendFilterCategory(const FText& InTitle, const FText& InTooltip) : Title(InTitle), Tooltip(InTooltip) {}
-
-	/** The title of this category, used for the menu heading */
-	const FText Title;
-
-	/** The menu tooltip for this category */
-	const FText Tooltip;
+	FFrontendFilterCategory(const FText& InTitle, const FText& InTooltip) : FFilterCategory(InTitle,InTooltip) {}
 };
 
-class FFrontendFilter : public IFilter<FAssetFilterType>
+class FFrontendFilter : public FFilterBase<FAssetFilterType>
 {
 public:
-	FFrontendFilter(TSharedPtr<FFrontendFilterCategory> InCategory) : FilterCategory(InCategory) {}
-
-	/** Returns the system name for this filter */
-	virtual FString GetName() const = 0;
-
-	/** Returns the human readable name for this filter */
-	virtual FText GetDisplayName() const = 0;
-
-	/** Returns the tooltip for this filter, shown in the filters menu */
-	virtual FText GetToolTipText() const = 0;
+	FFrontendFilter(TSharedPtr<FFrontendFilterCategory> InCategory) : FFilterBase<FAssetFilterType>(InCategory) {}
 
+	/** Invoke to set the source filter that is currently used to filter assets in the asset view */
+	virtual void SetCurrentFilter(TArrayView<const FName> InSourcePaths, const FContentBrowserDataFilter& InBaseFilter) { }
+	
+	// FFilterBase Interface
+	
 	/** Returns the color this filter button will be when displayed as a button */
-	virtual FLinearColor GetColor() const { return FLinearColor(0.6f, 0.6f, 0.6f, 1); }
+	virtual FLinearColor GetColor() const override { return FLinearColor(0.6f, 0.6f, 0.6f, 1); }
 
 	/** Returns the name of the icon to use in menu entries */
-	virtual FName GetIconName() const { return NAME_None; }
+	virtual FName GetIconName() const override { return NAME_None; }
 
 	/** Returns true if the filter should be in the list when disabled and not in the list when enabled */
-	virtual bool IsInverseFilter() const { return false; }
-
-	/** Invoke to set the source filter that is currently used to filter assets in the asset view */
-	virtual void SetCurrentFilter(TArrayView<const FName> InSourcePaths, const FContentBrowserDataFilter& InBaseFilter) { }
-
+	virtual bool IsInverseFilter() const override { return false; }
+	
 	/** Notification that the filter became active or inactive */
-	virtual void ActiveStateChanged(bool bActive) { }
+	virtual void ActiveStateChanged(bool bActive) override { }
 
 	/** Called when the right-click context menu is being built for this filter */
-	virtual void ModifyContextMenu(FMenuBuilder& MenuBuilder) { }
+	virtual void ModifyContextMenu(FMenuBuilder& MenuBuilder) override { }
 
 	/** Called when the state of a particular Content Browser is being saved to INI */
 	virtual void SaveSettings(const FString& IniFilename, const FString& IniSection, const FString& SettingsString) const {}
@@ -59,26 +47,4 @@ public:
 	/** Called when the state of a particular Content Browser is being loaded from INI */
 	virtual void LoadSettings(const FString& IniFilename, const FString& IniSection, const FString& SettingsString) {}
 
-	/** Set this filter as active/inactive */
-	void SetActive(bool bInActive) { SetActiveEvent.Broadcast(bInActive); }
-
-	// IFilter implementation
-	DECLARE_DERIVED_EVENT( FFrontendFilter, IFilter<FAssetFilterType>::FChangedEvent, FChangedEvent );
-	virtual FChangedEvent& OnChanged() override { return ChangedEvent; }
-
-	TSharedPtr<FFrontendFilterCategory> GetCategory() { return FilterCategory; }
-
-protected:
-	void BroadcastChangedEvent() const { ChangedEvent.Broadcast(); }
-
-private:
-	FChangedEvent ChangedEvent;
-
-	/** This event is broadcast to set this filter active in the content browser it is being used in */
-	DECLARE_EVENT_OneParam(FFrontendFilter, FSetActiveEvent, bool);
-	FSetActiveEvent SetActiveEvent;
-
-	TSharedPtr<FFrontendFilterCategory> FilterCategory;
-
-	friend struct FFrontendFilterExternalActivationHelper;
 };
diff --git a/Engine/Source/Editor/EditorWidgets/EditorWidgets.Build.cs b/Engine/Source/Editor/EditorWidgets/EditorWidgets.Build.cs
--- a/Engine/Source/Editor/EditorWidgets/EditorWidgets.Build.cs
+++ b/Engine/Source/Editor/EditorWidgets/EditorWidgets.Build.cs
@@ -8,20 +8,20 @@ public class EditorWidgets : ModuleRules
 	{
 		PrivateIncludePathModuleNames.Add("AssetRegistry");
 
-		PrivateDependencyModuleNames.AddRange(
+		PublicDependencyModuleNames.AddRange(
 			new string[] {
 				"Core",
-				"CoreUObject",
 				"Engine",
-                "InputCore",
+				"InputCore",
 				"Slate",
 				"SlateCore",
 				"EditorFramework",
 				"UnrealEd",
-				"ToolWidgets"
+				"CoreUObject",
+				"ToolWidgets",
 			}
 		);
-
+		
 		DynamicallyLoadedModuleNames.Add("AssetRegistry");
 	}
-}
+}
\ No newline at end of file
diff --git a/Engine/Source/Editor/EditorWidgets/Private/Filters/CustomClassFilterData.cpp b/Engine/Source/Editor/EditorWidgets/Private/Filters/CustomClassFilterData.cpp
new file mode 100644
--- /dev/null
+++ b/Engine/Source/Editor/EditorWidgets/Private/Filters/CustomClassFilterData.cpp
@@ -0,0 +1,82 @@
+//Copyright Epic Games, Inc. All Rights Reserved.
+
+#include "Filters/CustomClassFilterData.h"
+#include "Filters/SBasicFilterBar.h"
+
+#include "IAssetTypeActions.h"
+#include "AssetTypeCategories.h"
+#include "AssetToolsModule.h"
+#include "UObject/Object.h"
+
+void FCustomClassFilterData::AddCategory(TSharedPtr<FFilterCategory> InCategory)
+{
+	Categories.Add(InCategory);
+}
+
+TSharedPtr<IAssetTypeActions> FCustomClassFilterData::GetAssetTypeActions() const
+{
+	return AssetTypeActions.Pin();
+}
+
+UClass* FCustomClassFilterData::GetClass() const
+{
+	if(TSharedPtr<IAssetTypeActions> AssetTypeActionsPin = AssetTypeActions.Pin())
+	{
+		return AssetTypeActionsPin->GetSupportedClass();
+	}
+
+	return Class;
+}
+
+TArray<TSharedPtr<FFilterCategory>> FCustomClassFilterData::GetCategories() const
+{
+	return Categories;
+}
+
+FLinearColor FCustomClassFilterData::GetColor() const
+{
+	if(TSharedPtr<IAssetTypeActions> AssetTypeActionsPin = AssetTypeActions.Pin())
+	{
+		return AssetTypeActionsPin->GetTypeColor();
+	}
+
+	return Color;
+}
+
+void FCustomClassFilterData::BuildBackendFilter(FARFilter &Filter)
+{
+	if ( AssetTypeActions.IsValid() )
+	{
+		if (AssetTypeActions.Pin()->CanFilter())
+		{
+			AssetTypeActions.Pin()->BuildBackendFilter(Filter);
+		}
+	}
+	// If there is no AssetTypeAction for this filter, simply add the class name to the FARFilter
+	else
+	{
+		Filter.ClassNames.Add(Class->GetFName());
+		Filter.bRecursiveClasses = true;
+	}
+}
+
+FText FCustomClassFilterData::GetName() const
+{
+	if(TSharedPtr<IAssetTypeActions> AssetTypeActionsPin = AssetTypeActions.Pin())
+	{
+		return AssetTypeActionsPin->GetName();
+	}
+
+	return Class->GetDisplayNameText();
+}
+
+FString FCustomClassFilterData::GetFilterName() const
+{
+	if(TSharedPtr<IAssetTypeActions> AssetTypeActionsPin = AssetTypeActions.Pin())
+	{
+		return AssetTypeActionsPin->GetFilterName().ToString();
+	}
+
+	return Class->GetFName().ToString();
+}
+
diff --git a/Engine/Source/Editor/EditorWidgets/Public/Filters/AssetFilter.h b/Engine/Source/Editor/EditorWidgets/Public/Filters/AssetFilter.h
new file mode 100644
--- /dev/null
+++ b/Engine/Source/Editor/EditorWidgets/Public/Filters/AssetFilter.h
@@ -0,0 +1,156 @@
+// Copyright Epic Games, Inc. All Rights Reserved.
+
+#pragma once
+
+#include "Filters/FilterBase.h"
+
+#include "Internationalization/Text.h"
+#include "Core/Public/Math/Color.h"
+#include "Core/Public/UObject/NameTypes.h"
+
+#include "ARFilter.h"
+#include "AssetRegistry/IAssetRegistry.h"
+
+/* A filter used by SFilterBar that wraps all the back-end filters into one FFilterBase for convenience */
+template<typename FilterType>
+class FAssetFilter : public FFilterBase<FilterType>
+{
+public:
+
+	/* A delegate that the user can use to specify how to compare their Filter Item to a list of Asset Types represented as FNames */
+	DECLARE_DELEGATE_RetVal_TwoParams(bool, FCompareItemWithClassNames, FilterType /* InItem */, const TSet<FName>& /* InClassNames */);
+
+	/* A delegate that the user can use to specify how to convert their Filter Item into an FAssetData
+	 * @return: true if the Filter Item was successfully converted, false if not (the current item will not be tested for asset filters if so)
+	 */
+	DECLARE_DELEGATE_RetVal_TwoParams(bool, FConvertItemToAssetData, FilterType /* InItem */, FAssetData& /* OutAssetData */);
+
+	
+	FAssetFilter() :
+		FFilterBase<FilterType>(nullptr)
+	{
+		
+	}
+
+	void SetBackendFilter(const FARFilter& InBackendFilter)
+	{
+		// Clear the current compiled filter
+		CompiledBackendFilter.Clear();
+
+		// Compile the filter and store it
+		IAssetRegistry::Get()->CompileFilter(InBackendFilter, CompiledBackendFilter);
+	}
+	
+	void SetComparisonFunction(FCompareItemWithClassNames InOnCompareItem)
+	{
+		OnCompareItem = InOnCompareItem;
+	}
+
+	void SetConversionFunction(FConvertItemToAssetData InOnConvertItem)
+	{
+		OnConvertItem = InOnConvertItem;
+	}
+
+	/** Returns the system name for this filter */
+	virtual FString GetName() const override
+	{
+		return FString("FAssetFilter");
+	}
+
+	/** Returns the human readable name for this filter */
+	virtual FText GetDisplayName() const override
+	{
+		return FText::GetEmpty(); // This filter should never be displayed in the Filter Bar
+	}
+
+
+	/** Returns the tooltip for this filter, shown in the filters menu */
+	virtual FText GetToolTipText() const override
+	{
+		return FText::GetEmpty(); // This filter should never be displayed in the Filter Bar
+	}
+
+
+	/** Returns the color this filter button will be when displayed as a button */
+	virtual FLinearColor GetColor() const override
+	{
+		return FLinearColor();
+	}
+
+	/** Returns the name of the icon to use in menu entries */
+	virtual FName GetIconName() const override
+	{
+		return NAME_None;
+	}
+
+	/** Returns true if the filter should be in the list when disabled and not in the list when enabled */
+	virtual bool IsInverseFilter() const override
+	{
+		return false;
+	}
+
+	virtual bool PassesFilter( FilterType InItem ) const override
+	{
+		 // The Conversion Function is preferred if it is specified
+		if(OnConvertItem.IsBound())
+		{
+			// Convert the item to an FAssetData
+			FAssetData AssetData;
+			bool bWasConverted = OnConvertItem.Execute(InItem, AssetData);
+
+			// Skip the current item if it wasn't converted
+			if(!bWasConverted)
+			{
+				return false;
+			}
+			
+			// Use the AssetRegistry to check if the Item passes the compiled filter
+			return IAssetRegistry::Get()->IsAssetIncludedByFilter(AssetData, CompiledBackendFilter);
+		}
+		// If not, try to use the comparison function
+		else if(OnCompareItem.IsBound())
+		{
+			return OnCompareItem.Execute(InItem, CompiledBackendFilter.ClassNames);
+		}
+
+		return false;
+	}
+
+	// Functionality that is not needed by this filter
+	
+	/** Notification that the filter became active or inactive */
+	virtual void ActiveStateChanged(bool bActive) override
+	{
+		
+	}
+
+	/** Called when the right-click context menu is being built for this filter */
+	virtual void ModifyContextMenu(FMenuBuilder& MenuBuilder) override
+	{
+		
+	}
+	
+	/** Can be overriden for custom FilterBar subclasses to save settings, currently not implemented in any gneeric Filter Bar */
+	virtual void SaveSettings(const FString& IniFilename, const FString& IniSection, const FString& SettingsString) const override
+	{
+		
+	}
+
+	/** Can be overriden for custom FilterBar subclasses to load settings, currently not implemented in any gneeric Filter Bar */
+	virtual void LoadSettings(const FString& IniFilename, const FString& IniSection, const FString& SettingsString) override
+	{
+		
+	}
+	
+protected:
+
+	/* The compiled BackendFilter that contains all the Asset Filters attached to this Filter */
+	FARCompiledFilter CompiledBackendFilter;
+
+	/* Delegate used to compare an item to a UClass represented by an FName */
+	FCompareItemWithClassNames OnCompareItem;
+
+	/* Delegate used to convert an item to an FAssetData */
+	FConvertItemToAssetData OnConvertItem;
+	
+};
\ No newline at end of file
diff --git a/Engine/Source/Editor/EditorWidgets/Public/Filters/CustomClassFilterData.h b/Engine/Source/Editor/EditorWidgets/Public/Filters/CustomClassFilterData.h
new file mode 100644
--- /dev/null
+++ b/Engine/Source/Editor/EditorWidgets/Public/Filters/CustomClassFilterData.h
@@ -0,0 +1,73 @@
+//Copyright Epic Games, Inc. All Rights Reserved.
+#pragma once
+
+#include "Math/Color.h"
+#include "ARFilter.h"
+
+class IAssetTypeActions;
+class UClass;
+class FFilterCategory;
+
+/** Helper Class that can be used to provide custom class filters to the Filter Widget
+ *  Using this for Type Filters ensures that the type filters get OR'd together
+ *  For example: filtering for StaticMeshActor and DirectionalLight will show assets that belong
+ *  to either of those types.
+ */
+class EDITORWIDGETS_API FCustomClassFilterData
+{
+public:
+
+	/* You can provide an IAssetTypeActions for your type to get all the information from that */
+	FCustomClassFilterData(TWeakPtr<IAssetTypeActions> InAssetTypeAction)
+	: Class(nullptr)
+	, AssetTypeActions(InAssetTypeAction)
+	{
+		
+	}
+
+	/* Or you can provide a UClass, a Category and a Color to identify the filter */
+	FCustomClassFilterData(UClass *InClass, TSharedPtr<FFilterCategory> InCategory, FLinearColor InColor)
+	: Class(InClass)
+	, Color(InColor)
+	{
+		Categories.Add(InCategory);
+	}
+
+	/* Unlike normal filters, Type Filters are allowed to belong to multiple categories */
+	void AddCategory(TSharedPtr<FFilterCategory> InCategory);
+
+	/* Get the AssetTypeActions associated with this filter if they exist */
+	TSharedPtr<IAssetTypeActions> GetAssetTypeActions() const;
+
+	/* Get the UClass associated with this filter */
+	UClass *GetClass() const;
+
+	/* Get the Catgories this filter belongs to */
+	TArray<TSharedPtr<FFilterCategory>> GetCategories() const;
+
+	/* Get the color of this filter */
+	FLinearColor GetColor() const;
+
+	/** Add this filter to the input BackendFilter */
+	void BuildBackendFilter(FARFilter &Filter);
+
+	/** Get the logical name of this filter class */
+	FText GetName() const;
+
+	/** Get the display name of this filter class */
+	FString GetFilterName() const;
+
+private:
+
+	/** The Class the filter is associated with (if it does not have an AssetTypeAction) */
+	UClass *Class;
+
+	/** The Categories the filter will show up under */
+	TArray<TSharedPtr<FFilterCategory>> Categories;
+
+	/** The Color of the filter (if it does not have an AssetTypeAction to get the color from) */
+	FLinearColor Color;
+
+	/** An optional AssetTypeAction to get information about the filter class */
+	TWeakPtr<IAssetTypeActions> AssetTypeActions;
+};
\ No newline at end of file
diff --git a/Engine/Source/Editor/EditorWidgets/Public/Filters/SAssetFilterBar.h b/Engine/Source/Editor/EditorWidgets/Public/Filters/SAssetFilterBar.h
new file mode 100644
--- /dev/null
+++ b/Engine/Source/Editor/EditorWidgets/Public/Filters/SAssetFilterBar.h
@@ -0,0 +1,839 @@
+//Copyright Epic Games, Inc. All Rights Reserved.
+
+#pragma once
+
+#include "Filters/SBasicFilterBar.h"
+#include "Filters/CustomClassFilterData.h"
+
+#include "IAssetTypeActions.h"
+#include "AssetTypeCategories.h"
+#include "ARFilter.h"
+#include "AssetToolsModule.h"
+#include "UObject/Object.h"
+#include "SlateGlobals.h"
+#include "Logging/LogMacros.h"
+
+#include "SAssetFilterBar.generated.h"
+
+#define LOCTEXT_NAMESPACE "FilterBar"
+
+/** Delegate that subclasses can use to specify classes to not include in this filter
+ * Returning false for a class will prevent it from showing up in the add filter dropdown
+ */
+DECLARE_DELEGATE_RetVal_OneParam(bool, FOnFilterAssetType, UClass* /* Test Class */);
+
+/* Delegate used by SAssetFilterBar to populate the add filter menu */
+DECLARE_DELEGATE_ThreeParams(FOnPopulateAddAssetFilterMenu, UToolMenu*, TSharedPtr<FFilterCategory>, FOnFilterAssetType)
+
+/** ToolMenuContext that is used to create the Add Filter Menu */
+UCLASS()
+class EDITORWIDGETS_API UAssetFilterBarContext : public UObject
+{
+	GENERATED_BODY()
+
+public:
+	TSharedPtr<FFilterCategory> MenuExpansion;
+	FOnPopulateAddAssetFilterMenu PopulateFilterMenu;
+};
+
+
+/**
+* An Asset Filter Bar widget, which can be used to filter items of type [FilterType] given a list of custom filters along with built in support for AssetType filters
+* @see SFilterBar in EditorWidgets, which is a simplified version of this widget for most use cases which is probably what you want to use
+* @see SBasicFilterBar in ToolWidgets if you want a simple FilterBar without AssetType Filters (usable in non editor builds)
+* @see SFilterList in ContentBrowser for an extremely complex example of how to use this widget, though for most cases you probably want to use SFilterBar
+* NOTE: The filter functions create copies, so you want to use SAssetFilterBar<TSharedPtr<ItemType>> etc instead of SAssetFilterBar<ItemType> when possible
+* NOTE: You will need to also add "ToolWidgets" as a dependency to your module to use this widget
+* Sample Usage:
+*		SAssignNew(MyFilterBar, SAssetFilterBar<FText>)
+*		.OnFilterChanged() // A delegate for when the list of filters changes
+*		.CustomFilters() // An array of filters available to this FilterBar (@see FGenericFilter to create simple delegate based filters)
+*
+* Use the GetAllActiveFilters() and GetCombinedBackendFilter() functions to get all the custom and asset filters respectively
+* NOTE: GetCombinedBackendFilter returns an FARFilter, and it is on the user of this widget to compile it/use it to filter their items.
+*		If you want more straightforward filtering, look at SFilterBar which streamlines this
+* Use MakeAddFilterButton() to make the button that summons the dropdown showing all the filters
+* Sample Usage:
+*  void OnFilterChangedDelegate()
+*  {
+*		TSharedPtr< TFilterCollection<FilterType> > ActiveFilters = MyFilterBar->GetAllActiveFilters();
+*		FARFilter CombinedBackEndFilter = MyFilterBar->GetCombinedBackendFilter()
+*		TArray<FilterType> MyUnfilteredItems;
+*		TArray<FilterType> FilteredItems;
+*		
+*		for(FilterType& MyItem : MyUnfilteredItems)
+*		{
+*			if(CompileAndRunFARFilter(CombinedBackEndFilter, MyItem) && ActiveFilters.PassesAllFilters(MyItem))
+*			{
+*				FilteredItems.Add(MyItem);
+*			}
+*		}
+*  }
+*/
+template<typename FilterType>
+class SAssetFilterBar : public SBasicFilterBar<FilterType>
+{
+public:
+	
+	using FOnFilterChanged = typename SBasicFilterBar<FilterType>::FOnFilterChanged;
+	
+	SLATE_BEGIN_ARGS( SAssetFilterBar<FilterType> )
+		: _UseDefaultAssetFilters(true)
+		{
+		
+		}
+
+		/** Delegate for when filters have changed */
+		SLATE_EVENT( FOnFilterChanged, OnFilterChanged )
+	
+		/** Initial List of Custom Filters that will be added to the AddFilter Menu */
+		SLATE_ARGUMENT( TArray<TSharedRef<FFilterBase<FilterType>>>, CustomFilters)
+
+		SLATE_ARGUMENT( TArray<TSharedRef<FCustomClassFilterData>>, CustomClassFilters)
+
+		/** Whether the filter bar should provide the default asset filters */
+		SLATE_ARGUMENT(bool, UseDefaultAssetFilters)
+	
+	SLATE_END_ARGS()
+
+	/** Constructs this widget with InArgs */
+	void Construct( const FArguments& InArgs )
+	{
+		bUseDefaultAssetFilters = InArgs._UseDefaultAssetFilters;
+		CustomClassFilters = InArgs._CustomClassFilters;
+		
+		typename SBasicFilterBar<FilterType>::FArguments Args;
+		Args._OnFilterChanged = InArgs._OnFilterChanged;
+		Args._CustomFilters = InArgs._CustomFilters;
+		
+		SBasicFilterBar<FilterType>::Construct(Args);
+		
+		CreateAssetTypeActionFilters();
+	}
+
+protected:
+
+	typedef typename SBasicFilterBar<FilterType>::SFilter SFilter;
+
+	/** A subclass of SFilter in SBasicFilterBar to add functionality for Asset Filters */
+	class SAssetFilter : public SFilter
+	{
+		using FOnRequestRemove = typename SBasicFilterBar<FilterType>::SFilter::FOnRequestRemove;
+		using FOnRequestEnableOnly = typename SBasicFilterBar<FilterType>::SFilter::FOnRequestEnableOnly;
+		using FOnRequestEnableAll = typename SBasicFilterBar<FilterType>::SFilter::FOnRequestEnableAll;
+		using FOnRequestDisableAll = typename SBasicFilterBar<FilterType>::SFilter::FOnRequestDisableAll;
+		using FOnRequestRemoveAll = typename SBasicFilterBar<FilterType>::SFilter::FOnRequestRemoveAll;
+		using FOnRequestRemoveAllButThis = typename SBasicFilterBar<FilterType>::SFilter::FOnRequestRemoveAllButThis;
+		
+		SLATE_BEGIN_ARGS( SAssetFilter ){}
+
+			/** The Custom Class Data that is associated with this filter */
+			SLATE_ARGUMENT( TSharedPtr<FCustomClassFilterData>, CustomClassFilter )
+
+			// SFilter Arguments
+		
+			/** If this is an front end filter, this is the filter object */
+			SLATE_ARGUMENT( TSharedPtr<FFilterBase<FilterType>>, FrontendFilter )
+
+			/** Invoked when the filter toggled */
+			SLATE_EVENT( FOnFilterChanged, OnFilterChanged )
+
+			/** Invoked when a request to remove this filter originated from within this filter */
+			SLATE_EVENT( FOnRequestRemove, OnRequestRemove )
+
+			/** Invoked when a request to enable only this filter originated from within this filter */
+			SLATE_EVENT( FOnRequestEnableOnly, OnRequestEnableOnly )
+
+			/** Invoked when a request to enable all filters originated from within this filter */
+			SLATE_EVENT( FOnRequestEnableAll, OnRequestEnableAll)
+
+			/** Invoked when a request to disable all filters originated from within this filter */
+			SLATE_EVENT( FOnRequestDisableAll, OnRequestDisableAll )
+
+			/** Invoked when a request to remove all filters originated from within this filter */
+			SLATE_EVENT( FOnRequestRemoveAll, OnRequestRemoveAll )
+
+			/** Invoked when a request to remove all filters originated from within this filter */
+			SLATE_EVENT( FOnRequestRemoveAllButThis, OnRequestRemoveAllButThis )
+
+		SLATE_END_ARGS()
+
+		/** Constructs this widget with InArgs */
+		void Construct( const FArguments& InArgs )
+		{
+			this->bEnabled = false;
+			this->OnFilterChanged = InArgs._OnFilterChanged;
+			this->OnRequestRemove = InArgs._OnRequestRemove;
+			this->OnRequestEnableOnly = InArgs._OnRequestEnableOnly;
+			this->OnRequestEnableAll = InArgs._OnRequestEnableAll;
+			this->OnRequestDisableAll = InArgs._OnRequestDisableAll;
+			this->OnRequestRemoveAll = InArgs._OnRequestRemoveAll;
+			this->OnRequestRemoveAllButThis = InArgs._OnRequestRemoveAllButThis;
+			this->FrontendFilter = InArgs._FrontendFilter;
+
+			CustomClassFilter = InArgs._CustomClassFilter;
+			
+			// Get the tooltip and color of the type represented by this filter
+			this->FilterColor = FLinearColor::White;
+			if ( CustomClassFilter.IsValid() )
+			{
+				this->FilterColor = CustomClassFilter->GetColor();
+				// No tooltip for asset type filters
+			}
+			else if ( this->FrontendFilter.IsValid() )
+			{
+				this->FilterColor = this->FrontendFilter->GetColor();
+				this->FilterToolTip = TAttribute<FText>::Create(TAttribute<FText>::FGetter::CreateSP(this->FrontendFilter.ToSharedRef(), &FFilterBase<FilterType>::GetToolTipText));
+			}
+
+			SFilter::Construct_Internal();
+		}
+
+	public:
+		
+		/** Returns this widgets contribution to the combined filter */
+		FARFilter GetBackendFilter() const
+		{
+			FARFilter Filter;
+
+			if(CustomClassFilter)
+			{
+				CustomClassFilter->BuildBackendFilter(Filter);
+			}
+
+			return Filter;
+		}
+		
+		/** Gets the asset type actions associated with this filter */
+		const TSharedPtr<FCustomClassFilterData>& GetCustomClassFilterData() const
+		{
+			return CustomClassFilter;
+		}
+
+		/** Returns the display name for this filter */
+		virtual FText GetFilterDisplayName() const override
+		{
+			if (CustomClassFilter.IsValid())
+			{
+				return CustomClassFilter->GetName();
+			}
+			else
+			{
+				return SFilter::GetFilterDisplayName();
+			}
+
+		}
+
+		virtual FString GetFilterName() const override
+		{
+			if (CustomClassFilter.IsValid())
+			{
+				return CustomClassFilter->GetFilterName();
+			}
+			else
+			{
+				return SFilter::GetFilterName();
+			}
+		}
+		
+	protected:
+		/** The asset type actions that are associated with this filter */
+		TSharedPtr<FCustomClassFilterData> CustomClassFilter;
+	};
+
+public:
+
+	/** Use this function to get an FARFilter that represents all the AssetType filters that are currently active */
+	FARFilter GetCombinedBackendFilter() const
+	{
+		FARFilter CombinedFilter;
+
+		// Add all selected filters
+		for (int32 FilterIdx = 0; FilterIdx < AssetFilters.Num(); ++FilterIdx)
+		{
+			const TSharedPtr<SAssetFilter> AssetFilter = AssetFilters[FilterIdx];
+			
+			if ( AssetFilter.IsValid() && AssetFilter->IsEnabled())
+			{
+				CombinedFilter.Append(AssetFilter->GetBackendFilter());
+			}
+		}
+
+		if ( CombinedFilter.bRecursiveClasses )
+		{
+			// Add exclusions for AssetTypeActions NOT in the filter.
+			// This will prevent assets from showing up that are both derived from an asset in the filter set and derived from an asset not in the filter set
+			// Get the list of all asset type actions
+			for(const TSharedRef<FCustomClassFilterData> &CustomClassFilter : CustomClassFilters)
+			{
+				const UClass* TypeClass = CustomClassFilter->GetClass();
+				if (TypeClass && !CombinedFilter.ClassNames.Contains(TypeClass->GetFName()))
+				{
+					CombinedFilter.RecursiveClassesExclusionSet.Add(TypeClass->GetFName());
+				}
+			}
+		}
+
+		// HACK: A blueprint can be shown as Blueprint or as BlueprintGeneratedClass, but we don't want to distinguish them while filtering.
+		// This should be removed, once all blueprints are shown as BlueprintGeneratedClass.
+		if(CombinedFilter.ClassNames.Contains(FName(TEXT("Blueprint"))))
+		{
+			CombinedFilter.ClassNames.AddUnique(FName(TEXT("BlueprintGeneratedClass")));
+		}
+
+		return CombinedFilter;
+	}
+
+	/** Check if there is a filter associated with the input class in the filter bar */
+	bool IsAssetTypeFilter(UClass *InClass)
+	{
+		for(const TSharedRef<FCustomClassFilterData>& CustomClassFilterData : CustomClassFilters)
+		{
+			if(CustomClassFilterData->GetClass() == InClass)
+			{
+				return true;
+			}
+		}
+
+		return false;
+	}
+
+	/** Append the filter associated with the input class to the filter bar.
+	 *  A filter associated with the input class must exist in the widget for this to succeed.
+	 */
+	void AddClassFilterToBar(UClass *Class, bool bEnableFilter = true)
+	{
+		for(const TSharedRef<FCustomClassFilterData>& CustomClassFilterData : CustomClassFilters)
+		{
+			if(CustomClassFilterData->GetClass() == Class)
+			{
+				// Make sure it already isn't added
+				if (!IsClassTypeInUse(CustomClassFilterData))
+				{
+					TSharedRef<SFilter> NewFilter = AddAssetFilterToBar(CustomClassFilterData);
+
+					if(bEnableFilter)
+					{
+						NewFilter->SetEnabled(true);
+					}
+				}
+			}
+		}
+	}
+
+	/** Remove all filters from the filter bar, while disabling any active ones */
+	virtual void RemoveAllFilters() override
+	{
+		AssetFilters.Empty();
+		SBasicFilterBar<FilterType>::RemoveAllFilters();
+	}
+
+protected:
+
+	/** AssetFilter specific override to SBasicFilterBar::RemoveAllButThis */
+	virtual void RemoveAllButThis(const TSharedRef<SFilter>& FilterToKeep) override
+	{
+		TSharedPtr<SAssetFilter> AssetFilterToKeep;
+
+		// Make sure to keep it in our local list of AssetFilters
+		for(const TSharedPtr<SAssetFilter> AssetFilter : AssetFilters)
+		{
+			if(AssetFilter == FilterToKeep)
+			{
+				AssetFilterToKeep = AssetFilter;
+			}
+		}
+		
+		SBasicFilterBar<FilterType>::RemoveAllButThis(FilterToKeep);
+
+		AssetFilters.Empty();
+
+		if(AssetFilterToKeep)
+		{
+			AssetFilters.Add(AssetFilterToKeep.ToSharedRef());
+		}
+	}
+
+	/** Add an AssetFilter to the toolbar, making it "Active" but not enabled */
+	TSharedRef<SFilter> AddAssetFilterToBar(const TSharedPtr<FCustomClassFilterData>& CustomClassFilter)
+	{
+		TSharedRef<SAssetFilter> NewFilter =
+			SNew(SAssetFilter)
+			.CustomClassFilter(CustomClassFilter)
+			.OnFilterChanged(this->OnFilterChanged)
+			.OnRequestRemove(this, &SAssetFilterBar<FilterType>::RemoveFilterAndUpdate)
+			.OnRequestEnableOnly(this, &SAssetFilterBar<FilterType>::EnableOnlyThisFilter)
+			.OnRequestEnableAll(this, &SAssetFilterBar<FilterType>::EnableAllFilters)
+			.OnRequestDisableAll(this, &SAssetFilterBar<FilterType>::DisableAllFilters)
+			.OnRequestRemoveAll(this, &SAssetFilterBar<FilterType>::RemoveAllFilters)
+			.OnRequestRemoveAllButThis(this, &SAssetFilterBar<FilterType>::RemoveAllButThis);
+
+		this->AddFilterToBar( NewFilter );
+
+		// Add it to our list of just AssetFilters
+		AssetFilters.Add(NewFilter);
+		
+		return NewFilter;
+	}
+	
+	/** Remove a filter from the filter bar */
+	virtual void RemoveFilter(const TSharedRef<SFilter>& FilterToRemove) override
+	{
+		SBasicFilterBar<FilterType>::RemoveFilter(FilterToRemove);
+
+		AssetFilters.RemoveAll([&FilterToRemove](TSharedRef<SAssetFilter>& AssetFilter) { return AssetFilter == FilterToRemove; });
+	}
+
+	/** Handler for when the remove filter button was clicked on a filter */
+	void RemoveAssetFilter(const TSharedPtr<FCustomClassFilterData>& CustomClassData, bool ExecuteOnFilterChanged = true)
+	{
+		TSharedPtr<SAssetFilter> FilterToRemove;
+		for ( const TSharedPtr<SAssetFilter> AssetFilter : AssetFilters )
+		{
+			const TSharedPtr<FCustomClassFilterData>& ClassData = AssetFilter->GetCustomClassFilterData();
+			if (ClassData == CustomClassData)
+			{
+				FilterToRemove = AssetFilter;
+				break;
+			}
+		}
+
+		if ( FilterToRemove.IsValid() )
+		{
+			if (ExecuteOnFilterChanged)
+			{
+				this->RemoveFilterAndUpdate(FilterToRemove.ToSharedRef());
+			}
+			else
+			{
+				this->RemoveFilter(FilterToRemove.ToSharedRef());
+			}
+
+			// Remove it from our local list of AssetFilters
+			AssetFilters.Remove(FilterToRemove.ToSharedRef());
+		}
+	}
+
+	/** Create the default set of IAssetTypeActions Filters provided with the widget if requested */
+	void CreateAssetTypeActionFilters()
+	{
+		if(!bUseDefaultAssetFilters)
+		{
+			return;
+		}
+
+		FAssetToolsModule& AssetToolsModule = FModuleManager::LoadModuleChecked<FAssetToolsModule>(TEXT("AssetTools"));
+
+		AssetFilterCategories.Empty();
+
+		// Add the Basic category
+		AssetFilterCategories.Add(EAssetTypeCategories::Basic, MakeShareable(new FFilterCategory(LOCTEXT("BasicFilter", "Basic"), LOCTEXT("BasicFilterTooltip", "Filter by basic assets."))));
+
+		// Add the advanced categories
+		TArray<FAdvancedAssetCategory> AdvancedAssetCategories;
+		AssetToolsModule.Get().GetAllAdvancedAssetCategories(/*out*/ AdvancedAssetCategories);
+
+		for (const FAdvancedAssetCategory& AdvancedAssetCategory : AdvancedAssetCategories)
+		{
+			const FText Tooltip = FText::Format(LOCTEXT("WildcardFilterTooltip", "Filter by {0} Assets."), AdvancedAssetCategory.CategoryName);
+			AssetFilterCategories.Add(AdvancedAssetCategory.CategoryType, MakeShareable(new FFilterCategory(AdvancedAssetCategory.CategoryName, Tooltip)));
+		}
+
+		// Get the browser type maps
+		TArray<TWeakPtr<IAssetTypeActions>> AssetTypeActionsList;
+		AssetToolsModule.Get().GetAssetTypeActionsList(AssetTypeActionsList);
+		
+		// Sort the list
+		struct FCompareIAssetTypeActions
+		{
+			FORCEINLINE bool operator()( const TWeakPtr<IAssetTypeActions>& A, const TWeakPtr<IAssetTypeActions>& B ) const
+			{
+				return A.Pin()->GetName().CompareTo( B.Pin()->GetName() ) == -1;
+			}
+		};
+		AssetTypeActionsList.Sort( FCompareIAssetTypeActions() );
+
+		const TSharedRef<FNamePermissionList>& AssetClassPermissionList = AssetToolsModule.Get().GetAssetClassPermissionList(EAssetClassAction::CreateAsset);
+
+		// For every asset type, convert it to an FCustomClassFilterData and add it to the list
+		for (int32 ClassIdx = 0; ClassIdx < AssetTypeActionsList.Num(); ++ClassIdx)
+		{
+			const TWeakPtr<IAssetTypeActions>& WeakTypeActions = AssetTypeActionsList[ClassIdx];
+			if ( WeakTypeActions.IsValid() )
+			{
+				TSharedPtr<IAssetTypeActions> TypeActions = WeakTypeActions.Pin();
+				if ( ensure(TypeActions.IsValid()) && TypeActions->CanFilter() )
+				{
+					UClass* SupportedClass = TypeActions->GetSupportedClass();
+					if ((!SupportedClass || AssetClassPermissionList->PassesFilter(SupportedClass->GetFName())))
+					{
+						// Convert the AssetTypeAction to an FCustomClassFilterData and add it to our list
+						TSharedRef<FCustomClassFilterData> CustomClassFilterData = MakeShared<FCustomClassFilterData>(TypeActions);
+						CustomClassFilters.Add(CustomClassFilterData);
+					}
+				}
+			}
+		}
+		
+		// Do a second pass through all the CustomClassFilters with AssetTypeActions to update their categories
+		UpdateAssetTypeActionCategories();
+
+	}
+
+	void UpdateAssetTypeActionCategories()
+	{
+		for(const TSharedRef<FCustomClassFilterData> &CustomClassFilter : CustomClassFilters)
+		{
+			for (auto MenuIt = AssetFilterCategories.CreateIterator(); MenuIt; ++MenuIt)
+			{
+				if(TSharedPtr<IAssetTypeActions> AssetTypeActions = CustomClassFilter->GetAssetTypeActions())
+				{
+					if(MenuIt.Key() & AssetTypeActions->GetCategories())
+					{
+						CustomClassFilter->AddCategory(MenuIt.Value());
+					}
+				}
+				
+			}
+		}	
+	}
+
+	
+	/* Filter Dropdown Related Functionality */
+
+	/** Handler for when the add filter menu is populated by a category */
+	void CreateFiltersMenuCategory(FToolMenuSection& Section, const TArray<TSharedPtr<FCustomClassFilterData>> CustomClassFilterDatas) const
+	{
+		for (int32 ClassIdx = 0; ClassIdx < CustomClassFilterDatas.Num(); ++ClassIdx)
+		{
+			const TSharedPtr<FCustomClassFilterData>& CustomClassFilterData = CustomClassFilterDatas[ClassIdx];
+			
+			const FText& LabelText = CustomClassFilterData->GetName();
+			Section.AddMenuEntry(
+				NAME_None,
+				LabelText,
+				FText::Format( LOCTEXT("FilterByTooltipPrefix", "Filter by {0}"), LabelText ),
+				FSlateIcon(),
+				FUIAction(
+					FExecuteAction::CreateSP( const_cast< SAssetFilterBar<FilterType>* >(this), &SAssetFilterBar<FilterType>::FilterByTypeClicked, CustomClassFilterData ),
+					FCanExecuteAction(),
+					FIsActionChecked::CreateSP(this, &SAssetFilterBar<FilterType>::IsClassTypeInUse, CustomClassFilterData ) ),
+				EUserInterfaceActionType::ToggleButton
+				);
+		}
+	}
+	
+	void CreateFiltersMenuCategory(UToolMenu* InMenu, const TArray<TSharedPtr<FCustomClassFilterData>> CustomClassFilterDatas) const
+	{
+		CreateFiltersMenuCategory(InMenu->AddSection("Section"), CustomClassFilterDatas);
+	}
+
+	/** Handler for when the add filter button was clicked */
+	virtual TSharedRef<SWidget> MakeAddFilterMenu() override
+	{
+		const FName FilterMenuName = "FilterBar.FilterMenu";
+		if (!UToolMenus::Get()->IsMenuRegistered(FilterMenuName))
+		{
+			UToolMenu* Menu = UToolMenus::Get()->RegisterMenu(FilterMenuName);
+			Menu->bShouldCloseWindowAfterMenuSelection = true;
+			Menu->bCloseSelfOnly = true;
+
+			Menu->AddDynamicSection(NAME_None, FNewToolMenuDelegate::CreateLambda([](UToolMenu* InMenu)
+			{
+				if (UAssetFilterBarContext* Context = InMenu->FindContext<UAssetFilterBarContext>())
+				{
+					Context->PopulateFilterMenu.ExecuteIfBound(InMenu, Context->MenuExpansion, FOnFilterAssetType());
+				}
+			}));
+		}
+
+		UAssetFilterBarContext* FilterBarContext = NewObject<UAssetFilterBarContext>();
+		FilterBarContext->PopulateFilterMenu = FOnPopulateAddAssetFilterMenu::CreateSP(this, &SAssetFilterBar<FilterType>::PopulateAddFilterMenu);
+
+		/** Auto expand the Basic Category if it is present */
+		if(TSharedPtr<FFilterCategory>* BasicCategory = AssetFilterCategories.Find(EAssetTypeCategories::Basic))
+		{
+			FilterBarContext->MenuExpansion = *BasicCategory;
+		}
+	
+		FToolMenuContext ToolMenuContext(FilterBarContext);
+
+		return UToolMenus::Get()->GenerateWidget(FilterMenuName, ToolMenuContext);
+	} 
+	
+	/** Handler to Populate the Add Filter Menu. Use OnFilterAssetType in Subclasses to add classes to the exclusion list */
+	virtual void PopulateAddFilterMenu(UToolMenu* Menu, TSharedPtr<FFilterCategory> MenuExpansion, FOnFilterAssetType OnFilterAssetType)
+	{
+		FAssetToolsModule& AssetToolsModule = FModuleManager::LoadModuleChecked<FAssetToolsModule>(TEXT("AssetTools"));
+
+		// A local struct to describe a category in the filter menu
+		struct FCategoryMenu
+		{
+			/** The Classes that belong to this category */
+			TArray<TSharedPtr<FCustomClassFilterData>> Classes;
+
+			//Menu section
+			FName SectionExtensionHook;
+			FText SectionHeading;
+
+			FCategoryMenu(const FName& InSectionExtensionHook, const FText& InSectionHeading)
+				: SectionExtensionHook(InSectionExtensionHook)
+				, SectionHeading(InSectionHeading)
+			{}
+		};
+		
+		// Create a map of Categories to Menus
+		TMap<TSharedPtr<FFilterCategory>, FCategoryMenu> CategoryToMenuMap;
+		
+		// For every asset type, move it into all the categories it should appear in
+		for(const TSharedRef<FCustomClassFilterData> &CustomClassFilter : CustomClassFilters)
+		{
+			bool bPassesExternalFilters = true;
+
+			// Run any external class filters we have
+			if(OnFilterAssetType.IsBound())
+			{
+				bPassesExternalFilters = OnFilterAssetType.Execute(CustomClassFilter->GetClass());
+			}
+
+			if(bPassesExternalFilters)
+			{
+				/** Get all the categories this filter belongs to */
+				TArray<TSharedPtr<FFilterCategory>> Categories = CustomClassFilter->GetCategories();
+
+				for(const TSharedPtr<FFilterCategory>& Category : Categories)
+				{
+					// If the category for this custom class already exists
+					if(FCategoryMenu* CategoryMenu = CategoryToMenuMap.Find(Category))
+					{
+						CategoryMenu->Classes.Add( CustomClassFilter );
+					}
+					// Otherwise create a new FCategoryMenu for the category and add it to the map
+					else
+					{
+						const FName ExtensionPoint = NAME_None;
+						const FText SectionHeading = FText::Format(LOCTEXT("WildcardFilterHeadingHeadingTooltip", "{0} Filters"), Category->Title);
+
+						FCategoryMenu NewCategoryMenu(ExtensionPoint, SectionHeading);
+						NewCategoryMenu.Classes.Add(CustomClassFilter);
+					
+						CategoryToMenuMap.Add(Category, NewCategoryMenu);
+					}
+				}
+			}
+		}
+
+		// Remove any empty categories
+		for (auto MenuIt = CategoryToMenuMap.CreateIterator(); MenuIt; ++MenuIt)
+		{
+			if (MenuIt.Value().Classes.Num() == 0)
+			{
+				CategoryToMenuMap.Remove(MenuIt.Key());
+			}
+		}
+
+		// Set the extension hook for the basic category, if it exists and we have any assets for it
+		if(TSharedPtr<FFilterCategory>* BasicCategory = AssetFilterCategories.Find(EAssetTypeCategories::Basic))
+		{
+			if(FCategoryMenu* BasicMenu = CategoryToMenuMap.Find(*BasicCategory))
+			{
+				BasicMenu->SectionExtensionHook = "FilterBarFilterBasicAsset";
+			}
+		}
+
+		// Populate the common filter sections (Reset Filters etc)
+		{
+			this->PopulateCommonFilterSections(Menu);
+		}
+		
+		// If we want to expand a category
+		if(MenuExpansion)
+		{
+			// First add the expanded category, this appears as standard entries in the list (Note: intentionally not using FindChecked here as removing it from the map later would cause the ref to be garbage)
+			FCategoryMenu* ExpandedCategory = CategoryToMenuMap.Find( MenuExpansion );
+			if(ExpandedCategory)
+			{
+				FToolMenuSection& Section = Menu->AddSection(ExpandedCategory->SectionExtensionHook, ExpandedCategory->SectionHeading);
+				if(MenuExpansion == AssetFilterCategories.FindChecked(EAssetTypeCategories::Basic))
+				{
+					// If we are doing a full menu (i.e expanding basic) we add a menu entry which toggles all other categories
+					Section.AddMenuEntry(
+						NAME_None,
+						MenuExpansion->Title,
+						MenuExpansion->Tooltip,
+						FSlateIcon(),
+						FUIAction(
+						FExecuteAction::CreateSP( this, &SAssetFilterBar<FilterType>::FilterByTypeCategoryClicked, MenuExpansion, ExpandedCategory->Classes ),
+						FCanExecuteAction(),
+						FGetActionCheckState::CreateSP(this, &SAssetFilterBar<FilterType>::IsTypeCategoryChecked, MenuExpansion, ExpandedCategory->Classes ) ),
+						EUserInterfaceActionType::ToggleButton
+						);
+				}
+
+				// Now populate with all the assets from the expanded category
+				SAssetFilterBar<FilterType>::CreateFiltersMenuCategory( Section, ExpandedCategory->Classes);
+				
+				// Remove the Expanded from the map now, as this is treated differently and is no longer needed.
+				CategoryToMenuMap.Remove(MenuExpansion);
+			}
+		}
+
+		TSharedPtr<FFilterCategory>* BasicCategory = AssetFilterCategories.Find(EAssetTypeCategories::Basic);
+
+		// We are in Full Menu Mode if there is no menu expansion, or the menu expansion is EAssetTypeCategories::Basic
+		bool bInFullMenuMode = !MenuExpansion || (BasicCategory && MenuExpansion == *BasicCategory);
+
+		// If we are in full menu mode, add all the other categories as submenus
+		if(bInFullMenuMode)
+		{
+			FToolMenuSection& Section = Menu->AddSection("AssetFilterBarFilterAdvancedAsset", LOCTEXT("AdvancedAssetsMenuHeading", "Other Assets"));
+			
+			// Sort by category name so that we add the submenus in alphabetical order
+			CategoryToMenuMap.KeySort([](const TSharedPtr<FFilterCategory>& A, const TSharedPtr<FFilterCategory>& B) {
+				return A->Title.CompareTo(B->Title) < 0;
+			});
+
+			// For all the remaining categories, add them as submenus
+			for (const TPair<TSharedPtr<FFilterCategory>, FCategoryMenu>& CategoryMenuPair : CategoryToMenuMap)
+			{
+				Section.AddSubMenu(
+					NAME_None,
+					CategoryMenuPair.Key->Title,
+					CategoryMenuPair.Key->Tooltip,
+					FNewToolMenuDelegate::CreateSP(this, &SAssetFilterBar<FilterType>::CreateFiltersMenuCategory, CategoryMenuPair.Value.Classes),
+					FUIAction(
+					FExecuteAction::CreateSP(this, &SAssetFilterBar<FilterType>::FilterByTypeCategoryClicked, CategoryMenuPair.Key, CategoryMenuPair.Value.Classes),
+					FCanExecuteAction(),
+					FGetActionCheckState::CreateSP(this, &SAssetFilterBar<FilterType>::IsTypeCategoryChecked, CategoryMenuPair.Key, CategoryMenuPair.Value.Classes)),
+					EUserInterfaceActionType::ToggleButton
+					);
+			}
+		}
+
+		// Now add all non-asset filters
+		this->PopulateCustomFilters(Menu);
+	}
+
+	/* Asset Type filter related functionality */
+	
+	/** Handler for when filter by type is selected */
+ 	void FilterByTypeClicked(TSharedPtr<FCustomClassFilterData> CustomClassFilterData)
+	{
+		if (CustomClassFilterData.IsValid())
+		{
+			if (IsClassTypeInUse(CustomClassFilterData))
+			{
+				RemoveAssetFilter(CustomClassFilterData);
+			}
+			else
+			{
+				TSharedRef<SFilter> NewFilter = AddAssetFilterToBar(CustomClassFilterData);
+				NewFilter->SetEnabled(true);
+			}
+		}
+	}
+
+	/** Handler to determine the "checked" state of class filter in the filter dropdown */
+	bool IsClassTypeInUse(TSharedPtr<FCustomClassFilterData> Class) const
+	{
+		for (const TSharedPtr<SAssetFilter> AssetFilter : this->AssetFilters)
+		{
+			if (AssetFilter.IsValid() && AssetFilter->GetCustomClassFilterData() == Class)
+			{
+				return true;
+			}
+		}
+
+		return false;
+	}
+
+ 	/** Handler for when filter by type category is selected */
+ 	void FilterByTypeCategoryClicked(TSharedPtr<FFilterCategory> TypeCategory, TArray<TSharedPtr<FCustomClassFilterData>> Classes)
+	{
+		bool bFullCategoryInUse = IsTypeCategoryInUse(TypeCategory, Classes);
+		bool ExecuteOnFilterChanged = false;
+
+		for(const TSharedPtr<FCustomClassFilterData>& CustomClass : Classes)
+		{
+			if(bFullCategoryInUse)
+			{
+				RemoveAssetFilter(CustomClass);
+				ExecuteOnFilterChanged = true;
+			}
+			else if(!IsClassTypeInUse(CustomClass))
+			{
+				TSharedRef<SFilter> NewFilter = AddAssetFilterToBar(CustomClass);
+				NewFilter->SetEnabled(true, false);
+				ExecuteOnFilterChanged = true;
+			}
+		}
+
+		if (ExecuteOnFilterChanged)
+		{
+			this->OnFilterChanged.ExecuteIfBound();
+		}
+	}
+
+ 	/** Handler to determine the "checked" state of an type category in the filter dropdown */
+ 	ECheckBoxState IsTypeCategoryChecked(TSharedPtr<FFilterCategory> TypeCategory, TArray<TSharedPtr<FCustomClassFilterData>> Classes) const
+	{
+		bool bIsAnyActionInUse = false;
+		bool bIsAnyActionNotInUse = false;
+
+		for (const TSharedPtr<FCustomClassFilterData>& CustomClassFilter : Classes)
+		{
+			if (IsClassTypeInUse(CustomClassFilter))
+			{
+				bIsAnyActionInUse = true;
+			}
+			else
+			{
+				bIsAnyActionNotInUse = true;
+			}
+
+			if (bIsAnyActionInUse && bIsAnyActionNotInUse)
+			{
+				return ECheckBoxState::Undetermined;
+			}
+			
+		}
+
+		if (bIsAnyActionInUse)
+		{
+			return ECheckBoxState::Checked;
+		}
+		else
+		{
+			return ECheckBoxState::Unchecked;
+		}
+	}
+
+ 	/** Function to check if a given type category is in use */
+ 	bool IsTypeCategoryInUse(TSharedPtr<FFilterCategory> TypeCategory, TArray<TSharedPtr<FCustomClassFilterData>> Classes) const
+ 	{
+ 		ECheckBoxState AssetTypeCategoryCheckState = IsTypeCategoryChecked(TypeCategory, Classes);
+
+ 		if (AssetTypeCategoryCheckState == ECheckBoxState::Unchecked)
+ 		{
+ 			return false;
+ 		}
+
+ 		// An asset type category is in use if any of its type actions are in use (ECheckBoxState::Checked or ECheckBoxState::Undetermined)
+ 		return true;
+ 	}
+	
+protected:
+	/** A copy of all AssetFilters in this->Filters for convenient access */
+	TArray< TSharedRef<SAssetFilter> > AssetFilters;
+
+	/** List of custom Class Filters that will be shown in the filter bar */
+	TArray<TSharedRef<FCustomClassFilterData>> CustomClassFilters;
+
+	TMap<EAssetTypeCategories::Type, TSharedPtr<FFilterCategory>> AssetFilterCategories;
+	
+	/** Whether the filter bar provides the default Asset Filters */
+	bool bUseDefaultAssetFilters;
+
+};
+
+#undef LOCTEXT_NAMESPACE
\ No newline at end of file
diff --git a/Engine/Source/Editor/EditorWidgets/Public/Filters/SFilterBar.h b/Engine/Source/Editor/EditorWidgets/Public/Filters/SFilterBar.h
new file mode 100644
--- /dev/null
+++ b/Engine/Source/Editor/EditorWidgets/Public/Filters/SFilterBar.h
@@ -0,0 +1,124 @@
+//Copyright Epic Games, Inc. All Rights Reserved.
+
+#pragma once
+
+#include "SlateGlobals.h"
+#include "Logging/LogMacros.h"
+#include "Logging/LogVerbosity.h"
+
+#include "Filters/SAssetFilterBar.h"
+#include "Filters/AssetFilter.h"
+
+
+/**
+* A Filter Bar widget, which can be used to filter items of type [FilterType] given a list of custom filters
+* along with built in support for Asset Type filters
+* @see SBasicFilterBar if you don't want Asset Type filters, or if you want a filter bar usable in non-editor situations
+* NOTE: You will need to also add "ToolWidgets" as a dependency to your module to use this widget
+* NOTE: The filter functions create copies, so you want to use SFilterBar<TSharedPtr<ItemType>> etc instead of SFilterBar<ItemType> when possible
+* NOTE: The user must specify one of the following:
+* a) OnConvertItemToAssetData:   A conversion function to convert FilterType to an FAssetData. Specifying this will filter the asset type
+*							     through the AssetRegistry, which is potentially more thorough and fast
+* b) OnCompareItemWithClassName: A comparison function to check if an Item (FilterType) is the same as an AssetType (represented by an FName).
+*								 This allows to user to directly text compare with Class Names (UObjectBase::GetFName), which is easier but potentially
+*								 slower
+*								 
+* Sample Usage:
+*		SAssignNew(MyFilterBar, SFilterBar<FText>)
+*		.OnFilterChanged() // A delegate for when the list of filters changes
+*		.CustomFilters() // An array of filters available to this FilterBar (@see FGenericFilter to create simple delegate based filters)
+*		.OnConvertItemToAssetData() // Conversion Function as mentioned above
+*
+* Use the GetAllActiveFilters() function to get the FilterCollection of Active Filters in this FilterBar, that can be used to filter your items.
+* NOTE: GetAllActiveFilters() must be called every time the filters change (in OnFilterChanged() for example) to make sure you have the correct backend filter
+* NOTE: Use CustomClassFilters to provide any Type Filters to make sure they get resolved properly (See FCustomClassFilterData)
+* NOTE: Use MakeAddFilterButton() to make the button that summons the dropdown showing all the filters
+*/
+template<typename FilterType>
+class SFilterBar : public SAssetFilterBar<FilterType>
+{
+public:
+	using FOnFilterChanged = typename SBasicFilterBar<FilterType>::FOnFilterChanged;
+	using FConvertItemToAssetData = typename FAssetFilter<FilterType>::FConvertItemToAssetData;
+	using FCompareItemWithClassNames = typename FAssetFilter<FilterType>::FCompareItemWithClassNames;
+
+	SLATE_BEGIN_ARGS( SFilterBar<FilterType> )
+		: _UseDefaultAssetFilters(true)
+		{
+		
+		}
+
+		/** Delegate for when filters have changed */
+		SLATE_EVENT( FOnFilterChanged, OnFilterChanged )
+
+		/** Specify this delegate to use asset comparison through IAssetRegistry, where you specify how to convert your Item into an FAssetData */
+		SLATE_EVENT( FConvertItemToAssetData, OnConvertItemToAssetData )
+
+		/** Specify this delegate to use simple asset comparison, where you compare your Item a list of FNames representing Classes */
+		SLATE_EVENT( FCompareItemWithClassNames, OnCompareItemWithClassNames )
+
+		/** Initial List of Custom Filters that will be added to the AddFilter Menu */
+		SLATE_ARGUMENT( TArray<TSharedRef<FFilterBase<FilterType>>>, CustomFilters)
+
+		/** Initial List of Custom Class Filters that will be added to the AddFilter Menu */
+		SLATE_ARGUMENT( TArray<TSharedRef<FCustomClassFilterData>>, CustomClassFilters)
+	
+		/** Whether the filter bar should provide the default asset filters */
+		SLATE_ARGUMENT(bool, UseDefaultAssetFilters)
+	
+	SLATE_END_ARGS()
+
+	/** Constructs this widget with InArgs */
+	void Construct( const FArguments& InArgs )
+	{
+		typename SAssetFilterBar<FilterType>::FArguments Args;
+		Args._OnFilterChanged = InArgs._OnFilterChanged;
+		Args._CustomFilters = InArgs._CustomFilters;
+		Args._UseDefaultAssetFilters = InArgs._UseDefaultAssetFilters;
+		Args._CustomClassFilters = InArgs._CustomClassFilters;
+		
+		SAssetFilterBar<FilterType>::Construct(Args);
+
+		// Create the dummy filter that represents all currently active Asset Type filters
+		AssetFilter = MakeShareable(new FAssetFilter<FilterType>);
+
+		// Asset Conversion is preferred to Asset Comparison
+		if(InArgs._OnConvertItemToAssetData.IsBound())
+		{
+			AssetFilter->SetConversionFunction(InArgs._OnConvertItemToAssetData);
+		}
+		else if(InArgs._OnCompareItemWithClassNames.IsBound())
+		{
+			AssetFilter->SetComparisonFunction(InArgs._OnCompareItemWithClassNames);
+		}
+		else
+		{
+			UE_LOG(LogSlate, Error, TEXT("SFilterBar Requires that you specify either OnConvertItemToAssetData or OnCompareItemWithClassName"));
+		}
+
+		this->ActiveFilters->Add(AssetFilter);
+	}
+	
+	/** Use this function to get all currently active filters (to filter your items)
+	 * NOTE: Must be called every time the filters change (OnFilterChanged) to make sure you get the correct combined filter
+	 */
+	virtual TSharedPtr< TFilterCollection<FilterType> > GetAllActiveFilters() override
+	{
+		UpdateAssetFilter();
+		
+		return this->ActiveFilters;	
+	}
+
+private:
+
+	void UpdateAssetFilter()
+	{
+		/** We have to make sure to update the CombinedBackendFilter everytime the user requests all the Filters */
+		FARFilter CombinedBackendFilter = this->GetCombinedBackendFilter();
+		AssetFilter->SetBackendFilter(CombinedBackendFilter);
+	}
+	
+private:
+	/* The invisible filter used to conduct Asset Type filtering */
+	TSharedPtr<FAssetFilter<FilterType>> AssetFilter;
+};
