Restore the 'virtualized data' filter that was removed in CL 20340436

#rnx

diff --git a/Engine/Source/Editor/ContentBrowser/Private/SFilterList.cpp b/Engine/Source/Editor/ContentBrowser/Private/SFilterList.cpp
--- a/Engine/Source/Editor/ContentBrowser/Private/SFilterList.cpp
+++ b/Engine/Source/Editor/ContentBrowser/Private/SFilterList.cpp
@@ -62,6 +62,7 @@
 	AllFrontendFilters_Internal.Add( MakeShareable(new FFrontendFilter_ArbitraryComparisonOperation(DefaultCategory)) );
 	AllFrontendFilters_Internal.Add(MakeShareable(new FFrontendFilter_Recent(DefaultCategory)));
 	AllFrontendFilters_Internal.Add( MakeShareable(new FFrontendFilter_NotSourceControlled(DefaultCategory)) );
+	AllFrontendFilters_Internal.Add(MakeShareable(new FFrontendFilter_VirtualizedData(DefaultCategory)));
 
 	// Add any global user-defined frontend filters
 	for (TObjectIterator<UContentBrowserFrontEndFilterExtension> ExtensionIt(RF_NoFlags); ExtensionIt; ++ExtensionIt)
