Fix a bug with the animation browser where sound assets would always show up.

SFilterList: Use the AddFilter() function to add extra filters to make sure the activation helper gets bound and inverse filters are handled properly.

diff --git a/Engine/Source/Editor/ContentBrowser/Private/SFilterList.cpp b/Engine/Source/Editor/ContentBrowser/Private/SFilterList.cpp
--- a/Engine/Source/Editor/ContentBrowser/Private/SFilterList.cpp
+++ b/Engine/Source/Editor/ContentBrowser/Private/SFilterList.cpp
@@ -105,7 +105,7 @@ void SFilterList::Construct( const FArguments& InArgs )
 	// Add the local copy of all filters to SFilterBar's copy of all filters
 	for(TSharedRef<FFrontendFilter> FrontendFilter : AllFrontendFilters_Internal)
 	{
-		this->AllFrontendFilters.Add(FrontendFilter);
+		this->AddFilter(FrontendFilter);
 	}
 	
 	SAssetFilterBar<FAssetFilterType>::FArguments Args;
@@ -115,7 +115,6 @@ void SFilterList::Construct( const FArguments& InArgs )
 	Args._OnFilterChanged = this->OnFilterChanged;
 	
 	SAssetFilterBar<FAssetFilterType>::Construct(Args);
-
 }
 
 FReply SFilterList::OnMouseButtonUp( const FGeometry& MyGeometry, const FPointerEvent& MouseEvent )
