Fix issue where ContentBrowserDrawer is not marked as the primary content browser when selected in ChooseNewPrimaryBrowser

diff --git a/Engine/Source/Editor/ContentBrowser/Private/ContentBrowserSingleton.cpp b/Engine/Source/Editor/ContentBrowser/Private/ContentBrowserSingleton.cpp
--- a/Engine/Source/Editor/ContentBrowser/Private/ContentBrowserSingleton.cpp
+++ b/Engine/Source/Editor/ContentBrowser/Private/ContentBrowserSingleton.cpp
@@ -663,7 +663,7 @@ void FContentBrowserSingleton::ChooseNewPrimaryBrowser()
 	// Set the content browser drawer as the primary browser
 	if (!PrimaryContentBrowser.IsValid())
 	{
-		PrimaryContentBrowser = ContentBrowserDrawer;
+		SetPrimaryContentBrowser(ContentBrowserDrawer.Pin().ToSharedRef());
 	}
 }
 
