Increase CB description width to 700

diff --git a/Engine/Source/Editor/ContentBrowser/Private/AssetViewWidgets.cpp b/Engine/Source/Editor/ContentBrowser/Private/AssetViewWidgets.cpp
--- a/Engine/Source/Editor/ContentBrowser/Private/AssetViewWidgets.cpp
+++ b/Engine/Source/Editor/ContentBrowser/Private/AssetViewWidgets.cpp
@@ -906,7 +906,7 @@ TSharedRef<SWidget> SAssetViewItem::CreateToolTipWidget() const
 						.BorderImage(FAppStyle::GetBrush("ContentBrowser.TileViewTooltip.ContentBorder"))
 						[
 							SNew(STextBlock)
-							.WrapTextAt(300.0f)
+							.WrapTextAt(700.0f)
 							.Font(FAppStyle::GetFontStyle("ContentBrowser.TileViewTooltip.AssetUserDescriptionFont"))
 							.Text(UserDescription)
 						]
