Content Browser: Add the ability to save searches as filters,  port the config over to EditorConfig for SFilterList
SFilterBar: Update search saving to use an independant interface and move FFilterBarBase to SAssetFilterBar

diff --git a/Engine/Source/Developer/ToolWidgets/Private/Filters/SCustomTextFilterDialog.cpp b/Engine/Source/Developer/ToolWidgets/Private/Filters/SCustomTextFilterDialog.cpp
--- a/Engine/Source/Developer/ToolWidgets/Private/Filters/SCustomTextFilterDialog.cpp
+++ b/Engine/Source/Developer/ToolWidgets/Private/Filters/SCustomTextFilterDialog.cpp
@@ -12,20 +12,33 @@
 void SCustomTextFilterDialog::Construct( const FArguments& InArgs )
 {
 	FilterData = InArgs._FilterData;
+	InitialFilterData = FilterData;
 	bInEditMode = InArgs._InEditMode;
 	OnCreateFilter = InArgs._OnCreateFilter;
 	OnDeleteFilter = InArgs._OnDeleteFilter;
 	OnCancelClicked = InArgs._OnCancelClicked;
 	OnModifyFilter = InArgs._OnModifyFilter;
+	OnGetFilterLabels = InArgs._OnGetFilterLabels;
 	
 	TSharedPtr<SHorizontalBox> ButtonBox = SNew(SHorizontalBox);
 
 	// Buttons for when we are editing a filter
 	if(bInEditMode)
 	{
+	    // Button to modify the current filter
+       	ButtonBox->AddSlot()
+		.FillWidth(1.0)
+       	.Padding(0, 0, 16, 0)
+       	.HAlign(HAlign_Right)
+       	[
+       		SNew(SPrimaryButton)
+       		.Text(LOCTEXT("ModifyFilterButton", "Save"))
+       		.OnClicked(this, &SCustomTextFilterDialog::OnModifyButtonClicked)
+       	];
+	
 		// Button to delete the current filter
 		ButtonBox->AddSlot()
-		.FillWidth(1.0)
+		.AutoWidth()
 		.Padding(0, 0, 16, 0)
 		.HAlign(HAlign_Right)
 		[
@@ -33,28 +46,6 @@ void SCustomTextFilterDialog::Construct( const FArguments& InArgs )
 			.Text(LOCTEXT("DeleteButton", "Delete"))
 			.OnClicked(this, &SCustomTextFilterDialog::OnDeleteButtonClicked)
 		];
-
-		// Button to create a new filter with the current data
-		ButtonBox->AddSlot()
-		.AutoWidth()
-		.Padding(0, 0, 16, 0)
-		.HAlign(HAlign_Right)
-		[
-			SNew(SButton)
-			.Text(LOCTEXT("CreateNewButton", "Create New"))
-			.OnClicked(this, &SCustomTextFilterDialog::OnCreateFilterButtonClicked, false)
-		];
-
-		// Button to modify the current filter
-		ButtonBox->AddSlot()
-		.AutoWidth()
-		.Padding(0, 0, 16, 0)
-		.HAlign(HAlign_Right)
-		[
-			SNew(SPrimaryButton)
-			.Text(LOCTEXT("ModifyFilterButton", "Modify"))
-			.OnClicked(this, &SCustomTextFilterDialog::OnModifyButtonClicked)
-		];
 	}
 	else
 	{
@@ -231,10 +222,28 @@ bool SCustomTextFilterDialog::CheckFilterValidity() const
 	if(FilterData.FilterLabel.IsEmpty())
 	{
 		FilterLabelTextBox->SetError(LOCTEXT("EmptyFilterLabelError", "Filter Label cannot be empty"));
-		
 		return false;
 	}
 
+	if(OnGetFilterLabels.IsBound())
+	{
+		TArray<FText> ExistingFilterLabels;
+		OnGetFilterLabels.Execute(ExistingFilterLabels);
+
+		// Check for duplicate filter labels
+		for(const FText& FilterLabel : ExistingFilterLabels)
+		{
+			/* Special Case: If we are editing a filter and don't change the filter label, it will be considered a duplicate of itself!
+			 * To prevent this we check against the original filter label if we are in edit mode
+			 */
+			if(FilterLabel.EqualTo(FilterData.FilterLabel) && !(bInEditMode && FilterLabel.EqualTo(InitialFilterData.FilterLabel)))
+			{
+				FilterLabelTextBox->SetError(LOCTEXT("DuplicateFilterLabelError", "A filter with this label already exists!"));
+				return false;
+			}
+		}
+	}
+
 	return true;
 }
 
diff --git a/Engine/Source/Developer/ToolWidgets/Private/Filters/SFilterSearchBox.cpp b/Engine/Source/Developer/ToolWidgets/Private/Filters/SFilterSearchBox.cpp
--- a/Engine/Source/Developer/ToolWidgets/Private/Filters/SFilterSearchBox.cpp
+++ b/Engine/Source/Developer/ToolWidgets/Private/Filters/SFilterSearchBox.cpp
@@ -16,6 +16,8 @@ void SFilterSearchBox::Construct( const FArguments& InArgs )
 	MaxSearchHistory = InArgs._MaxSearchHistory;
 	OnTextChanged = InArgs._OnTextChanged;
 	OnTextCommitted = InArgs._OnTextCommitted;
+	bShowSearchHistory = InArgs._ShowSearchHistory;
+	OnSaveSearchClicked = InArgs._OnSaveSearchClicked;
 
 	// Default text shown when there are no items in the search history
 	EmptySearchHistoryText = MakeShareable(new FText(LOCTEXT("EmptySearchHistoryText", "The Search History is Empty")));
@@ -36,6 +38,8 @@ void SFilterSearchBox::Construct( const FArguments& InArgs )
 				.OnTextChanged(this, &SFilterSearchBox::HandleTextChanged)
 				.OnTextCommitted(this, &SFilterSearchBox::HandleTextCommitted)
 				.SelectAllTextWhenFocused( false )
+				.DelayChangeNotificationsWhileTyping( InArgs._DelayChangeNotificationsWhileTyping )
+				.OnKeyDownHandler(InArgs._OnKeyDownHandler)
 			]
 			+ SOverlay::Slot()
 			.HAlign(HAlign_Right)
@@ -44,37 +48,45 @@ void SFilterSearchBox::Construct( const FArguments& InArgs )
 
 				+SHorizontalBox::Slot()
 				.VAlign(VAlign_Center)
+				.HAlign(HAlign_Right)
 				[
 					// Button to save the currently occuring search 
 					SNew(SButton)
+					.ContentPadding(0)
 					.Visibility_Lambda([this]()
 					{
-						// Only visible if there is a search active currently, and the 
-						return this->GetText().IsEmpty() || !this->OnClickedAddSearchHistoryButton.IsBound() ? EVisibility::Collapsed : EVisibility::Visible;
+						// Only visible if there is a search active currently and the OnSaveSearchClicked delegate is bound
+						return this->GetText().IsEmpty() || !this->OnSaveSearchClicked.IsBound() ? EVisibility::Collapsed : EVisibility::Visible;
 					})
-					.ButtonStyle(FAppStyle::Get(), "NoBorder")
+					.ButtonStyle(FAppStyle::Get(), "HoverOnlyButton")
 					.OnClicked_Lambda([this]()
 					{
-						this->OnClickedAddSearchHistoryButton.ExecuteIfBound(this->GetText());
+						this->OnSaveSearchClicked.ExecuteIfBound(this->GetText());
 						return FReply::Handled();
 					})
 					[
 						SNew(SImage)
 						.Image(FAppStyle::Get().GetBrush("Icons.PlusCircle"))
+						.ColorAndOpacity(FSlateColor::UseForeground())
 					]
 				]
 
 				+SHorizontalBox::Slot()
+				.HAlign(HAlign_Right)
+				.Padding(0, 0, 4, 0)
 				[
 					// Chevron to open the search history dropdown
 					SNew(SButton)
-					.ButtonStyle(FAppStyle::Get(), "NoBorder")
+					.ContentPadding(0)
+					.ButtonStyle(FAppStyle::Get(), "HoverOnlyButton")
 					.ClickMethod(EButtonClickMethod::MouseDown)
 					.OnClicked(this, &SFilterSearchBox::OnClickedSearchHistory)
 					.ToolTipText(LOCTEXT("SearchHistoryToolTipText", "Click to show the Search History"))
+					.Visibility(this, &SFilterSearchBox::GetSearchHistoryVisibility)
 					[
 						SNew(SImage)
 						.Image(FAppStyle::Get().GetBrush("Icons.ChevronDown"))
+						.ColorAndOpacity(FSlateColor::UseForeground())
 					]
 				]
 			]
@@ -87,16 +99,22 @@ void SFilterSearchBox::Construct( const FArguments& InArgs )
 			[
 				SAssignNew(SearchHistoryListView, SListView< TSharedPtr<FText> >)
 				.ListItemsSource(&SearchHistory)
-				.SelectionMode( ESelectionMode::Single )							// Ideally the mouse over would not highlight while keyboard controls the UI
+				.SelectionMode( ESelectionMode::Single )
 				.OnGenerateRow(this, &SFilterSearchBox::MakeSearchHistoryRowWidget)
 				.OnSelectionChanged( this, &SFilterSearchBox::OnSelectionChanged)
 				.ItemHeight(18)
-				.ScrollbarDragFocusCause(EFocusCause::SetDirectly) // Use SetDirect so that clicking the scrollbar doesn't close the suggestions list
+				.ScrollbarDragFocusCause(EFocusCause::SetDirectly) 
 			]
 		)
 	];
 }
 
+bool SFilterSearchBox::HasKeyboardFocus() const
+{
+	// Since keyboard focus is forwarded to our editable text, we will test it instead
+	return SearchBox->HasKeyboardFocus();
+}
+
 /** Handler for when text in the editable text box changed */
 void SFilterSearchBox::HandleTextChanged(const FText& NewText)
 {
@@ -168,32 +186,7 @@ TSharedRef<ITableRow> SFilterSearchBox::MakeSearchHistoryRowWidget(TSharedPtr<FT
 		SNew(STextBlock)
 		.Text(*SearchText.Get())
 	];
-
-	// If the Search History is not empty, add the button to save the search history
-	if(!bIsEmptySearchHistory)
-	{
-		RowWidget->AddSlot()
-		.HAlign(HAlign_Right)
-		[
-			SNew(SButton)
-			.ButtonStyle(FAppStyle::Get(), "NoBorder")
-			.OnClicked_Lambda([this, SearchText]()
-			{
-				this->OnClickedAddSearchHistoryButton.ExecuteIfBound(*SearchText.Get());
-				return FReply::Handled();
-			})
-			.Visibility_Lambda([this]()
-			{
-				return this->OnClickedAddSearchHistoryButton.IsBound() ? EVisibility::Visible : EVisibility::Collapsed;
-			})
-			.ToolTipText(LOCTEXT("SaveSearchHistoryToolTipText", "Save this search as a custom filter"))
-			[
-				SNew(SImage)
-				.Image(FAppStyle::Get().GetBrush("Icons.PlusCircle"))
-			]
-		];
-	}
-
+	
 	return SNew(STableRow< TSharedPtr<FString> >, OwnerTable)
 			.ShowSelection(!bIsEmptySearchHistory)
 			[
@@ -252,4 +245,14 @@ FReply SFilterSearchBox::OnClickedSearchHistory()
 	return FReply::Handled();
 }
 
+void SFilterSearchBox::SetOnSaveSearchHandler(FOnSaveSearchClicked InOnSaveSearchHandler)
+{
+	OnSaveSearchClicked = InOnSaveSearchHandler;
+}
+
+EVisibility SFilterSearchBox::GetSearchHistoryVisibility() const
+{
+	return bShowSearchHistory.Get() ? EVisibility::Visible : EVisibility::Collapsed;
+}
+
 #undef LOCTEXT_NAMESPACE
diff --git a/Engine/Source/Developer/ToolWidgets/Public/Filters/CustomTextFilters.h b/Engine/Source/Developer/ToolWidgets/Public/Filters/CustomTextFilters.h
new file mode 100644
--- /dev/null
+++ b/Engine/Source/Developer/ToolWidgets/Public/Filters/CustomTextFilters.h
@@ -0,0 +1,208 @@
+// Copyright Epic Games, Inc. All Rights Reserved.
+
+#pragma once
+
+#include "Filters/FilterBase.h"
+#include "Core/Public/Misc/TextFilter.h"
+
+#include "CustomTextFilters.generated.h"
+
+
+/* Struct containing the data that SCustomTextFilterDialog is currently editing */
+USTRUCT()
+struct FCustomTextFilterData
+{
+	GENERATED_BODY()
+
+	UPROPERTY()
+	FText FilterLabel;
+	
+	UPROPERTY()
+	FText FilterString;
+
+	UPROPERTY()
+	FLinearColor FilterColor;
+
+	FCustomTextFilterData()
+		: FilterColor(FLinearColor::White)
+	{
+		
+	}
+};
+
+/** Interface class to define how a custom text filter converts to/from FCustomTextFilterData
+ *  FCustomTextFilter contains a generic implementation that can be used for most cases, look at
+ *  FFrontendFilter_CustomText in the Content Browser for an example of a specific implementation
+ */
+template<typename FilterType>
+class ICustomTextFilter
+{
+public:
+
+	virtual ~ICustomTextFilter() = default;
+	
+	/** Set the internals of this filter from an FCustomTextFilterData */
+	virtual void SetFromCustomTextFilterData(const FCustomTextFilterData& InFilterData) = 0;
+
+	/** Create an FCustomTextFilterData from the internals of this filter */
+	virtual FCustomTextFilterData CreateCustomTextFilterData() const = 0;
+
+	/** Get the actual filter */
+	virtual TSharedPtr<FFilterBase<FilterType>> GetFilter() = 0;
+};
+
+
+/* A generic CustomTextFilter that can be created by providing a TTextFilter.
+ * Provide a delegate to a filter bar widget's CreateTextFilter argument that creates an instance of this class
+ * to handle most generic text comparisons.
+ */
+template<typename FilterType>
+class FCustomTextFilter :
+	public FFilterBase<FilterType>,
+	public ICustomTextFilter<FilterType>,
+	public TSharedFromThis<FCustomTextFilter<FilterType>>
+{
+public:
+
+	FCustomTextFilter(TSharedPtr<TTextFilter<FilterType>> InTextFilter)
+	: FFilterBase<FilterType>(nullptr)
+	, TextFilter(InTextFilter)
+	{
+		
+	}
+
+	/** Returns the system name for this filter */
+	virtual FString GetName() const override
+	{
+		return GetFilterTypeName().ToString();
+	}
+
+	/** Returns the human readable name for this filter */
+	virtual FText GetDisplayName() const override
+	{
+		return DisplayName; 
+	}
+
+	/** Set the human readable name for this filter */
+	virtual void SetDisplayName(const FText& InDisplayName)
+	{
+		DisplayName = InDisplayName;
+	}
+
+	/** Returns the tooltip for this filter, shown in the filters menu */
+	virtual FText GetToolTipText() const override
+	{
+		return TextFilter->GetRawFilterText(); // The tooltip will display the filter string
+	}
+
+	/** Returns the color this filter button will be when displayed as a button */
+	virtual FLinearColor GetColor() const override
+	{
+		return Color;
+	}
+
+	/** Set the color this filter button will be when displayed as a button */
+	virtual void SetColor(const FLinearColor& InColor)
+	{
+		Color = InColor;
+	}
+
+	/** Returns the name of the icon to use in menu entries */
+	virtual FName GetIconName() const override
+	{
+		return NAME_None;
+	}
+
+	/** Returns true if the filter should be in the list when disabled and not in the list when enabled */
+	virtual bool IsInverseFilter() const override
+	{
+		return false;
+	}
+
+	virtual bool PassesFilter( FilterType InItem ) const override
+	{
+		return TextFilter->PassesFilter(InItem);
+	}
+
+	/** Get the actual text this filter is using to test against */
+	virtual FText GetFilterString() const
+	{
+		return TextFilter->GetRawFilterText();
+	}
+
+	/** Set the actual text this filter is using to test against */
+	virtual void SetFilterString(const FText& InFilterString)
+	{
+		TextFilter->SetRawFilterText(InFilterString);
+	}
+
+	/** All FCustomTextFilters have the same internal name, this is a helper function to get that name to test against */
+	static FName GetFilterTypeName()
+	{
+		return FName("CustomTextFilter");
+	}
+
+	// ICustomTextFilter interface
+	
+	/** Set the internals of this filter from an FCustomTextFilterData */
+	virtual void SetFromCustomTextFilterData(const FCustomTextFilterData& InFilterData) override
+	{
+		SetDisplayName(InFilterData.FilterLabel);
+		SetColor(InFilterData.FilterColor);
+		SetFilterString(InFilterData.FilterString);
+	}
+
+	/** Create an FCustomTextFilterData from the internals of this filter */
+	virtual FCustomTextFilterData CreateCustomTextFilterData() const override
+	{
+		FCustomTextFilterData CustomTextFilterData;
+		CustomTextFilterData.FilterLabel = GetDisplayName();
+		CustomTextFilterData.FilterColor = GetColor();
+		CustomTextFilterData.FilterString = GetFilterString();
+
+		return CustomTextFilterData;
+	}
+	
+	/** Get the actual filter */
+	virtual TSharedPtr<FFilterBase<FilterType>> GetFilter() override
+	{
+		return this->AsShared();
+	}
+	
+	// Functionality that is not needed by this filter
+	
+	/** Notification that the filter became active or inactive */
+	virtual void ActiveStateChanged(bool bActive) override
+	{
+		
+	}
+
+	/** Called when the right-click context menu is being built for this filter */
+	virtual void ModifyContextMenu(FMenuBuilder& MenuBuilder) override
+	{
+		
+	}
+	
+	/** Can be overriden for custom FilterBar subclasses to save settings, currently not implemented in any generic Filter Bar */
+	virtual void SaveSettings(const FString& IniFilename, const FString& IniSection, const FString& SettingsString) const override
+	{
+		
+	}
+
+	/** Can be overriden for custom FilterBar subclasses to load settings, currently not implemented in any generic Filter Bar */
+	virtual void LoadSettings(const FString& IniFilename, const FString& IniSection, const FString& SettingsString) override
+	{
+		
+	}
+	
+protected:
+
+	/* The actual Text Filter containing information about the text being tested against */
+	TSharedPtr<TTextFilter<FilterType>> TextFilter;
+
+	/* The Display Name of this custom filter that the user sees */
+	FText DisplayName;
+
+	/* The Color of this filter pill */
+	FLinearColor Color;
+};
\ No newline at end of file
diff --git a/Engine/Source/Developer/ToolWidgets/Public/Filters/SBasicFilterBar.h b/Engine/Source/Developer/ToolWidgets/Public/Filters/SBasicFilterBar.h
--- a/Engine/Source/Developer/ToolWidgets/Public/Filters/SBasicFilterBar.h
+++ b/Engine/Source/Developer/ToolWidgets/Public/Filters/SBasicFilterBar.h
@@ -68,7 +68,7 @@ public:
 	DECLARE_DELEGATE_OneParam( FOnExtendAddFilterMenu, UToolMenu* );
 	
 	/** Delegate to create a TTextFilter used to compare FilterType with text queries */
-    DECLARE_DELEGATE_RetVal( TSharedPtr<TTextFilter<FilterType>>, FCreateTextFilter);
+    DECLARE_DELEGATE_RetVal( TSharedPtr<ICustomTextFilter<FilterType>>, FCreateTextFilter);
 	
  	SLATE_BEGIN_ARGS( SBasicFilterBar<FilterType> ){}
 
@@ -81,14 +81,15 @@ public:
 		/** Initial List of Custom Filters that will be added to the AddFilter Menu */
 		SLATE_ARGUMENT( TArray<TSharedRef<FFilterBase<FilterType>>>, CustomFilters)
 
-		/** A delegate to create a TTextFilter for FilterType items. If provided, will allow creation of custom text filters
-		 *  from the filter dropdown menu.
+		/** A delegate to create a Text Filter for FilterType items. If provided, will allow creation of custom
+		 * text filters from the filter dropdown menu.
+		 * @see FCustomTextFilter
 		 */
 		SLATE_ARGUMENT(FCreateTextFilter, CreateTextFilter)
 		
 		/** An SFilterSearchBox that can be attached to this filter bar. When provided along with a CreateTextFilter
 		 *  delegate, allows the user to save searches from the Search Box as text filters for the filter bar.
-		 *	NOTE: Will bind a delegate to SFilterSearchBox::OnClickedAddSearchHistoryButton
+		 *	NOTE: Will bind a delegate to SFilterSearchBox::OnSaveSearchClicked
 		 */
         SLATE_ARGUMENT(TSharedPtr<SFilterSearchBox>, FilterSearchBox)
 	
@@ -152,15 +153,15 @@ public:
 				];
  	}
 
-	/** Attach an SFilterSearchBox to this filter bar, overriding theSFilterSearchBox:: OnClickedAddSearchHistoryButton
+	/** Attach an SFilterSearchBox to this filter bar, overriding the SFilterSearchBox:: OnSaveSearchClicked
 	 *  event to save a search as a filter
 	 */
 	void AttachFilterSearchBox(TSharedPtr<SFilterSearchBox> InFilterSearchBox)
  	{
 		if(InFilterSearchBox)
 		{
-			InFilterSearchBox->SetOnClickedAddSearchHistoryButton(
-				 SFilterSearchBox::FOnClickedAddSearchHistoryButton::CreateSP(this, &SBasicFilterBar<FilterType>::CreateCustomTextFilterFromSearch));
+			InFilterSearchBox->SetOnSaveSearchHandler(
+				 SFilterSearchBox::FOnSaveSearchClicked::CreateSP(this, &SBasicFilterBar<FilterType>::CreateCustomTextFilterFromSearch));
 		}
  	}
 
@@ -1080,7 +1081,7 @@ protected:
  	}
  	
  	/** Handler for when a checkbox next to a custom text filter is clicked */
-    void CustomTextFilterClicked(ECheckBoxState CheckBoxState, TSharedRef<FCustomTextFilter<FilterType>> FrontendFilter)
+	void CustomTextFilterClicked(ECheckBoxState CheckBoxState, TSharedRef<FFilterBase<FilterType>> FrontendFilter)
     {
     	if (CheckBoxState == ECheckBoxState::Unchecked)
     	{
@@ -1094,30 +1095,36 @@ protected:
     }
 
 	/** Handler for when a custom text filter is created */
-	void OnCreateCustomTextFilter(const FCustomTextFilterData& InFilterData, bool bApplyFilter)
+	virtual void OnCreateCustomTextFilter(const FCustomTextFilterData& InFilterData, bool bApplyFilter)
 	{
-		// Create a TTextFilter for the current filter we are making using the provided delegate
-		TSharedRef<TTextFilter<FilterType>> NewTextFilter = CreateTextFilter.Execute().ToSharedRef();
+		// Create the text filter using the given delegate
+		TSharedRef<ICustomTextFilter<FilterType>> NewTextFilter = CreateTextFilter.Execute().ToSharedRef();
 		
-		// Create the actual text filter
-		TSharedRef<FCustomTextFilter<FilterType>> NewFilter = MakeShared<FCustomTextFilter<FilterType>>(NewTextFilter);
-
 		// Fill in the data the widget gives us
-		NewFilter->SetFromCustomTextFilterData(InFilterData);
+		NewTextFilter->SetFromCustomTextFilterData(InFilterData);
+
+		TSharedPtr<FFilterBase<FilterType>> NewFilter = NewTextFilter->GetFilter();
 
-		CustomTextFilters.Add(NewFilter);
+		if(!NewFilter)
+		{
+			// Close the create custom filter dialog 		
+			OnCancelCustomTextFilterDialog();
+			return; // TODO: error message? Some way to ensure the user always implements NewTextFilter->GetFilter()?
+		}
 		
-		TSharedRef<SFilter> AddedFilter = AddFilterToBar(NewFilter);
+		CustomTextFilters.Add(NewTextFilter);
+		
+		TSharedRef<SFilter> AddedFilter = AddFilterToBar(NewFilter.ToSharedRef());
 		
 		AddedFilter->SetEnabled(bApplyFilter);
-		SetFrontendFilterActive(NewFilter, bApplyFilter);
+		SetFrontendFilterActive(NewFilter.ToSharedRef(), bApplyFilter);
 
 		// Close the create custom filter dialog 		
 		OnCancelCustomTextFilterDialog();
 	}
 
 	/** Handler for when a custom text filter is modified */
-	void OnModifyCustomTextFilter(const FCustomTextFilterData& InFilterData, TSharedPtr<FCustomTextFilter<FilterType>> InFilter)
+	virtual void OnModifyCustomTextFilter(const FCustomTextFilterData& InFilterData, TSharedPtr<ICustomTextFilter<FilterType>> InFilter)
 	{
 		// Update the filter with the data the widget has given us
 		if(InFilter)
@@ -1125,19 +1132,21 @@ protected:
 			InFilter->SetFromCustomTextFilterData(InFilterData);
 		}
 
+		TSharedPtr<FFilterBase<FilterType>> NewFilter = InFilter->GetFilter();
+		
 		// If the filter we modified is active in the bar, remove and re-add it to update the data
 		for (const TSharedRef<SFilter>& Filter : Filters)
 		{
-			if (Filter->GetFrontendFilter() == InFilter)
+			if (Filter->GetFrontendFilter() == NewFilter)
 			{
 				bool bWasEnabled = Filter->IsEnabled();
 				
-				RemoveFilter(InFilter.ToSharedRef());
+				RemoveFilter(NewFilter.ToSharedRef());
 
-				TSharedRef<SFilter> AddedFilter = AddFilterToBar(InFilter.ToSharedRef());
+				TSharedRef<SFilter> AddedFilter = AddFilterToBar(NewFilter.ToSharedRef());
 				
 				AddedFilter->SetEnabled(bWasEnabled);
-				SetFrontendFilterActive(InFilter.ToSharedRef(), bWasEnabled);
+				SetFrontendFilterActive(NewFilter.ToSharedRef(), bWasEnabled);
 			}
 		}
 
@@ -1145,12 +1154,15 @@ protected:
 	}
 
 	/** Handler for when a custom text filter is deleted */
-	void OnDeleteCustomTextFilter(const TSharedPtr<FCustomTextFilter<FilterType>> InFilter)
+	virtual void OnDeleteCustomTextFilter(const TSharedPtr<ICustomTextFilter<FilterType>> InFilter)
 	{
 		if(InFilter)
 		{
 			CustomTextFilters.RemoveSingle(InFilter.ToSharedRef());
-			RemoveFilter(InFilter.ToSharedRef());
+
+			TSharedPtr<FFilterBase<FilterType>> NewFilter = InFilter->GetFilter();
+			
+			RemoveFilter(NewFilter.ToSharedRef());
 
 			// We fire this delegate here, since a CustomTextFilter was removed from the list 
 			OnFilterChanged.ExecuteIfBound();
@@ -1160,7 +1172,7 @@ protected:
 	}
 
 	/** Handler to close the custom text filter dialog */
-	void OnCancelCustomTextFilterDialog()
+	virtual void OnCancelCustomTextFilterDialog()
 	{
 		if(CustomTextFilterWindow.IsValid())
 		{
@@ -1168,8 +1180,17 @@ protected:
 		}
 	}
 
+	void GetAllCustomTextFilterLabels(TArray<FText>& OutFilterLabels)
+	{
+		for (const TSharedRef<ICustomTextFilter<FilterType>>& Filter : CustomTextFilters)
+		{
+			FCustomTextFilterData FilterData = Filter->CreateCustomTextFilterData();
+			OutFilterLabels.Add(FilterData.FilterLabel);
+		}
+	}
+
 	/** Creates a dialog box with the SCustomTextFilterDialog Widget */
-	void CreateCustomTextFilterWindow(const FCustomTextFilterData& CustomTextFilterData, TSharedPtr<FCustomTextFilter<FilterType>> InFilter)
+	void CreateCustomTextFilterWindow(const FCustomTextFilterData& CustomTextFilterData, TSharedPtr<ICustomTextFilter<FilterType>> InFilter)
 	{
 		/** If we already have a window, delete it */
 		OnCancelCustomTextFilterDialog();
@@ -1184,6 +1205,7 @@ protected:
 			.HasCloseButton(true)
 			.SupportsMaximize(false)
 			.SupportsMinimize(false)
+			.SizingRule(ESizingRule::FixedSize)
 			.ClientSize(FVector2D(724, 183));
 		
 		TSharedPtr<SCustomTextFilterDialog> CustomTextFilterDialog =
@@ -1193,7 +1215,8 @@ protected:
 			.OnCreateFilter(this, &SBasicFilterBar<FilterType>::OnCreateCustomTextFilter)
 			.OnDeleteFilter(this, &SBasicFilterBar<FilterType>::OnDeleteCustomTextFilter, InFilter)
 			.OnModifyFilter(this, &SBasicFilterBar<FilterType>::OnModifyCustomTextFilter, InFilter)
-			.OnCancelClicked(this, &SBasicFilterBar<FilterType>::OnCancelCustomTextFilterDialog);
+			.OnCancelClicked(this, &SBasicFilterBar<FilterType>::OnCancelCustomTextFilterDialog)
+			.OnGetFilterLabels(this, &SBasicFilterBar<FilterType>::GetAllCustomTextFilterLabels);
 
 		NewTextFilterWindow->SetContent(CustomTextFilterDialog.ToSharedRef());
 		FSlateApplication::Get().AddWindow(NewTextFilterWindow.ToSharedRef());
@@ -1208,7 +1231,7 @@ protected:
 	}
 
 	/** Creates a dialog box to Edit an existing custom text filter */
-	void CreateEditCustomTextFilterWindow(TSharedPtr<FCustomTextFilter<FilterType>> InFilter)
+	void CreateEditCustomTextFilterWindow(TSharedPtr<ICustomTextFilter<FilterType>> InFilter)
 	{
 		if(!InFilter)
 		{
@@ -1247,8 +1270,10 @@ protected:
 
 		FToolMenuSection& CustomFiltersSection = InMenu->AddSection("FilterBarCustomTextFiltersSection");
 		
-		for (const TSharedRef<FCustomTextFilter<FilterType>>& FrontendFilter : CustomTextFilters)
+		for (const TSharedRef<ICustomTextFilter<FilterType>>& Filter : CustomTextFilters)
 		{
+			const TSharedRef<FFilterBase<FilterType>>& FrontendFilter = Filter->GetFilter().ToSharedRef();
+			
 			bool bIsFilterChecked = IsFrontendFilterInUse(FrontendFilter);
 			
 			TSharedPtr<SHorizontalBox> CustomTextFilterWidget = SNew(SHorizontalBox);
@@ -1288,9 +1313,9 @@ protected:
 			[
 				SNew(SSimpleButton)
 				.Icon(FAppStyle::GetBrush("Icons.Edit"))
-				.OnClicked_Lambda([this, &FrontendFilter]()
+				.OnClicked_Lambda([this, &Filter]()
 				{
-					this->CreateEditCustomTextFilterWindow(FrontendFilter);
+					this->CreateEditCustomTextFilterWindow(Filter);
 					return FReply::Handled();
 				})
 			];
@@ -1417,7 +1442,7 @@ protected:
 	FCreateTextFilter CreateTextFilter;
 	
 	/** Custom text filters that a user can create */
-	TArray< TSharedRef< FCustomTextFilter<FilterType> > > CustomTextFilters;
+	TArray< TSharedRef< ICustomTextFilter<FilterType> > > CustomTextFilters;
 
 	/** The window containing the custom text filter dialog */
 	TWeakPtr<SWindow> CustomTextFilterWindow;
diff --git a/Engine/Source/Developer/ToolWidgets/Public/Filters/SCustomTextFilterDialog.h b/Engine/Source/Developer/ToolWidgets/Public/Filters/SCustomTextFilterDialog.h
--- a/Engine/Source/Developer/ToolWidgets/Public/Filters/SCustomTextFilterDialog.h
+++ b/Engine/Source/Developer/ToolWidgets/Public/Filters/SCustomTextFilterDialog.h
@@ -3,180 +3,11 @@
 #pragma once
 
 #include "Widgets/SCompoundWidget.h"
-#include "Filters/FilterBase.h"
-#include "Core/Public/Misc/TextFilter.h"
-
-#include "SCustomTextFilterDialog.generated.h"
+#include "Filters/CustomTextFilters.h"
 
 class SColorBlock;
 class SEditableTextBox;
 
-/* Struct containing the data that SCustomTextFilterDialog is currently editing */
-USTRUCT()
-struct FCustomTextFilterData
-{
-	GENERATED_BODY()
-
-	UPROPERTY()
-	FText FilterLabel;
-	
-	UPROPERTY()
-	FText FilterString;
-
-	UPROPERTY()
-	FLinearColor FilterColor;
-
-	FCustomTextFilterData()
-		: FilterColor(FLinearColor::White)
-	{
-		
-	}
-};
-
-/* A filter that allows testing items of type FilterType against text expressions */
-template<typename FilterType>
-class FCustomTextFilter : public FFilterBase<FilterType>
-{
-public:
-
-	FCustomTextFilter(TSharedPtr<TTextFilter<FilterType>> InTextFilter)
-	: FFilterBase<FilterType>(nullptr)
-	, TextFilter(InTextFilter)
-	{
-		
-	}
-
-	/** Returns the system name for this filter */
-	virtual FString GetName() const override
-	{
-		return GetCommonName().ToString();
-	}
-
-	/** Returns the human readable name for this filter */
-	virtual FText GetDisplayName() const override
-	{
-		return DisplayName; 
-	}
-
-	/** Set the human readable name for this filter */
-	void SetDisplayName(const FText& InDisplayName)
-	{
-		DisplayName = InDisplayName;
-	}
-
-	/** Returns the tooltip for this filter, shown in the filters menu */
-	virtual FText GetToolTipText() const override
-	{
-		return TextFilter->GetRawFilterText(); // The tooltip will display the filter string
-	}
-
-	/** Returns the color this filter button will be when displayed as a button */
-	virtual FLinearColor GetColor() const override
-	{
-		return Color;
-	}
-
-	/** Set the color this filter button will be when displayed as a button */
-	void SetColor(const FLinearColor& InColor)
-	{
-		Color = InColor;
-	}
-
-	/** Returns the name of the icon to use in menu entries */
-	virtual FName GetIconName() const override
-	{
-		return NAME_None;
-	}
-
-	/** Returns true if the filter should be in the list when disabled and not in the list when enabled */
-	virtual bool IsInverseFilter() const override
-	{
-		return false;
-	}
-
-	virtual bool PassesFilter( FilterType InItem ) const override
-	{
-		return TextFilter->PassesFilter(InItem);
-	}
-
-	/** Get the actual text this filter is using to test against */
-	FText GetFilterString() const
-	{
-		return TextFilter->GetRawFilterText();
-	}
-
-	/** Set the actual text this filter is using to test against */
-	void SetFilterString(const FText& InFilterString)
-	{
-		TextFilter->SetRawFilterText(InFilterString);
-	}
-
-	/** All FCustomTextFilters have the same internal name, this is a helper function to get that name to test against */
-	static FName GetCommonName()
-	{
-		return FName("CustomTextFilter");
-	}
-
-	/** Set the internals of this filter from an FCustomTextFilterData */
-	void SetFromCustomTextFilterData(const FCustomTextFilterData& InFilterData)
-	{
-		SetDisplayName(InFilterData.FilterLabel);
-		SetColor(InFilterData.FilterColor);
-		SetFilterString(InFilterData.FilterString);
-	}
-
-	/** Create an FCustomTextFilterData from the internals of this filter */
-	FCustomTextFilterData CreateCustomTextFilterData() const
-	{
-		FCustomTextFilterData CustomTextFilterData;
-		CustomTextFilterData.FilterLabel = GetDisplayName();
-		CustomTextFilterData.FilterColor = GetColor();
-		CustomTextFilterData.FilterString = GetFilterString();
-
-		return CustomTextFilterData;
-	}
-	
-	// Functionality that is not needed by this filter
-	
-	/** Notification that the filter became active or inactive */
-	virtual void ActiveStateChanged(bool bActive) override
-	{
-		
-	}
-
-	/** Called when the right-click context menu is being built for this filter */
-	virtual void ModifyContextMenu(FMenuBuilder& MenuBuilder) override
-	{
-		
-	}
-	
-	/** Can be overriden for custom FilterBar subclasses to save settings, currently not implemented in any gneeric Filter Bar */
-	virtual void SaveSettings(const FString& IniFilename, const FString& IniSection, const FString& SettingsString) const override
-	{
-		
-	}
-
-	/** Can be overriden for custom FilterBar subclasses to load settings, currently not implemented in any gneeric Filter Bar */
-	virtual void LoadSettings(const FString& IniFilename, const FString& IniSection, const FString& SettingsString) override
-	{
-		
-	}
-	
-protected:
-
-	/* The actual Text Filter containing information about the text being tested against */
-	TSharedPtr<TTextFilter<FilterType>> TextFilter;
-
-	/* The Display Name of this custom filter that the user sees */
-	FText DisplayName;
-
-	/* The Color of this filter pill */
-	FLinearColor Color;
-};
-
-
-
-
 class TOOLWIDGETS_API SCustomTextFilterDialog : public SCompoundWidget
 {
 public:
@@ -185,6 +16,7 @@ public:
 	DECLARE_DELEGATE(FOnDeleteFilter);
 	DECLARE_DELEGATE_OneParam(FOnModifyFilter, const FCustomTextFilterData& /* InFilterData */);
 	DECLARE_DELEGATE(FOnCancelClicked);
+	DECLARE_DELEGATE_OneParam(FOnGetFilterLabels, TArray<FText> & /* FilterNames */);
 
 	SLATE_BEGIN_ARGS(SCustomTextFilterDialog) {}
 	
@@ -205,6 +37,9 @@ public:
 
 		/** Delegate for when the Modify Filter button is clicked */
         SLATE_EVENT(FOnModifyFilter, OnModifyFilter)
+
+		/** Delegate to get all existing filter labels to check for duplicates */
+		SLATE_EVENT(FOnGetFilterLabels, OnGetFilterLabels)
     
     SLATE_END_ARGS()
     	
@@ -234,6 +69,9 @@ protected:
 	/* The current filter data we are editing */
 	FCustomTextFilterData FilterData;
 	
+	/* The initial, unedited filter data we were provided */
+	FCustomTextFilterData InitialFilterData;
+	
 	FOnCreateFilter OnCreateFilter;
 	
 	FOnDeleteFilter OnDeleteFilter;
@@ -242,6 +80,8 @@ protected:
 
 	FOnModifyFilter OnModifyFilter;
 
+	FOnGetFilterLabels OnGetFilterLabels;
+
 	/* The color block widget that edits the filter color */
 	TSharedPtr<SColorBlock> ColorBlock;
 
diff --git a/Engine/Source/Developer/ToolWidgets/Public/Filters/SFilterSearchBox.h b/Engine/Source/Developer/ToolWidgets/Public/Filters/SFilterSearchBox.h
--- a/Engine/Source/Developer/ToolWidgets/Public/Filters/SFilterSearchBox.h
+++ b/Engine/Source/Developer/ToolWidgets/Public/Filters/SFilterSearchBox.h
@@ -10,7 +10,7 @@ class SMenuAnchor;
 class SSearchBox;
 template <typename T> class SListView;
 
-/** A SearchBox widget that contains support for displaying search history.
+/** A SearchBox widget that contains support for displaying search history
  * When used with a FilterBar widget (@see SBasicFilterBar or SFilterBar), it allows the user to save searches from the
  * history as filter pills
  */
@@ -18,14 +18,16 @@ class TOOLWIDGETS_API SFilterSearchBox : public SCompoundWidget
 {
 public:
 
-	DECLARE_DELEGATE_OneParam(FOnClickedAddSearchHistoryButton, const FText&)
+	DECLARE_DELEGATE_OneParam(FOnSaveSearchClicked, const FText&)
 
 	SLATE_BEGIN_ARGS(SFilterSearchBox)
 	: _OnTextChanged()
 	, _OnTextCommitted()
 	, _InitialText()
 	, _HintText()
+	, _ShowSearchHistory(true)
 	, _MaxSearchHistory(5)
+	, _DelayChangeNotificationsWhileTyping( false )
 	{}
 
 	/** Invoked whenever the text changes */
@@ -46,6 +48,15 @@ public:
 	/** The maximum number of items to show in the Search History */
 	SLATE_ARGUMENT(int32, MaxSearchHistory)
 
+	/** Whether the SearchBox should delay notifying listeners of text changed events until the user is done typing */
+	SLATE_ATTRIBUTE( bool, DelayChangeNotificationsWhileTyping )
+
+	/** Callback delegate to have first chance handling of the OnKeyDown event */
+	SLATE_EVENT(FOnKeyDown, OnKeyDownHandler)
+
+	/** Handler for when the + Button next to a search is clicked */
+	SLATE_EVENT(FOnSaveSearchClicked, OnSaveSearchClicked)
+
 	SLATE_END_ARGS()
 
 	/** Constructs this widget with InArgs */
@@ -63,11 +74,10 @@ public:
 
 	/** SWidget interface */
 	void OnFocusLost(const FFocusEvent& InFocusEvent) override;
+	virtual bool HasKeyboardFocus() const override;
 
-	void SetOnClickedAddSearchHistoryButton(FOnClickedAddSearchHistoryButton InOnClickedAddSearchHistoryButton)
-	{
-		OnClickedAddSearchHistoryButton = InOnClickedAddSearchHistoryButton;
-	}
+	/** Show a + button next to the current search, and set the handler for when that is clicked */
+	void SetOnSaveSearchHandler(FOnSaveSearchClicked InOnSaveSearchHandler);
 
 private:
 	
@@ -89,8 +99,14 @@ private:
 	/** Handler for when the chevron to open the search history is clicked */
 	FReply OnClickedSearchHistory();
 
+	/** Whether the search history dropdown chevron is visible */
+	EVisibility GetSearchHistoryVisibility() const;
+
 private:
 
+	/* Whether we are showing the search history */
+	TAttribute<bool> bShowSearchHistory;
+	
 	/** The actual Search Box */
 	TSharedPtr<SSearchBox> SearchBox;
 	
@@ -101,7 +117,7 @@ private:
 	FOnTextCommitted OnTextCommitted;
 
 	/** Delegate for when the Plus icon next to an item in the search history is clicked */
-	FOnClickedAddSearchHistoryButton OnClickedAddSearchHistoryButton;
+	FOnSaveSearchClicked OnSaveSearchClicked;
 
 	/** The max amount of items to show in the Search History */
 	int32 MaxSearchHistory;
diff --git a/Engine/Source/Editor/ContentBrowser/Private/SAssetPicker.cpp b/Engine/Source/Editor/ContentBrowser/Private/SAssetPicker.cpp
--- a/Engine/Source/Editor/ContentBrowser/Private/SAssetPicker.cpp
+++ b/Engine/Source/Editor/ContentBrowser/Private/SAssetPicker.cpp
@@ -235,6 +235,7 @@ void SAssetPicker::Construct( const FArguments& InArgs )
 			.OnFilterChanged(this, &SAssetPicker::OnFilterChanged)
 			.FrontendFilters(FrontendFilters)
 			.InitialClassFilters(FilterClassList)
+			.FilterBarIdentifier(FName(SaveSettingsName))
 			.ExtraFrontendFilters(InArgs._AssetPickerConfig.ExtraFrontendFilters)
 		];
 
@@ -701,7 +702,7 @@ void SAssetPicker::LoadSettings()
 		// Load all our data using the settings string as a key in the user settings ini
 		if (FilterListPtr.IsValid())
 		{
-			FilterListPtr->LoadSettings(GEditorPerProjectIni, SContentBrowser::SettingsIniSection, SettingsString);
+			FilterListPtr->LoadSettings();
 		}
 		
 		AssetViewPtr->LoadSettings(GEditorPerProjectIni, SContentBrowser::SettingsIniSection, SettingsString);
@@ -717,7 +718,7 @@ void SAssetPicker::SaveSettings() const
 		// Save all our data using the settings string as a key in the user settings ini
 		if (FilterListPtr.IsValid())
 		{
-			FilterListPtr->SaveSettings(GEditorPerProjectIni, SContentBrowser::SettingsIniSection, SettingsString);
+			FilterListPtr->SaveSettings();
 		}
 
 		AssetViewPtr->SaveSettings(GEditorPerProjectIni, SContentBrowser::SettingsIniSection, SettingsString);
diff --git a/Engine/Source/Editor/ContentBrowser/Private/SCollectionView.cpp b/Engine/Source/Editor/ContentBrowser/Private/SCollectionView.cpp
--- a/Engine/Source/Editor/ContentBrowser/Private/SCollectionView.cpp
+++ b/Engine/Source/Editor/ContentBrowser/Private/SCollectionView.cpp
@@ -777,7 +777,7 @@ FReply SCollectionView::OnDrop( const FGeometry& MyGeometry, const FDragDropEven
 	return FReply::Unhandled();
 }
 
-void SCollectionView::MakeSaveDynamicCollectionMenu(FText InQueryString)
+void SCollectionView::MakeSaveDynamicCollectionMenu(FText InQueryString, FSimpleDelegate OnSaveSearchClicked)
 {
 	// Get all menu extenders for this context menu from the content browser module
 	FContentBrowserModule& ContentBrowserModule = FModuleManager::GetModuleChecked<FContentBrowserModule>( TEXT("ContentBrowser") );
@@ -797,6 +797,11 @@ void SCollectionView::MakeSaveDynamicCollectionMenu(FText InQueryString)
 
 	CollectionContextMenu->UpdateProjectSourceControl();
 
+	if(OnSaveSearchClicked.IsBound())
+	{
+		MakeSaveSearchMenu(MenuBuilder, OnSaveSearchClicked);
+	}
+
 	CollectionContextMenu->MakeSaveDynamicCollectionSubMenu(MenuBuilder, InQueryString);
 
 	FWidgetPath WidgetPath;
@@ -812,6 +817,20 @@ void SCollectionView::MakeSaveDynamicCollectionMenu(FText InQueryString)
 	}
 }
 
+void SCollectionView::MakeSaveSearchMenu(FMenuBuilder& InMenuBuilder, FSimpleDelegate OnSaveSearchClicked) const
+{
+	InMenuBuilder.BeginSection("ContentBrowserSaveSearch", LOCTEXT("ContentBrowserCreateFilterMenuHeading", "Create Filter"));
+
+	InMenuBuilder.AddMenuEntry(
+		LOCTEXT("ContentBrowserSaveAsCustomFilter", "Save as Custom Filter"),
+		LOCTEXT("ContentBrowserSaveAsCustomFilterTooltip", "Save the current search text as a custom filter in the filter bar"),
+		FSlateIcon(),
+		FUIAction(OnSaveSearchClicked)
+	);
+	
+	InMenuBuilder.EndSection();
+}
+
 FReply SCollectionView::OnAddCollectionClicked()
 {
 	MakeAddCollectionMenu(AsShared());
diff --git a/Engine/Source/Editor/ContentBrowser/Private/SCollectionView.h b/Engine/Source/Editor/ContentBrowser/Private/SCollectionView.h
--- a/Engine/Source/Editor/ContentBrowser/Private/SCollectionView.h
+++ b/Engine/Source/Editor/ContentBrowser/Private/SCollectionView.h
@@ -102,7 +102,7 @@ public:
 	virtual FReply OnDrop( const FGeometry& MyGeometry, const FDragDropEvent& DragDropEvent ) override;
 
 	/** Creates the menu for the save dynamic collection button */
-	void MakeSaveDynamicCollectionMenu(FText InQueryString);
+	void MakeSaveDynamicCollectionMenu(FText InQueryString, FSimpleDelegate OnSaveSearchClicked = FSimpleDelegate());
 
 	/** Creates the menu for the add collection button */
 	FReply OnAddCollectionClicked();
@@ -257,6 +257,9 @@ private:
 
 	/** Get the active filter text */
 	FText GetCollectionsSearchFilterText() const;
+
+	/** Make the menu to save a search using the given delegate */
+	void MakeSaveSearchMenu(FMenuBuilder& InMenuBuilder, FSimpleDelegate OnSaveSearchClicked) const;
 private:
 
 	/** A helper class to manage PreventSelectionChangedDelegateCount by incrementing it when constructed (on the stack) and decrementing when destroyed */
diff --git a/Engine/Source/Editor/ContentBrowser/Private/SContentBrowser.cpp b/Engine/Source/Editor/ContentBrowser/Private/SContentBrowser.cpp
--- a/Engine/Source/Editor/ContentBrowser/Private/SContentBrowser.cpp
+++ b/Engine/Source/Editor/ContentBrowser/Private/SContentBrowser.cpp
@@ -139,6 +139,7 @@ void SContentBrowser::Construct( const FArguments& InArgs, const FName& InInstan
  	bIsLocked = InArgs._InitiallyLocked;
 	bCanSetAsPrimaryBrowser = Config != nullptr ? Config->bCanSetAsPrimaryBrowser : true;
 	bIsDrawer = InArgs._IsDrawer;
+	this->InstanceName = InInstanceName;
 
 	HistoryManager.SetOnApplyHistoryData(FOnApplyHistoryData::CreateSP(this, &SContentBrowser::OnApplyHistoryData));
 	HistoryManager.SetOnUpdateHistoryData(FOnUpdateHistoryData::CreateSP(this, &SContentBrowser::OnUpdateHistoryData));
@@ -487,7 +488,6 @@ void SContentBrowser::Construct( const FArguments& InArgs, const FName& InInstan
 	HistoryManager.AddHistoryData();
 
 	// Load settings if they were specified
-	this->InstanceName = InInstanceName;
 	LoadSettings(InInstanceName);
 
 	if( Config != nullptr )
@@ -654,9 +654,16 @@ void SContentBrowser::TogglePrivateContentEdit(const TArray<FString>& FolderPath
 
 void SContentBrowser::HandleAssetViewSearchOptionsChanged()
 {
-	TextFilter->SetIncludeClassName(AssetViewPtr->IsIncludingClassNames());
-	TextFilter->SetIncludeAssetPath(AssetViewPtr->IsIncludingAssetPaths());
-	TextFilter->SetIncludeCollectionNames(AssetViewPtr->IsIncludingCollectionNames());
+	bool bIncludeClassName = AssetViewPtr->IsIncludingClassNames();
+	bool bIncludeAssetPath = AssetViewPtr->IsIncludingAssetPaths();
+	bool bIncludeCollectionNames = AssetViewPtr->IsIncludingCollectionNames();
+	
+	TextFilter->SetIncludeClassName(bIncludeClassName);
+	TextFilter->SetIncludeAssetPath(bIncludeAssetPath);
+	TextFilter->SetIncludeCollectionNames(bIncludeCollectionNames);
+
+	// Make sure all custom text filters get the updated Asset View Search Options
+	FilterListPtr->UpdateCustomTextFilterIncludes(bIncludeClassName, bIncludeAssetPath, bIncludeCollectionNames);
 }
 
 TSharedRef<SWidget> SContentBrowser::CreateToolBar(const FContentBrowserConfig* Config)
@@ -708,8 +715,19 @@ TSharedRef<SWidget> SContentBrowser::CreateAssetView(const FContentBrowserConfig
 					.OnGetContextMenu(this, &SContentBrowser::GetFilterContextMenu)
 					.Visibility((Config != nullptr ? Config->bCanShowFilters : true) ? EVisibility::Visible : EVisibility::Collapsed)
 					.FrontendFilters(FrontendFilters)
-					.AddMetaData<FTagMetaData>(FTagMetaData(TEXT("ContentBrowserFilters")));
+					.FilterBarIdentifier(InstanceName)
+					.UseSharedSettings(true)
+					.CreateTextFilter(SFilterList::FCreateTextFilter::CreateLambda([this]
+					{
+						TSharedPtr<FFrontendFilter_CustomText> NewFilter = MakeShareable(new FFrontendFilter_CustomText);
+
+						// Make sure the new filter has the right search options from the AssetView. We only have to set it once, SFilterList handles syncing it on change
+						NewFilter->UpdateCustomTextFilterIncludes(AssetViewPtr->IsIncludingClassNames(), AssetViewPtr->IsIncludingAssetPaths(), AssetViewPtr->IsIncludingCollectionNames());
 
+						return NewFilter;
+					}))
+					.AddMetaData<FTagMetaData>(FTagMetaData(TEXT("ContentBrowserFilters")));
+	
 	// Create the Filter Combo Button
 	TSharedPtr<SWidget> FilterComboButton = SFilterList::MakeAddFilterButton(FilterListPtr.ToSharedRef());
 	
@@ -750,36 +768,18 @@ TSharedRef<SWidget> SContentBrowser::CreateAssetView(const FContentBrowserConfig
 						[
 							SAssignNew(SearchBoxPtr, SAssetSearchBox)
 							.HintText(this, &SContentBrowser::GetSearchAssetsHintText)
+							.ShowSearchHistory(true)
 							.OnTextChanged(this, &SContentBrowser::OnSearchBoxChanged)
 							.OnTextCommitted(this, &SContentBrowser::OnSearchBoxCommitted)
 							.OnKeyDownHandler(this, &SContentBrowser::OnSearchKeyDown)
+							.OnSaveSearchClicked(this, &SContentBrowser::OnSaveSearchButtonClicked)
 							.OnAssetSearchBoxSuggestionFilter(this, &SContentBrowser::OnAssetSearchSuggestionFilter)
 							.OnAssetSearchBoxSuggestionChosen(this, &SContentBrowser::OnAssetSearchSuggestionChosen)
 							.DelayChangeNotificationsWhileTyping(true)
 							.Visibility((Config != nullptr ? Config->bCanShowAssetSearch : true) ? EVisibility::Visible : EVisibility::Collapsed)
 							.AddMetaData<FTagMetaData>(FTagMetaData(TEXT("ContentBrowserSearchAssets")))
 						]
-					
-						// Save Search
-						+ SHorizontalBox::Slot()
-						.AutoWidth()
-						.VAlign(VAlign_Center)
-						.Padding(2.0f, 0.0f, 0.0f, 0.0f)
-						[
-							SNew(SButton)
-							.ButtonStyle(FAppStyle::Get(), "SimpleButton")
-							.ToolTipText(LOCTEXT("SaveSearchButtonTooltip", "Save the current search as a dynamic collection."))
-							.IsEnabled(this, &SContentBrowser::IsSaveSearchButtonEnabled)
-							.OnClicked(this, &SContentBrowser::OnSaveSearchButtonClicked)
-							.ContentPadding(FMargin(1, 1))
-							.Visibility((!bHideSaveSearchButton && (Config != nullptr ? Config->bCanShowAssetSearch : true)) ? EVisibility::Visible : EVisibility::Collapsed)
-							[
-								SNew(STextBlock)
-								.TextStyle(FAppStyle::Get(), "GenericFilters.TextStyle")
-								.Font(FAppStyle::Get().GetFontStyle("FontAwesome.10"))
-								.Text(FEditorFontGlyphs::Floppy_O)
-							]
-						]
+						
 						// Filter dropdown
 						+ SHorizontalBox::Slot()
 						.AutoWidth()
@@ -1486,7 +1486,7 @@ void SContentBrowser::SaveSettings() const
 
 
 	// Save all our data using the settings string as a key in the user settings ini
-	FilterListPtr->SaveSettings(GEditorPerProjectIni, SettingsIniSection, SettingsString);
+	FilterListPtr->SaveSettings();
 	PathViewPtr->SaveSettings(GEditorPerProjectIni, SettingsIniSection, SettingsString);
 	FavoritePathViewPtr->SaveSettings(GEditorPerProjectIni, SettingsIniSection, SettingsString + TEXT(".Favorites"));
 	CollectionViewPtr->SaveSettings(GEditorPerProjectIni, SettingsIniSection, SettingsString);
@@ -1636,7 +1636,7 @@ void SContentBrowser::LoadSettings(const FName& InInstanceName)
 	}
 
 	// Save all our data using the settings string as a key in the user settings ini
-	FilterListPtr->LoadSettings(GEditorPerProjectIni, SettingsIniSection, SettingsString);
+	FilterListPtr->LoadSettings();
 	PathViewPtr->LoadSettings(GEditorPerProjectIni, SettingsIniSection, SettingsString);
 	FavoritePathViewPtr->LoadSettings(GEditorPerProjectIni, SettingsIniSection, SettingsString + TEXT(".Favorites"));
 	CollectionViewPtr->LoadSettings(GEditorPerProjectIni, SettingsIniSection, SettingsString);
@@ -1964,7 +1964,7 @@ bool SContentBrowser::IsSaveSearchButtonEnabled() const
 	return !TextFilter->GetRawFilterText().IsEmptyOrWhitespace();
 }
 
-FReply SContentBrowser::OnSaveSearchButtonClicked()
+void SContentBrowser::OnSaveSearchButtonClicked(const FText& InSearchText)
 {
 	// Need to make sure we can see the collections view
 	if (!bSourcesViewExpanded)
@@ -2007,8 +2007,12 @@ FReply SContentBrowser::OnSaveSearchButtonClicked()
 		FinalQueryText = FText::FromString(FString::Printf(TEXT("(%s) AND (%s)"), *TextFilter->GetRawFilterText().ToString(), *SelectedPathsQuery));
 	}
 
-	CollectionViewPtr->MakeSaveDynamicCollectionMenu(FinalQueryText);
-	return FReply::Handled();
+	CollectionViewPtr->MakeSaveDynamicCollectionMenu(FinalQueryText, FSimpleDelegate::CreateSP(this, &SContentBrowser::SaveSearchAsFilter));
+}
+
+void SContentBrowser::SaveSearchAsFilter()
+{
+	FilterListPtr->CreateCustomFilterDialog(TextFilter->GetRawFilterText());
 }
 
 void SContentBrowser::OnPathClicked( const FString& CrumbData )
diff --git a/Engine/Source/Editor/ContentBrowser/Private/SContentBrowser.h b/Engine/Source/Editor/ContentBrowser/Private/SContentBrowser.h
--- a/Engine/Source/Editor/ContentBrowser/Private/SContentBrowser.h
+++ b/Engine/Source/Editor/ContentBrowser/Private/SContentBrowser.h
@@ -276,8 +276,11 @@ private:
 	/** Should the "Save Search" button be enabled? */
 	bool IsSaveSearchButtonEnabled() const;
 
-	/** Open the menu to let you save the current search text as a dynamic collection */
-	FReply OnSaveSearchButtonClicked();
+	/** Open the menu to let you save the current search text as a filter or dynamic collection */
+	void OnSaveSearchButtonClicked(const FText& InSearchText);
+
+	/** Save the current search as a filter pill */
+	void SaveSearchAsFilter();
 
 	/** Called when a crumb in the path breadcrumb trail or menu is clicked */
 	void OnPathClicked(const FString& CrumbData);
diff --git a/Engine/Source/Editor/ContentBrowser/Private/SFilterList.cpp b/Engine/Source/Editor/ContentBrowser/Private/SFilterList.cpp
--- a/Engine/Source/Editor/ContentBrowser/Private/SFilterList.cpp
+++ b/Engine/Source/Editor/ContentBrowser/Private/SFilterList.cpp
@@ -39,9 +39,12 @@
 // SFilterList
 /////////////////////
 
+const FName SFilterList::SharedIdentifier("FilterListSharedSettings");
+SFilterList::FCustomTextFilterEvent SFilterList::CustomTextFilterEvent;
 
 void SFilterList::Construct( const FArguments& InArgs )
 {
+	bUseSharedSettings = InArgs._UseSharedSettings;
 	OnGetContextMenu = InArgs._OnGetContextMenu;
 	this->OnFilterChanged = InArgs._OnFilterChanged;
 	this->ActiveFilters = InArgs._FrontendFilters;
@@ -113,8 +116,22 @@ void SFilterList::Construct( const FArguments& InArgs )
 	/** Explicitly setting this to true as it should ALWAYS be true for SFilterList */
 	Args._UseDefaultAssetFilters = true;
 	Args._OnFilterChanged = this->OnFilterChanged;
-	
+	Args._CreateTextFilter = InArgs._CreateTextFilter;
+	Args._FilterBarIdentifier = InArgs._FilterBarIdentifier;
+
 	SAssetFilterBar<FAssetFilterType>::Construct(Args);
+
+	/* If we are using shared settings, add a default config for the shared settings in case it doesnt exist
+	 * This needs to go after SAssetFilterBar<FAssetFilterType>::Construct() to ensure UFilterBarConfig is valid
+	 */
+	if(bUseSharedSettings)
+	{
+		UFilterBarConfig::Get()->FilterBars.FindOrAdd(SharedIdentifier);
+
+		// Bind our delegate for when another SFilterList creates a custom text filter, so we can sync our list
+		CustomTextFilterEvent.AddSP(this, &SFilterList::OnExternalCustomTextFilterCreated);
+	}
+	
 }
 
 FReply SFilterList::OnMouseButtonUp( const FGeometry& MyGeometry, const FPointerEvent& MouseEvent )
@@ -255,140 +272,6 @@ void SFilterList::DisableFiltersThatHideItems(TArrayView<const FContentBrowserIt
 	}
 }
 
-void SFilterList::SaveSettings(const FString& IniFilename, const FString& IniSection, const FString& SettingsString) const
-{
-	FString ActiveTypeFilterString;
-	FString EnabledTypeFilterString;
-	FString ActiveFrontendFilterString;
-	FString EnabledFrontendFilterString;
-	for ( const TSharedPtr<SFilter> Filter : this->Filters )
-	{
-		const FString FilterName = Filter->GetFilterName();
-
-		// If it is a FrontendFilter
-		if ( Filter->GetFrontendFilter().IsValid() )
-		{
-			if ( ActiveFrontendFilterString.Len() > 0 )
-			{
-				ActiveFrontendFilterString += TEXT(",");
-			}
-			ActiveFrontendFilterString += FilterName;
-		
-			if ( Filter->IsEnabled() )
-			{
-				if ( EnabledFrontendFilterString.Len() > 0 )
-				{
-					EnabledFrontendFilterString += TEXT(",");
-				}
-
-				EnabledFrontendFilterString += FilterName;
-			}
-
-			const TSharedPtr<FFilterBase<FAssetFilterType>>& FrontendFilter = Filter->GetFrontendFilter();
-			const FString CustomSettingsString = FString::Printf(TEXT("%s.CustomSettings.%s"), *SettingsString, *FilterName);
-			FrontendFilter->SaveSettings(IniFilename, IniSection, CustomSettingsString);
-
-		}
-		// Otherwise we assume it is a type filter
-		else
-		{
-			if ( ActiveTypeFilterString.Len() > 0 )
-			{
-				ActiveTypeFilterString += TEXT(",");
-			}
-			ActiveTypeFilterString += FilterName;
-		
-			if ( Filter->IsEnabled() )
-			{
-				if ( EnabledTypeFilterString.Len() > 0 )
-				{
-					EnabledTypeFilterString += TEXT(",");
-				}
-
-				EnabledTypeFilterString += FilterName;
-			}
-		}
-	}
-
-	GConfig->SetString(*IniSection, *(SettingsString + TEXT(".ActiveTypeFilters")), *ActiveTypeFilterString, IniFilename);
-	GConfig->SetString(*IniSection, *(SettingsString + TEXT(".EnabledTypeFilters")), *EnabledTypeFilterString, IniFilename);
-	GConfig->SetString(*IniSection, *(SettingsString + TEXT(".ActiveFrontendFilters")), *ActiveFrontendFilterString, IniFilename);
-	GConfig->SetString(*IniSection, *(SettingsString + TEXT(".EnabledFrontendFilters")), *EnabledFrontendFilterString, IniFilename);
-}
-
-void SFilterList::LoadSettings(const FString& IniFilename, const FString& IniSection, const FString& SettingsString)
-{
-	{
-		// Add all the type filters that were found in the ActiveTypeFilters
-		FString ActiveTypeFilterString;
-		FString EnabledTypeFilterString;
-		GConfig->GetString(*IniSection, *(SettingsString + TEXT(".ActiveTypeFilters")), ActiveTypeFilterString, IniFilename);
-		GConfig->GetString(*IniSection, *(SettingsString + TEXT(".EnabledTypeFilters")), EnabledTypeFilterString, IniFilename);
-
-		// Parse comma delimited strings into arrays
-		TArray<FString> TypeFilterNames;
-		TArray<FString> EnabledTypeFilterNames;
-		ActiveTypeFilterString.ParseIntoArray(TypeFilterNames, TEXT(","), /*bCullEmpty=*/true);
-		EnabledTypeFilterString.ParseIntoArray(EnabledTypeFilterNames, TEXT(","), /*bCullEmpty=*/true);
-
-		for(const TSharedRef<FCustomClassFilterData> &CustomClassFilter : CustomClassFilters)
-		{
-			if(!this->IsClassTypeInUse(CustomClassFilter))
-			{
-				const FString FilterName = CustomClassFilter->GetFilterName();
-				if ( TypeFilterNames.Contains(FilterName) )
-				{
-					TSharedRef<SFilter> NewFilter = AddAssetFilterToBar(CustomClassFilter);
-
-					if ( EnabledTypeFilterNames.Contains(FilterName) )
-					{
-						NewFilter->SetEnabled(true, false);
-					}
-				}
-			}
-		}
-	}
-
-	{
-		// Add all the frontend filters that were found in the ActiveFrontendFilters
-		FString ActiveFrontendFilterString;	
-		FString EnabledFrontendFilterString;
-		GConfig->GetString(*IniSection, *(SettingsString + TEXT(".ActiveFrontendFilters")), ActiveFrontendFilterString, IniFilename);
-		GConfig->GetString(*IniSection, *(SettingsString + TEXT(".EnabledFrontendFilters")), EnabledFrontendFilterString, IniFilename);
-
-		// Parse comma delimited strings into arrays
-		TArray<FString> FrontendFilterNames;
-		TArray<FString> EnabledFrontendFilterNames;
-		ActiveFrontendFilterString.ParseIntoArray(FrontendFilterNames, TEXT(","), /*bCullEmpty=*/true);
-		EnabledFrontendFilterString.ParseIntoArray(EnabledFrontendFilterNames, TEXT(","), /*bCullEmpty=*/true);
-
-		// For each FrontendFilter, add any that were active and enable any that were previously enabled
-		for ( auto FrontendFilterIt = AllFrontendFilters.CreateIterator(); FrontendFilterIt; ++FrontendFilterIt )
-		{
-			TSharedRef<FFilterBase<FAssetFilterType>>& FrontendFilter = *FrontendFilterIt;
-			const FString& FilterName = FrontendFilter->GetName();
-			if (!IsFrontendFilterInUse(FrontendFilter))
-			{
-				if ( FrontendFilterNames.Contains(FilterName) )
-				{
-					TSharedRef<SFilter> NewFilter = AddFilterToBar(FrontendFilter);
-
-					if ( EnabledFrontendFilterNames.Contains(FilterName) )
-					{
-						NewFilter->SetEnabled(true, false);
-						SetFrontendFilterActive(FrontendFilter, NewFilter->IsEnabled());
-					}
-				}
-			}
-
-			const FString CustomSettingsString = FString::Printf(TEXT("%s.CustomSettings.%s"), *SettingsString, *FilterName);
-			FrontendFilter->LoadSettings(IniFilename, IniSection, CustomSettingsString);
-		}
-	}
-
-	OnFilterChanged.ExecuteIfBound();
-}
-
 void SFilterList::SetFrontendFilterCheckState(const TSharedPtr<FFrontendFilter>& InFrontendFilter, ECheckBoxState CheckState)
 {
 	this->SetFilterCheckState(InFrontendFilter, CheckState);
@@ -494,4 +377,323 @@ TSharedRef<SWidget> SFilterList::MakeAddFilterMenu(EAssetTypeCategories::Type Me
 	return UToolMenus::Get()->GenerateWidget(FilterMenuName, ToolMenuContext);
 }
 
+void SFilterList::CreateCustomFilterDialog(const FText& InText)
+{
+	CreateCustomTextFilterFromSearch(InText);
+}
+
+void SFilterList::OnCreateCustomTextFilter(const FCustomTextFilterData& InFilterData, bool bApplyFilter)
+{
+	SAssetFilterBar<FAssetFilterType>::OnCreateCustomTextFilter(InFilterData, bApplyFilter);
+
+	// If we are using shared settings (i.e sharing custom text filters) broadcast the event for all other instances to update
+	if(bUseSharedSettings)
+	{
+		// First save the shared settings for other instances to use
+		SaveSettings();
+		CustomTextFilterEvent.Broadcast(AsShared());
+	}
+}
+
+void SFilterList::OnModifyCustomTextFilter(const FCustomTextFilterData& InFilterData, TSharedPtr<ICustomTextFilter<FAssetFilterType>> InFilter)
+{
+	SAssetFilterBar<FAssetFilterType>::OnModifyCustomTextFilter(InFilterData, InFilter);
+
+	// If we are using shared settings (i.e sharing custom text filters) broadcast the event for all other instances to update
+	if(bUseSharedSettings)
+	{
+		// First save the shared settings for other instances to use
+		SaveSettings();
+		CustomTextFilterEvent.Broadcast(AsShared());
+	}
+}
+
+void SFilterList::OnDeleteCustomTextFilter(const TSharedPtr<ICustomTextFilter<FAssetFilterType>> InFilter)
+{
+	SAssetFilterBar<FAssetFilterType>::OnDeleteCustomTextFilter(InFilter);
+
+	// If we are using shared settings (i.e sharing custom text filters) broadcast the event for all other instances to update
+	if(bUseSharedSettings)
+	{
+		// First save the shared settings for other instances to use
+		SaveSettings();
+		CustomTextFilterEvent.Broadcast(AsShared());
+	}
+}
+
+void SFilterList::RestoreCustomTextFilterState(const FCustomTextFilterState& InFilterState)
+{
+	// Find the filter associated with the current instance data from our list of custom text filters
+	TSharedRef< ICustomTextFilter<FAssetFilterType> >* Filter =
+		CustomTextFilters.FindByPredicate([&InFilterState](const TSharedRef< ICustomTextFilter<FAssetFilterType> >& Element)
+	{
+		return Element->CreateCustomTextFilterData().FilterLabel.EqualTo(InFilterState.FilterData.FilterLabel);
+	});
+
+	// InFilterState should always be a filter that already exists!
+	check(Filter);
+
+	// Get the actual FFilterBase
+	TSharedRef<FFilterBase<FAssetFilterType>> ActualFilter = Filter->Get().GetFilter().ToSharedRef();
+
+	// Add it to the filter bar, since if it exists in this list it is checked
+	TSharedRef<SFilter> AddedFilter = this->AddFilterToBar(ActualFilter);
+
+	// Set the filter as active if it was previously
+	AddedFilter->SetEnabled(InFilterState.bIsActive, false);
+	this->SetFrontendFilterActive(ActualFilter, InFilterState.bIsActive);
+}
+
+void SFilterList::OnExternalCustomTextFilterCreated(TSharedPtr<SWidget> BroadcastingFilterList)
+{
+	// Do nothing if we aren't using shared settings or if the event was broadcasted by this filter list
+	if(!bUseSharedSettings || BroadcastingFilterList == AsShared())
+	{
+		return;
+	}
+
+	/* We are going to remove all our custom text filters and re-load them from the shared settings, since a different
+	 * instance modified them.
+	 */
+
+	// To preserve the state of any checked/active custom text filters
+	TArray<FCustomTextFilterState> CurrentCustomTextFilterStates;
+	
+	for (const TSharedRef<ICustomTextFilter<FAssetFilterType>>& CustomTextFilter : this->CustomTextFilters)
+	{
+		// Get the actual FFilterBase
+		TSharedRef<FFilterBase<FAssetFilterType>> CustomFilter = CustomTextFilter->GetFilter().ToSharedRef();
+			
+		// Is the filter "checked", i.e visible in the filter bar
+		bool bIsChecked = IsFrontendFilterInUse(CustomFilter);
+
+		// Is the filter "active", i.e visible and enabled in the filter bar
+		bool bIsActive = IsFilterActive(CustomFilter);
+
+		// Only save the state if the filter is checked so we can restore it
+		if(bIsChecked)
+		{
+			/* Remove the filter from the list (calling SBasicFilterBar::RemoveFilter because we get a compiler error
+			*  due to SAssetFilterBar overriding RemoveFilter that takes in an SFilter that hides the parent class function
+			*/
+			SBasicFilterBar<FAssetFilterType>::RemoveFilter(CustomFilter, false);
+			
+			FCustomTextFilterState FilterState;
+			FilterState.FilterData = CustomTextFilter->CreateCustomTextFilterData();
+			FilterState.bIsChecked = bIsChecked;
+			FilterState.bIsActive = bIsActive;
+			
+			CurrentCustomTextFilterStates.Add(FilterState);
+		}
+	}
+
+	// Get the shared settings and reload the filters
+	FFilterBarSettings* SharedSettings = &UFilterBarConfig::Get()->FilterBars.FindOrAdd(SharedIdentifier);
+	LoadCustomTextFilters(SharedSettings);
+
+	// Restore the state of any previously active ones
+	for(const FCustomTextFilterState& SavedFilterState : CurrentCustomTextFilterStates)
+	{
+		RestoreCustomTextFilterState(SavedFilterState);
+	}
+}
+
+void SFilterList::UpdateCustomTextFilterIncludes(const bool InIncludeClassName, const bool InIncludeAssetPath, const bool InIncludeCollectionNames)
+{
+	bIncludeClassName = InIncludeClassName;
+	bIncludeAssetPath = InIncludeAssetPath;
+	bIncludeCollectionNames = InIncludeCollectionNames;
+
+	for (TSharedPtr<ICustomTextFilter<FAssetFilterType>> CustomTextFilter : CustomTextFilters)
+	{
+		// This is a safe cast, since SFilterList will always and only have FFilterListCustomTextFilters
+		if(TSharedPtr<FFrontendFilter_CustomText> FilterListCustomTextFilter = StaticCastSharedPtr<FFrontendFilter_CustomText>(CustomTextFilter))
+		{
+			FilterListCustomTextFilter->UpdateCustomTextFilterIncludes(bIncludeClassName, bIncludeAssetPath, bIncludeCollectionNames);
+		}
+	}
+}
+
+void SFilterList::SaveSettings()
+{
+	// If this instance doesn't want to use the shared settings, save the settings normally
+	if(!bUseSharedSettings)
+	{
+		SAssetFilterBar<FAssetFilterType>::SaveSettings();
+		return;
+	}
+
+	if(FilterBarIdentifier.IsNone())
+	{
+		UE_LOG(LogSlate, Error, TEXT("SFilterList Requires that you specify a FilterBarIdentifier to save settings"));
+		return;
+	}
+
+	// Get the settings unique to this instance and the common settings
+	FFilterBarSettings* InstanceSettings = &UFilterBarConfig::Get()->FilterBars.FindOrAdd(FilterBarIdentifier);
+	FFilterBarSettings* SharedSettings = &UFilterBarConfig::Get()->FilterBars.FindOrAdd(SharedIdentifier);
+
+	// Empty both the configs, we are just going to re-save everything there
+	InstanceSettings->Empty();
+	SharedSettings->Empty();
+
+	// Save all the programatically added filters normally
+	SaveFilters(InstanceSettings);
+
+	/** For each custom text filter: Save the filterdata into the common settings, so that all instances that use it
+	 *	are synced.
+	 *	For each CHECKED custom text filter: Save just the filter name, and the checked and active state into the
+	 *	instance settings. Those are specific to this instance (i.e we don't want a filter to be active in all
+	 *	instances if activated in one)
+	 */
+	for (const TSharedRef<ICustomTextFilter<FAssetFilterType>>& CustomTextFilter : this->CustomTextFilters)
+	{
+		// Get the actual FFilterBase
+		TSharedRef<FFilterBase<FAssetFilterType>> CustomFilter = CustomTextFilter->GetFilter().ToSharedRef();
+			
+		// Is the filter "checked", i.e visible in the filter bar
+		bool bIsChecked = IsFrontendFilterInUse(CustomFilter);
+
+		// Is the filter "active", i.e visible and enabled in the filter bar
+		bool bIsActive = IsFilterActive(CustomFilter);
+
+		// Get the data associated with this filter
+		FCustomTextFilterData FilterData = CustomTextFilter->CreateCustomTextFilterData();
+
+		// Just save the filter data into the shared settings
+		FCustomTextFilterState SharedFilterState;
+		SharedFilterState.FilterData = FilterData;
+		SharedSettings->CustomTextFilters.Add(SharedFilterState);
+
+		if(bIsChecked)
+		{
+			// Create a duplicate filter data that just contains the filter label for this instance to know
+			FCustomTextFilterData InstanceFilterData;
+			InstanceFilterData.FilterLabel = FilterData.FilterLabel;
+			
+			// Just save the filter name and enabled/active state into the shared settings
+			FCustomTextFilterState InstanceFilterState;
+			InstanceFilterState.bIsChecked = bIsChecked;
+			InstanceFilterState.bIsActive = bIsActive;
+			InstanceFilterState.FilterData = InstanceFilterData;
+			
+			InstanceSettings->CustomTextFilters.Add(InstanceFilterState);
+		}
+		
+	}
+
+
+	SaveConfig();
+}
+
+void SFilterList::LoadSettings()
+{
+	// If this instance doesn't want to use the shared settings, load the settings normally
+	if(!bUseSharedSettings)
+	{
+		SAssetFilterBar<FAssetFilterType>::LoadSettings();
+		return;
+	}
+
+	if(FilterBarIdentifier.IsNone())
+	{
+		UE_LOG(LogSlate, Error, TEXT("SFilterList Requires that you specify a FilterBarIdentifier to load settings"));
+		return;
+	}
+
+	// Get the settings unique to this instance and the common settings
+	const FFilterBarSettings* InstanceSettings = UFilterBarConfig::Get()->FilterBars.Find(FilterBarIdentifier);
+	const FFilterBarSettings* SharedSettings = UFilterBarConfig::Get()->FilterBars.Find(SharedIdentifier);
+
+	// Load the filters specified programatically normally
+	LoadFilters(InstanceSettings);
+
+	// Load the custom text filters from the shared settings
+	LoadCustomTextFilters(SharedSettings);
+	
+	// From the instance settings, get each checked filter and set the checked and active state
+	for(const FCustomTextFilterState& FilterState : InstanceSettings->CustomTextFilters)
+	{
+		RestoreCustomTextFilterState(FilterState);
+	}
+	
+	this->OnFilterChanged.ExecuteIfBound();
+}
+
+void SFilterList::LoadCustomTextFilters(const FFilterBarSettings* FilterBarConfig)
+{
+	CustomTextFilters.Empty();
+	
+	// Extract just the filter data from the common settings
+	for(const FCustomTextFilterState& FilterState : FilterBarConfig->CustomTextFilters)
+	{
+		// Create an ICustomTextFilter using the provided delegate
+		TSharedRef<ICustomTextFilter<FAssetFilterType>> NewTextFilter = this->CreateTextFilter.Execute().ToSharedRef();
+
+		// Get the actual FFilterBase
+		TSharedRef<FFilterBase<FAssetFilterType>> NewFilter = NewTextFilter->GetFilter().ToSharedRef();
+
+		// Set the internals of the custom text filter from what we have saved
+		NewTextFilter->SetFromCustomTextFilterData(FilterState.FilterData);
+
+		// Add this to our list of custom text filters
+		CustomTextFilters.Add(NewTextFilter);
+	}
+}
+
+/////////////////////////////////////////
+// FFilterListCustomTextFilter
+/////////////////////////////////////////
+
+
+FString FFrontendFilter_CustomText::GetName() const
+{
+	return TEXT("FFrontendFilter_CustomText");
+}
+
+FText FFrontendFilter_CustomText::GetDisplayName() const
+{
+	return DisplayName;
+}
+FText FFrontendFilter_CustomText::GetToolTipText() const
+{
+	return GetRawFilterText();
+}
+
+FLinearColor FFrontendFilter_CustomText::GetColor() const
+{
+	return Color;
+}
+
+void FFrontendFilter_CustomText::UpdateCustomTextFilterIncludes(const bool InIncludeClassName, const bool InIncludeAssetPath, const bool InIncludeCollectionNames)
+{
+	SetIncludeClassName(InIncludeClassName);
+	SetIncludeAssetPath(InIncludeAssetPath);
+	SetIncludeCollectionNames(InIncludeCollectionNames);
+}
+
+void FFrontendFilter_CustomText::SetFromCustomTextFilterData(const FCustomTextFilterData& InFilterData)
+{
+	Color = InFilterData.FilterColor;
+	DisplayName = InFilterData.FilterLabel;
+	SetRawFilterText(InFilterData.FilterString);
+}
+
+FCustomTextFilterData FFrontendFilter_CustomText::CreateCustomTextFilterData() const
+{
+	FCustomTextFilterData CustomTextFilterData;
+
+	CustomTextFilterData.FilterColor = Color;
+	CustomTextFilterData.FilterLabel = DisplayName;
+	CustomTextFilterData.FilterString = GetRawFilterText();
+
+	return CustomTextFilterData;
+}
+
+TSharedPtr<FFilterBase<FAssetFilterType>> FFrontendFilter_CustomText::GetFilter()
+{
+	return AsShared();
+}
+
+
 #undef LOCTEXT_NAMESPACE
diff --git a/Engine/Source/Editor/ContentBrowser/Private/SFilterList.h b/Engine/Source/Editor/ContentBrowser/Private/SFilterList.h
--- a/Engine/Source/Editor/ContentBrowser/Private/SFilterList.h
+++ b/Engine/Source/Editor/ContentBrowser/Private/SFilterList.h
@@ -10,6 +10,7 @@
 #include "AssetRegistry/AssetData.h"
 #include "AssetRegistry/ARFilter.h"
 #include "FrontendFilterBase.h"
+#include "FrontendFilters.h"
 #include "Filters/SAssetFilterBar.h"
 
 class UToolMenu;
@@ -17,6 +18,8 @@ struct FToolMenuSection;
 class SFilter;
 class SWrapBox;
 enum class ECheckBoxState : uint8;
+class FFilterListCustomTextFilter;
+class FFrontendFilter_Text;
 
 /**
  * A list of filters currently applied to an asset view.
@@ -26,10 +29,19 @@ class SFilterList : public SAssetFilterBar<FAssetFilterType>
 public:
 
 	DECLARE_DELEGATE_RetVal( TSharedPtr<SWidget>, FOnGetContextMenu );
-
+	/**
+	 * An event delegate that is executed when a custom text filter has been created/modified/deleted in any FilterList
+	 * that is using shared settings.
+	 */
+	DECLARE_EVENT_OneParam(SFilterList, FCustomTextFilterEvent, TSharedPtr<SWidget> /* BroadcastingFilterList */);
+	
 	using FOnFilterChanged = typename SAssetFilterBar<FAssetFilterType>::FOnFilterChanged;
 
-	SLATE_BEGIN_ARGS( SFilterList ){}
+	SLATE_BEGIN_ARGS( SFilterList ):
+	_UseSharedSettings(false)
+	{
+		
+	}
 
 		/** Called when an asset is right clicked */
 		SLATE_EVENT( FOnGetContextMenu, OnGetContextMenu )
@@ -45,6 +57,19 @@ public:
 
 		/** Custom front end filters to be displayed */
 		SLATE_ARGUMENT( TArray< TSharedRef<FFrontendFilter> >, ExtraFrontendFilters )
+	
+		/** A delegate to create a Text Filter for FilterType items. If provided, will allow creation of custom text filters
+		 *  from the filter dropdown menu.
+		 */
+		SLATE_ARGUMENT(FCreateTextFilter, CreateTextFilter)
+	
+		/** A unique identifier for this filter bar needed to enable saving settings in a config file */
+		SLATE_ARGUMENT(FName, FilterBarIdentifier)
+
+		/** If true, CustomTextFilters are saved/loaded from a config shared between multiple SFilterLists
+		 *	Currently used to sync all content browser custom text filters.
+		 */
+		SLATE_ARGUMENT(bool, UseSharedSettings)
 
 	SLATE_END_ARGS()
 
@@ -69,28 +94,47 @@ public:
 	/** Disables any active filters that would hide the supplied items */
 	void DisableFiltersThatHideItems(TArrayView<const FContentBrowserItem> ItemList);
 
-	/** Saves any settings to config that should be persistent between editor sessions */
-	void SaveSettings(const FString& IniFilename, const FString& IniSection, const FString& SettingsString) const;
-
-	/** Loads any settings to config that should be persistent between editor sessions */
-	void LoadSettings(const FString& IniFilename, const FString& IniSection, const FString& SettingsString);
-
 	virtual FReply OnMouseButtonUp( const FGeometry& MyGeometry, const FPointerEvent& MouseEvent ) override;
 
 	/** Returns the class filters specified at construction using argument 'InitialClassFilters'. */
 	const TArray<UClass*>& GetInitialClassFilters();
 
+	/** Open the dialog to create a custom filter from the given text */
+	void CreateCustomFilterDialog(const FText& InText);
+	
+	/** Updates bIncludeClassName, bIncludeAssetPath and bIncludeCollectionNames for all custom text filters */
+	void UpdateCustomTextFilterIncludes(const bool InIncludeClassName, const bool InIncludeAssetPath, const bool InIncludeCollectionNames);
+
+	virtual void SaveSettings()override;
+	virtual void LoadSettings() override;
+
 protected:
 	
 	/** Handler for when the add filter button was clicked */
 	TSharedRef<SWidget> MakeAddFilterMenu() override;
+	
+	/** Handler for when a custom text filter is created */
+	virtual void OnCreateCustomTextFilter(const FCustomTextFilterData& InFilterData, bool bApplyFilter) override;
+	/** Handler for when a custom text filter is modified */
+	virtual void OnModifyCustomTextFilter(const FCustomTextFilterData& InFilterData, TSharedPtr<ICustomTextFilter<FAssetFilterType>> InFilter) override;
+	/** Handler for when a custom text filter is deleted */
+	virtual void OnDeleteCustomTextFilter(const TSharedPtr<ICustomTextFilter<FAssetFilterType>> InFilter) override;
 
 private:
 
-	// Exists for backwards compatibility with ExternalMakeAddFilterMenu
+	/** Exists for backwards compatibility with ExternalMakeAddFilterMenu */
 	TSharedRef<SWidget> MakeAddFilterMenu(EAssetTypeCategories::Type MenuExpansion = EAssetTypeCategories::Basic);
 
 	void PopulateAddFilterMenu_Internal(UToolMenu* Menu);
+
+	/** Handler for when another SFilterList using shared settings creates a custom text filter */
+	void OnExternalCustomTextFilterCreated(TSharedPtr<SWidget> BroadcastingFilterList);
+	
+	/** Empty our list of custom text filters, and load from the given config */
+	void LoadCustomTextFilters(const FFilterBarSettings* FilterBarConfig);
+
+	/** Find the custom text filter corresponding to the specified state, and restore it's state to what is specified */
+	void RestoreCustomTextFilterState(const FCustomTextFilterState& InFilterState);
 	
 private:
 
@@ -105,4 +149,62 @@ private:
 
 	/** A reference to AllFrontEndFilters so we can access the filters as FFrontEndFilter instead of FFilterBase<FAssetFilterType> */
 	TArray< TSharedRef<FFrontendFilter> > AllFrontendFilters_Internal;
+
+	/** An identifier shared by all SFilterLists, used to save and load settings common to every instance */
+	static const FName SharedIdentifier;
+
+	/** If bIncludeClassName is true, custom text filters will include an asset's class name in the search */
+	bool bIncludeClassName;
+
+	/** If bIncludeAssetPath is true, custom text filters will match against full Asset path */
+	bool bIncludeAssetPath;
+
+	/** If bIncludeCollectionNames is true, custom text filters will match against collection names as well */
+	bool bIncludeCollectionNames;
+
+	/** Whether this FilterList wants to load/save from the settings common to all instances */
+	bool bUseSharedSettings;
+
+	/** The event that executes when a custom text filter is created/modified/deleted in any filter list that is using Shared Settings */
+	static FCustomTextFilterEvent CustomTextFilterEvent;
 };
+
+/* A custom implementation of ICustomTextFilter that uses FFrontendFilter_Text to handle comparing items in the
+ * Asset View to Custom Text Filters. This ensures that the advanced search syntax etc behaves properly when used
+ * by a custom text filter created by saving a search
+ */
+class FFrontendFilter_CustomText :
+	public ICustomTextFilter<FAssetFilterType>,
+	public FFrontendFilter_Text,
+	public TSharedFromThis<FFrontendFilter_CustomText>
+{
+public:
+
+	// FFrontendFilter implementation
+	virtual FString GetName() const override;
+	virtual FText GetDisplayName() const override;
+	virtual FText GetToolTipText() const override;
+	virtual FLinearColor GetColor() const override;
+	
+	/** Updates bIncludeClassName, bIncludeAssetPath and bIncludeCollectionNames for this filter */
+	void UpdateCustomTextFilterIncludes(const bool InIncludeClassName, const bool InIncludeAssetPath, const bool InIncludeCollectionNames);
+
+	//ICustomTextFilter interface
+	
+	/** Set the internals of this filter from an FCustomTextFilterData */
+	virtual void SetFromCustomTextFilterData(const FCustomTextFilterData& InFilterData) override;
+
+	/** Create an FCustomTextFilterData from the internals of this filter */
+	virtual FCustomTextFilterData CreateCustomTextFilterData() const override;
+
+	/** Get the actual filter */
+	virtual TSharedPtr<FFilterBase<FAssetFilterType>> GetFilter() override;
+
+protected:
+
+	/* The Display Name of this custom filter that the user sees */
+	FText DisplayName;
+
+	/* The Color of this filter pill */
+	FLinearColor Color;
+};
\ No newline at end of file
diff --git a/Engine/Source/Editor/EditorWidgets/Private/Filters/SAssetFilterBar.cpp b/Engine/Source/Editor/EditorWidgets/Private/Filters/SAssetFilterBar.cpp
new file mode 100644
--- /dev/null
+++ b/Engine/Source/Editor/EditorWidgets/Private/Filters/SAssetFilterBar.cpp
@@ -0,0 +1,40 @@
+// Copyright Epic Games, Inc. All Rights Reserved.
+
+#include "Filters/SAssetFilterBar.h"
+
+FFilterBarSettings* FFilterBarBase::GetMutableConfig()
+{
+	if(FilterBarIdentifier.IsNone())
+	{
+		return nullptr;
+	}
+
+	UFilterBarConfig* FilterBarConfig = GetMutableDefault<UFilterBarConfig>();
+	return &UFilterBarConfig::Get()->FilterBars.FindOrAdd(FilterBarIdentifier);
+}
+
+const FFilterBarSettings* FFilterBarBase::GetConstConfig() const
+{
+	if(FilterBarIdentifier.IsNone())
+	{
+		return nullptr;
+	}
+
+	const UFilterBarConfig* FilterBarConfig = GetDefault<UFilterBarConfig>();
+	return UFilterBarConfig::Get()->FilterBars.Find(FilterBarIdentifier);
+}
+
+void FFilterBarBase::SaveConfig()
+{
+	UFilterBarConfig::Get()->SaveEditorConfig();
+}
+
+void FFilterBarBase::InitializeConfig()
+{
+	UFilterBarConfig::Initialize();
+
+	UFilterBarConfig::Get()->LoadEditorConfig();
+
+	// Call GetMutableConfig to force create a config for this filter bar if the user specified FilterBarIdentifier
+	FFilterBarSettings* FilterBarConfig = GetMutableConfig();
+}
diff --git a/Engine/Source/Editor/EditorWidgets/Private/Filters/SFilterBar.cpp b/Engine/Source/Editor/EditorWidgets/Private/Filters/SFilterBar.cpp
deleted file mode 100644
diff --git a/Engine/Source/Editor/EditorWidgets/Private/SAssetSearchBox.cpp b/Engine/Source/Editor/EditorWidgets/Private/SAssetSearchBox.cpp
--- a/Engine/Source/Editor/EditorWidgets/Private/SAssetSearchBox.cpp
+++ b/Engine/Source/Editor/EditorWidgets/Private/SAssetSearchBox.cpp
@@ -59,14 +59,18 @@ void SAssetSearchBox::Construct( const FArguments& InArgs )
 			SAssignNew(SuggestionBox, SMenuAnchor)
 			.Placement( InArgs._SuggestionListPlacement )
 			[
-				SAssignNew(InputText, SSearchBox)
+				/* Use an SFilterSearchBox internally to add the ability to show search history and potentially
+				 * save searches as filters if used with a Filter Bar widget (@see SBasicFilterBar etc)
+				 */
+				SAssignNew(InputText, SFilterSearchBox)
 				.InitialText(InArgs._InitialText)
 				.HintText(InArgs._HintText)
 				.OnTextChanged(this, &SAssetSearchBox::HandleTextChanged)
 				.OnTextCommitted(this, &SAssetSearchBox::HandleTextCommitted)
-				.SelectAllTextWhenFocused( false )
 				.DelayChangeNotificationsWhileTyping( InArgs._DelayChangeNotificationsWhileTyping )
 				.OnKeyDownHandler(this, &SAssetSearchBox::HandleKeyDown)
+				.ShowSearchHistory(InArgs._ShowSearchHistory)
+				.OnSaveSearchClicked(InArgs._OnSaveSearchClicked)
 			]
 			.MenuContent
 				(
@@ -401,3 +405,8 @@ FText SAssetSearchBox::DefaultSuggestionChosenImpl(const FText& SearchText, cons
 	// Default implementation just uses the suggestion as the search text
 	return FText::FromString(Suggestion);
 }
+
+void SAssetSearchBox::SetOnSaveSearchHandler(SFilterSearchBox::FOnSaveSearchClicked InOnSaveSearchHandler)
+{
+	InputText->SetOnSaveSearchHandler(InOnSaveSearchHandler);
+}
diff --git a/Engine/Source/Editor/EditorWidgets/Public/Filters/SAssetFilterBar.h b/Engine/Source/Editor/EditorWidgets/Public/Filters/SAssetFilterBar.h
--- a/Engine/Source/Editor/EditorWidgets/Public/Filters/SAssetFilterBar.h
+++ b/Engine/Source/Editor/EditorWidgets/Public/Filters/SAssetFilterBar.h
@@ -13,6 +13,7 @@
 #include "SlateGlobals.h"
 #include "Logging/LogMacros.h"
 #include "Styling/SlateIconFinder.h"
+#include "Filters/FilterBarConfig.h"
 
 #include "SAssetFilterBar.generated.h"
 
@@ -37,6 +38,31 @@ public:
 	FOnPopulateAddAssetFilterMenu PopulateFilterMenu;
 };
 
+/** A non-templated base class for SAssetFilterBar, where functionality that does not depend on the type of item
+ *  being filtered lives
+ */
+class EDITORWIDGETS_API FFilterBarBase
+{
+protected:
+
+	/** Get a mutable version of this filter bar's config */
+	FFilterBarSettings* GetMutableConfig();
+
+	/** Get a const version of this filter bar's config */
+	const FFilterBarSettings* GetConstConfig() const;
+
+	/** Save this filte bar's data to the config file */
+	void SaveConfig();
+
+	/** Initialize and load this filter bar's config */
+	void InitializeConfig();
+
+protected:
+
+	/* Unique name for this filter bar */
+	FName FilterBarIdentifier;
+};
+
 
 /**
 * An Asset Filter Bar widget, which can be used to filter items of type [FilterType] given a list of custom filters along with built in support for AssetType filters
@@ -72,14 +98,14 @@ public:
 *  }
 */
 template<typename FilterType>
-class SAssetFilterBar : public SBasicFilterBar<FilterType>
+class SAssetFilterBar : public SBasicFilterBar<FilterType>, public FFilterBarBase
 {
 public:
 	
 	using FOnFilterChanged = typename SBasicFilterBar<FilterType>::FOnFilterChanged;
 	using FOnExtendAddFilterMenu = typename SBasicFilterBar<FilterType>::FOnExtendAddFilterMenu;
 	using FCreateTextFilter = typename SBasicFilterBar<FilterType>::FCreateTextFilter;
-
+	
 	SLATE_BEGIN_ARGS( SAssetFilterBar<FilterType> )
 		: _UseDefaultAssetFilters(true)
 		{
@@ -109,6 +135,9 @@ public:
 		 */
 		SLATE_ARGUMENT(TSharedPtr<SFilterSearchBox>, FilterSearchBox)
 
+		/** A unique identifier for this filter bar needed to enable saving settings in a config file */
+		SLATE_ARGUMENT(FName, FilterBarIdentifier)
+	
 		/** Whether the filter bar should provide the default asset filters */
 		SLATE_ARGUMENT(bool, UseDefaultAssetFilters)
 	
@@ -120,6 +149,7 @@ public:
 		bUseDefaultAssetFilters = InArgs._UseDefaultAssetFilters;
 		CustomClassFilters = InArgs._CustomClassFilters;
 		UserCustomClassFilters = InArgs._CustomClassFilters;
+		FilterBarIdentifier = InArgs._FilterBarIdentifier;
 		
 		typename SBasicFilterBar<FilterType>::FArguments Args;
 		Args._OnFilterChanged = InArgs._OnFilterChanged;
@@ -130,6 +160,8 @@ public:
 		
 		SBasicFilterBar<FilterType>::Construct(Args);
 
+		InitializeConfig();
+		
 		// Create the Asset Type Action filters if requested
 		CreateAssetTypeActionFilters();
 
@@ -141,8 +173,176 @@ public:
 
 	}
 
+	virtual void SaveSettings()
+	{
+		FFilterBarSettings* FilterBarConfig = GetMutableConfig();
+		
+		if(!FilterBarConfig)
+		{
+			UE_LOG(LogSlate, Error, TEXT("SFilterBar Requires that you specify a FilterBarIdentifier to save settings"));
+			return;
+		}
+
+		// Empty the config, we are just going to re-save everything
+		FilterBarConfig->Empty();
+
+		SaveCustomTextFilters(FilterBarConfig);
+		SaveFilters(FilterBarConfig);
+
+		SaveConfig();
+	}
+
+	virtual void LoadSettings()
+	{
+		const FFilterBarSettings* FilterBarConfig = GetConstConfig();
+		
+		if(!FilterBarConfig)
+		{
+			UE_LOG(LogSlate, Error, TEXT("SFilterBar Requires that you specify a FilterBarIdentifier to load settings"));
+			return;
+		}
+
+		if(!FilterBarConfig->CustomTextFilters.IsEmpty())
+		{
+			// We must have a CreateTextFilter bound if we have any custom text filters saved!!
+			check(this->CreateTextFilter.IsBound());
+		}
+
+		LoadFilters(FilterBarConfig);
+		LoadCustomTextFilters(FilterBarConfig);
+
+		this->OnFilterChanged.ExecuteIfBound();
+	}
+	
 protected:
 
+	/** Save all the custom class filters (created by the user) into the specified config */
+	void SaveCustomTextFilters(FFilterBarSettings* FilterBarConfig)
+	{
+		/* Go through all the custom text filters (including unchecked ones) to save them, so that the user does not
+		 * lose the text filters they created from the menu or a saved search
+		 */
+		for (const TSharedRef<ICustomTextFilter<FilterType>>& CustomTextFilter : this->CustomTextFilters)
+		{
+			// Get the actual FFilterBase
+			TSharedRef<FFilterBase<FilterType>> CustomFilter = CustomTextFilter->GetFilter().ToSharedRef();
+			
+			// Is the filter "checked", i.e visible in the filter bar
+			bool bIsChecked = this->IsFrontendFilterInUse(CustomFilter);
+
+			// Is the filter "active", i.e visible and enabled in the filter bar
+			bool bIsActive = this->IsFilterActive(CustomFilter);
+
+			// Get the data associated with this filter
+			FCustomTextFilterData FilterData = CustomTextFilter->CreateCustomTextFilterData();
+
+			FCustomTextFilterState FilterState;
+			FilterState.bIsChecked = bIsChecked;
+			FilterState.bIsActive = bIsActive;
+			FilterState.FilterData = FilterData;
+
+			FilterBarConfig->CustomTextFilters.Add(FilterState);
+		}
+	}
+
+	/** Save all the custom filters (created programatically) into the specified config */
+	void SaveFilters(FFilterBarSettings* FilterBarConfig)
+	{
+		/* For the remaining (custom and type) filters, go through the currently active (i.e visible in the filter bar)
+		 * ones to save their state, since they will be added to the filter bar programatically every time
+		 */
+		for ( const TSharedPtr<SFilter> Filter : this->Filters )
+		{
+			const FString FilterName = Filter->GetFilterName();
+
+			// Ignore custom text filters, since we saved them previously
+			if(FilterName == FCustomTextFilter<FilterType>::GetFilterTypeName().ToString())
+			{
+				continue;
+			}
+			// If it is a FrontendFilter
+			else if ( Filter->GetFrontendFilter().IsValid() )
+			{
+				FilterBarConfig->CustomFilters.Add(FilterName, Filter->IsEnabled());
+			}
+			// Otherwise we assume it is a type filter
+			else
+			{
+				FilterBarConfig->TypeFilters.Add(FilterName, Filter->IsEnabled());
+			}
+		}
+	}
+
+	/** Load all the custom class filters (created by the user) from the specified config */
+	void LoadCustomTextFilters(const FFilterBarSettings* FilterBarConfig)
+	{
+		// Load all the custom text filters
+		for(const FCustomTextFilterState& FilterState : FilterBarConfig->CustomTextFilters)
+		{
+			// Create an ICustomTextFilter using the provided delegate
+			TSharedRef<ICustomTextFilter<FilterType>> NewTextFilter = this->CreateTextFilter.Execute().ToSharedRef();
+
+			// Get the actual FFilterBase
+			TSharedRef<FFilterBase<FilterType>> NewFilter = NewTextFilter->GetFilter().ToSharedRef();
+
+			// Set the internals of the custom text filter from what we have saved
+			NewTextFilter->SetFromCustomTextFilterData(FilterState.FilterData);
+
+			// Add this to our list of custom text filters
+			this->CustomTextFilters.Add(NewTextFilter);
+
+			// If the filter was checked previously, add it to the filter bar
+			if(FilterState.bIsChecked)
+			{
+				TSharedRef<SFilter> AddedFilter = this->AddFilterToBar(NewFilter);
+
+				// Set the filter as active if it was previously
+				AddedFilter->SetEnabled(FilterState.bIsActive, false);
+				
+				this->SetFrontendFilterActive(NewFilter, FilterState.bIsActive);
+			}
+		}
+	}
+	
+	/** Load all the custom filters (created programatically) from the specified config */
+	void LoadFilters(const FFilterBarSettings* FilterBarConfig)
+	{
+		// Load all the custom filters (i.e FrontendFilters)
+		for ( auto FrontendFilterIt = this->AllFrontendFilters.CreateIterator(); FrontendFilterIt; ++FrontendFilterIt )
+		{
+			TSharedRef<FFilterBase<FilterType>>& FrontendFilter = *FrontendFilterIt;
+			const FString& FilterName = FrontendFilter->GetName();
+			
+			if (!this->IsFrontendFilterInUse(FrontendFilter))
+			{
+				// Try to find this type filter in our list of saved filters
+				if ( const bool* bIsActive = FilterBarConfig->CustomFilters.Find(FilterName) )
+				{
+					TSharedRef<SFilter> NewFilter = this->AddFilterToBar(FrontendFilter);
+					
+					NewFilter->SetEnabled(*bIsActive, false);
+					this->SetFrontendFilterActive(FrontendFilter, NewFilter->IsEnabled());
+				}
+			}
+		}
+
+		// Load all the type filters
+		for(const TSharedRef<FCustomClassFilterData> &CustomClassFilter : this->CustomClassFilters)
+		{
+			if(!this->IsClassTypeInUse(CustomClassFilter))
+			{
+				const FString FilterName = CustomClassFilter->GetFilterName();
+
+				// Try to find this type filter in our list of saved filters
+				if ( const bool* bIsActive = FilterBarConfig->TypeFilters.Find(CustomClassFilter->GetFilterName()) )
+				{
+					TSharedRef<SFilter> NewFilter = this->AddAssetFilterToBar(CustomClassFilter);
+					NewFilter->SetEnabled(*bIsActive, false);
+				}
+			}
+		}
+	}
+	
 	typedef typename SBasicFilterBar<FilterType>::SFilter SFilter;
 
 	/** A subclass of SFilter in SBasicFilterBar to add functionality for Asset Filters */
diff --git a/Engine/Source/Editor/EditorWidgets/Public/Filters/SFilterBar.h b/Engine/Source/Editor/EditorWidgets/Public/Filters/SFilterBar.h
--- a/Engine/Source/Editor/EditorWidgets/Public/Filters/SFilterBar.h
+++ b/Engine/Source/Editor/EditorWidgets/Public/Filters/SFilterBar.h
@@ -8,32 +8,6 @@
 
 #include "Filters/SAssetFilterBar.h"
 #include "Filters/AssetFilter.h"
-#include "Filters/FilterBarConfig.h"
-
-/** A non-templated base class for SFilterBar, where functionality that does not depend on the type of item
- *  being filtered lives
- */
-class EDITORWIDGETS_API FFilterBarBase
-{
-protected:
-
-	/** Get a mutable version of this filter bar's config */
-	FFilterBarSettings* GetMutableConfig();
-
-	/** Get a const version of this filter bar's config */
-	const FFilterBarSettings* GetConstConfig() const;
-
-	/** Save this filte bar's data to the config file */
-	void SaveConfig();
-
-	/** Initialize and load this filter bar's config */
-	void InitializeConfig();
-
-protected:
-
-	/* Unique name for this filter bar */
-	FName FilterBarIdentifier;
-};
 
 /**
 * A Filter Bar widget, which can be used to filter items of type [FilterType] given a list of custom filters
@@ -60,7 +34,7 @@ protected:
 * NOTE: Use MakeAddFilterButton() to make the button that summons the dropdown showing all the filters
 */
 template<typename FilterType>
-class SFilterBar : public SAssetFilterBar<FilterType>, public FFilterBarBase
+class SFilterBar : public SAssetFilterBar<FilterType>
 {
 public:
 	using FOnFilterChanged = typename SBasicFilterBar<FilterType>::FOnFilterChanged;
@@ -125,10 +99,9 @@ public:
 		Args._OnExtendAddFilterMenu = InArgs._OnExtendAddFilterMenu;
 		Args._CreateTextFilter = InArgs._CreateTextFilter;
 		Args._FilterSearchBox = InArgs._FilterSearchBox;
+		Args._FilterBarIdentifier = InArgs._FilterBarIdentifier;
 		
 		SAssetFilterBar<FilterType>::Construct(Args);
-
-		FilterBarIdentifier = InArgs._FilterBarIdentifier;
 		
 		// Create the dummy filter that represents all currently active Asset Type filters
 		AssetFilter = MakeShareable(new FAssetFilter<FilterType>);
@@ -148,8 +121,6 @@ public:
 		}
 
 		this->ActiveFilters->Add(AssetFilter);
-
-		InitializeConfig();
 	}
 	
 	/** Use this function to get all currently active filters (to filter your items)
@@ -162,149 +133,6 @@ public:
 		return this->ActiveFilters;	
 	}
 
-	virtual void SaveSettings()
-	{
-		FFilterBarSettings* FilterBarConfig = GetMutableConfig();
-		
-		if(!FilterBarConfig)
-		{
-			UE_LOG(LogSlate, Error, TEXT("SFilterBar Requires that you specify a FilterBarIdentifier to save settings"));
-			return;
-		}
-
-		// Empty the config, we are just going to re-save everything
-		FilterBarConfig->Empty();
-
-		/* Go through all the custom text filters (including unchecked ones) to save them, so that the user does not
-		 * lose the text filters they created from the menu or a saved search
-		 */
-		for (const TSharedRef<FCustomTextFilter<FilterType>>& CustomTextFilter : this->CustomTextFilters)
-		{
-			// Is the filter "checked", i.e visible in the filter bar
-			bool bIsChecked = this->IsFrontendFilterInUse(CustomTextFilter);
-
-			// Is the filter "active", i.e visible and enabled in the filter bar
-			bool bIsActive = this->IsFilterActive(CustomTextFilter);
-
-			// Get the data associated with this filter
-			FCustomTextFilterData FilterData = CustomTextFilter->CreateCustomTextFilterData();
-
-			FCustomTextFilterState FilterState;
-			FilterState.bIsChecked = bIsChecked;
-			FilterState.bIsActive = bIsActive;
-			FilterState.FilterData = FilterData;
-
-			FilterBarConfig->CustomTextFilters.Add(FilterState);
-		}
-		
-		/* For the remaining (custom and type) filters, go through the currently active (i.e visible in the filter bar)
-		 * ones to save their state, since they will be added to the filter bar programatically every time
-		 */
-		for ( const TSharedPtr<SFilter> Filter : this->Filters )
-		{
-			const FString FilterName = Filter->GetFilterName();
-
-			// Ignore custom text filters, since we saved them previously
-			if(FilterName == FCustomTextFilter<FilterType>::GetCommonName().ToString())
-			{
-				
-			}
-			// If it is a FrontendFilter
-			else if ( Filter->GetFrontendFilter().IsValid() )
-			{
-				FilterBarConfig->CustomFilters.Add(FilterName, Filter->IsEnabled());
-			}
-			// Otherwise we assume it is a type filter
-			else
-			{
-				FilterBarConfig->TypeFilters.Add(FilterName, Filter->IsEnabled());
-			}
-		}
-
-		SaveConfig();
-	}
-
-	virtual void LoadSettings()
-	{
-		const FFilterBarSettings* FilterBarConfig = GetConstConfig();
-		
-		if(!FilterBarConfig)
-		{
-			UE_LOG(LogSlate, Error, TEXT("SFilterBar Requires that you specify a FilterBarIdentifier to load settings"));
-			return;
-		}
-
-		if(!FilterBarConfig->CustomTextFilters.IsEmpty())
-		{
-			// We must have a CreateTextFilter bound if we have any custom text filters saved!!
-			check(this->CreateTextFilter.IsBound());
-		}
-
-		// Load all the custom filters (i.e FrontendFilters)
-		for ( auto FrontendFilterIt = this->AllFrontendFilters.CreateIterator(); FrontendFilterIt; ++FrontendFilterIt )
-		{
-			TSharedRef<FFilterBase<FilterType>>& FrontendFilter = *FrontendFilterIt;
-			const FString& FilterName = FrontendFilter->GetName();
-			
-			if (!this->IsFrontendFilterInUse(FrontendFilter))
-			{
-				// Try to find this type filter in our list of saved filters
-				if ( const bool* bIsActive = FilterBarConfig->CustomFilters.Find(FilterName) )
-				{
-					TSharedRef<SFilter> NewFilter = this->AddFilterToBar(FrontendFilter);
-					
-					NewFilter->SetEnabled(*bIsActive, false);
-					this->SetFrontendFilterActive(FrontendFilter, NewFilter->IsEnabled());
-				}
-			}
-		}
-		
-		// Load all the type filters
-		for(const TSharedRef<FCustomClassFilterData> &CustomClassFilter : this->CustomClassFilters)
-		{
-			if(!this->IsClassTypeInUse(CustomClassFilter))
-			{
-				const FString FilterName = CustomClassFilter->GetFilterName();
-
-				// Try to find this type filter in our list of saved filters
-				if ( const bool* bIsActive = FilterBarConfig->TypeFilters.Find(CustomClassFilter->GetFilterName()) )
-				{
-					TSharedRef<SFilter> NewFilter = this->AddAssetFilterToBar(CustomClassFilter);
-					NewFilter->SetEnabled(*bIsActive, false);
-				}
-			}
-		}
-
-		// Load all the custom text filters
-		for(const FCustomTextFilterState& FilterState : FilterBarConfig->CustomTextFilters)
-		{
-			// Create a TTextFilter for the current filter we are making using the provided delegate
-			TSharedRef<TTextFilter<FilterType>> NewTextFilter = this->CreateTextFilter.Execute().ToSharedRef();
-
-			// Create the actual text filter
-			TSharedRef<FCustomTextFilter<FilterType>> NewFilter = MakeShared<FCustomTextFilter<FilterType>>(NewTextFilter);
-
-			// Set the internals of the custom text filter from what we have saved
-			NewFilter->SetFromCustomTextFilterData(FilterState.FilterData);
-
-			// Add this to our list of custom text filters
-			this->CustomTextFilters.Add(NewFilter);
-
-			// If the filter was checked previously, add it to the filter bar
-			if(FilterState.bIsChecked)
-			{
-				TSharedRef<SFilter> AddedFilter = this->AddFilterToBar(NewFilter);
-
-				// Set the filter as active if it was previously
-				AddedFilter->SetEnabled(FilterState.bIsActive, false);
-				
-				this->SetFrontendFilterActive(NewFilter, FilterState.bIsActive);
-			}
-		}
-
-		this->OnFilterChanged.ExecuteIfBound();
-	}
-
 private:
 
 	void UpdateAssetFilter()
diff --git a/Engine/Source/Editor/EditorWidgets/Public/SAssetSearchBox.h b/Engine/Source/Editor/EditorWidgets/Public/SAssetSearchBox.h
--- a/Engine/Source/Editor/EditorWidgets/Public/SAssetSearchBox.h
+++ b/Engine/Source/Editor/EditorWidgets/Public/SAssetSearchBox.h
@@ -10,6 +10,7 @@
 #include "Widgets/SCompoundWidget.h"
 #include "Widgets/Views/STableViewBase.h"
 #include "Widgets/Views/STableRow.h"
+#include "Filters/SFilterSearchBox.h"
 
 class SMenuAnchor;
 
@@ -42,6 +43,8 @@ DECLARE_DELEGATE_RetVal_TwoParams(FText, FOnAssetSearchBoxSuggestionChosen, cons
 
 /**
  * A widget to provide a search box with a filtered dropdown menu.
+ * Also provides the functionality to show search history optionally, and a button/delegate to save a search via
+ * SetOnSaveSearchHandler
  */
 
 class EDITORWIDGETS_API SAssetSearchBox : public SCompoundWidget
@@ -56,6 +59,7 @@ public:
 		, _PossibleSuggestions(TArray<FAssetSearchBoxSuggestion>())
 		, _DelayChangeNotificationsWhileTyping( false )
 		, _MustMatchPossibleSuggestions( false )
+		, _ShowSearchHistory(false)
 	{}
 
 		/** Where to place the suggestion list */
@@ -90,6 +94,12 @@ public:
 
 		/** Callback delegate to have first chance handling of the OnKeyDown event */
 		SLATE_EVENT( FOnKeyDown, OnKeyDownHandler )
+	
+		/** Whether we should show a dropdown containing the last few searches */
+		SLATE_ATTRIBUTE(bool, ShowSearchHistory)
+
+		/** Handler for when the + Button next to a search is clicked */
+		SLATE_EVENT(SFilterSearchBox::FOnSaveSearchClicked, OnSaveSearchClicked)
 
 	SLATE_END_ARGS()
 
@@ -103,6 +113,9 @@ public:
 	void SetError( const FText& InError );
 	void SetError( const FString& InError );
 
+	/** Show a + button next to the current search and set the handler for when that is clicked */
+	void SetOnSaveSearchHandler(SFilterSearchBox::FOnSaveSearchClicked InOnSaveSearchHandler);
+	
 	// SWidget implementation
 	virtual FReply OnPreviewKeyDown( const FGeometry& MyGeometry, const FKeyEvent& InKeyEvent ) override;
 	virtual bool SupportsKeyboardFocus() const override;
@@ -152,7 +165,7 @@ private:
 
 private:
 	/** The editable text field */
-	TSharedPtr< SSearchBox > InputText;
+	TSharedPtr< SFilterSearchBox > InputText;
 
 	/** The the state of the text prior to being committed */
 	FText PreCommittedText;
