Content Browser: Fix a bug where you could have duplicate filters in the filter bar, and fix a regression where docking the Content Drawer would discard all the active filters

diff --git a/Engine/Source/Developer/ToolWidgets/Public/Filters/SBasicFilterBar.h b/Engine/Source/Developer/ToolWidgets/Public/Filters/SBasicFilterBar.h
--- a/Engine/Source/Developer/ToolWidgets/Public/Filters/SBasicFilterBar.h
+++ b/Engine/Source/Developer/ToolWidgets/Public/Filters/SBasicFilterBar.h
@@ -840,6 +840,8 @@ protected:
 	{
 		for (const TSharedRef<SFilter>& Filter : Filters)
 		{
+			RemoveWidgetFromLayout(Filter);
+			
 			if (Filter == FilterToKeep)
 			{
 				continue;
diff --git a/Engine/Source/Editor/ContentBrowser/Private/SContentBrowser.cpp b/Engine/Source/Editor/ContentBrowser/Private/SContentBrowser.cpp
--- a/Engine/Source/Editor/ContentBrowser/Private/SContentBrowser.cpp
+++ b/Engine/Source/Editor/ContentBrowser/Private/SContentBrowser.cpp
@@ -1718,7 +1718,7 @@ void SContentBrowser::LoadSettings(const FName& InInstanceName)
 	}
 
 	// Save all our data using the settings string as a key in the user settings ini
-	FilterListPtr->LoadSettings();
+	FilterListPtr->LoadSettings(InInstanceName);
 	PathViewPtr->LoadSettings(GEditorPerProjectIni, SettingsIniSection, SettingsString);
 	FavoritePathViewPtr->LoadSettings(GEditorPerProjectIni, SettingsIniSection, SettingsString + TEXT(".Favorites"));
 	CollectionViewPtr->LoadSettings(GEditorPerProjectIni, SettingsIniSection, SettingsString);
diff --git a/Engine/Source/Editor/ContentBrowser/Private/SFilterList.cpp b/Engine/Source/Editor/ContentBrowser/Private/SFilterList.cpp
--- a/Engine/Source/Editor/ContentBrowser/Private/SFilterList.cpp
+++ b/Engine/Source/Editor/ContentBrowser/Private/SFilterList.cpp
@@ -637,7 +637,7 @@ void SFilterList::SaveSettings()
 	SaveConfig();
 }
 
-void SFilterList::LoadSettings()
+void SFilterList::LoadSettings(const FName& InInstanceName)
 {
 	// If this instance doesn't want to use the shared settings, load the settings normally
 	if(!bUseSharedSettings)
@@ -646,14 +646,14 @@ void SFilterList::LoadSettings()
 		return;
 	}
 
-	if(FilterBarIdentifier.IsNone())
+	if(InInstanceName.IsNone())
 	{
 		UE_LOG(LogSlate, Error, TEXT("SFilterList Requires that you specify a FilterBarIdentifier to load settings"));
 		return;
 	}
 
 	// Get the settings unique to this instance and the common settings
-	const FFilterBarSettings* InstanceSettings = UFilterBarConfig::Get()->FilterBars.Find(FilterBarIdentifier);
+	const FFilterBarSettings* InstanceSettings = UFilterBarConfig::Get()->FilterBars.Find(InInstanceName);
 	const FFilterBarSettings* SharedSettings = UFilterBarConfig::Get()->FilterBars.Find(SharedIdentifier);
 
 	// Load the filters specified programatically normally
@@ -679,6 +679,11 @@ void SFilterList::LoadSettings()
 	this->OnFilterChanged.ExecuteIfBound();
 }
 
+void SFilterList::LoadSettings()
+{
+	LoadSettings(FilterBarIdentifier);
+}
+
 void SFilterList::LoadCustomTextFilters(const FFilterBarSettings* FilterBarConfig)
 {
 	CustomTextFilters.Empty();
@@ -721,17 +726,21 @@ void SFilterList::AddWidgetToCurrentLayout(TSharedRef<SWidget> InWidget)
 void SFilterList::SetFilterLayout(EFilterBarLayout InFilterBarLayout)
 {
 	FilterBarLayout = InFilterBarLayout;
+
+	/* Clear both layouts, because for SFilterList it is valid to call SetFilterLayout with InFilterBarLayout being the
+	 * same as the current layout just to fire OnFilterBarLayoutChanging.
+	 * Unlike the parent class SBasicFilterBar which guards against that. If we don't clear both child widgets you can
+	 * end up with duplicate widgets.
+	 */
+	HorizontalFilterBox->ClearChildren();
+	VerticalFilterBox->ClearChildren();
  		
 	if(FilterBarLayout == EFilterBarLayout::Horizontal)
 	{
-		VerticalFilterBox->ClearChildren();
- 			
 		FilterBox->SetActiveWidget(HorizontalFilterBox.ToSharedRef());
 	}
 	else
 	{
-		HorizontalFilterBox->ClearChildren();
-        	
 		FilterBox->SetActiveWidget(VerticalFilterBox.ToSharedRef());
 	}
 
diff --git a/Engine/Source/Editor/ContentBrowser/Private/SFilterList.h b/Engine/Source/Editor/ContentBrowser/Private/SFilterList.h
--- a/Engine/Source/Editor/ContentBrowser/Private/SFilterList.h
+++ b/Engine/Source/Editor/ContentBrowser/Private/SFilterList.h
@@ -120,9 +120,12 @@ public:
 	/** Add a custom widget to the filter bar alongside the filters */
 	void AddWidgetToCurrentLayout(TSharedRef<SWidget> InWidget);
 
-	virtual void SaveSettings()override;
+	virtual void SaveSettings() override;
 	virtual void LoadSettings() override;
 
+	/* Copy the settings from a specific instance name */
+	void LoadSettings(const FName& InInstanceName);
+
 	virtual void SetFilterLayout(EFilterBarLayout InFilterBarLayout) override;
 
 protected:
