Allow actors to change which asset they Browse to
Allow browsing to aliases if source asset is not visible
Allow aliases in collections

diff --git a/Engine/Plugins/Editor/ContentBrowser/ContentBrowserAliasDataSoure/Source/ContentBrowserAliasDataSource/Private/ContentBrowserAliasDataSource.cpp b/Engine/Plugins/Editor/ContentBrowser/ContentBrowserAliasDataSoure/Source/ContentBrowserAliasDataSource/Private/ContentBrowserAliasDataSource.cpp
--- a/Engine/Plugins/Editor/ContentBrowser/ContentBrowserAliasDataSoure/Source/ContentBrowserAliasDataSource/Private/ContentBrowserAliasDataSource.cpp
+++ b/Engine/Plugins/Editor/ContentBrowser/ContentBrowserAliasDataSoure/Source/ContentBrowserAliasDataSource/Private/ContentBrowserAliasDataSource.cpp
@@ -5,11 +5,14 @@
 #include "ContentBrowserAssetDataCore.h"
 #include "ContentBrowserAssetDataSource.h"
 #include "ContentBrowserDataSubsystem.h"
+#include "ContentBrowserItemPath.h"
 
 #include "AssetToolsModule.h"
 #include "ObjectTools.h"
 #include "Misc/NamePermissionList.h"
 #include "Modules/ModuleManager.h"
+#include "CollectionManagerModule.h"
+#include "ICollectionManager.h"
 
 DEFINE_LOG_CATEGORY(LogContentBrowserAliasDataSource);
 
@@ -58,7 +61,7 @@ void UContentBrowserAliasDataSource::BuildRootPathVirtualTree()
 void UContentBrowserAliasDataSource::CompileFilter(const FName InPath, const FContentBrowserDataFilter& InFilter, FContentBrowserDataCompiledFilter& OutCompiledFilter)
 {
 	UContentBrowserAssetDataSource::FAssetFilterInputParams Params;
-	if (UContentBrowserAssetDataSource::PopulateAssetFilterInputParams(Params, this, AssetRegistry, InFilter, OutCompiledFilter))
+	if (UContentBrowserAssetDataSource::PopulateAssetFilterInputParams(Params, this, AssetRegistry, InFilter, OutCompiledFilter, &FCollectionManagerModule::GetModule().Get()))
 	{
 		// Use the DataSource's custom PathTree instead of the AssetRegistry
 		const bool bCreatedPathFilter = UContentBrowserAssetDataSource::CreatePathFilter(Params, InPath, InFilter, OutCompiledFilter, [this](FName Path, TFunctionRef<bool(FName)> Callback, bool bRecursive)
@@ -70,7 +73,7 @@ void UContentBrowserAliasDataSource::CompileFilter(const FName InPath, const FCo
 		{
 			UContentBrowserAssetDataSource::CreateAssetFilter(Params, InPath, InFilter, OutCompiledFilter, [this, &Params](FARFilter& InputFilter, FARCompiledFilter& OutputFilter)
 			{
-				// Same as above - CompileFilter calls EnumerateSubPaths internally so this needs to intercept
+				// Same as CreatePathFilter - CompileFilter calls EnumerateSubPaths internally so this needs to intercept
 				// the filter compilation and use its own PathTree to generate the sub paths.
 				if (InputFilter.bRecursivePaths)
 				{
@@ -274,22 +277,28 @@ bool UContentBrowserAliasDataSource::DoesAliasPassFilter(const FAliasData& Alias
 {
 	// Create a fake asset data using the alias path instead of the asset's original path
 	// AssetRegistry->IsAssetIncludedByFilter is effectively a static function and does not actually use AssetRegistry data
-	FAssetData AliasAssetData(AliasData.AliasName, AliasData.PackagePath, AliasData.AliasName, AliasData.AssetData.AssetClassPath, AliasData.AssetData.TagsAndValues.CopyMap());
+	FAssetData AliasAssetData(AliasData.PackageName, AliasData.PackagePath, AliasData.AssetData.AssetName, AliasData.AssetData.AssetClassPath, AliasData.AssetData.TagsAndValues.CopyMap());
 
 	return (Filter.InclusiveFilter.IsEmpty() || AssetRegistry->IsAssetIncludedByFilter(AliasAssetData, Filter.InclusiveFilter)) // Passes Inclusive
 		&& (Filter.ExclusiveFilter.IsEmpty() || !AssetRegistry->IsAssetIncludedByFilter(AliasAssetData, Filter.ExclusiveFilter)); // Passes Exclusive
 }
 
-FName UContentBrowserAliasDataSource::GetAliasForObjectPath(const FName ObjectPath) const
+TArray<FContentBrowserItemPath> UContentBrowserAliasDataSource::GetAliasesForPath(const FName ObjectPath) const
 {
-	if (const TArray<FName>* Aliases = AliasesForObjectPath.Find(ObjectPath))
+	TArray<FContentBrowserItemPath> OutAliases;
+
+	if (const TArray<FName>* FoundAliases = AliasesForObjectPath.Find(ObjectPath))
 	{
-		if (Aliases->Num() > 0)
+		for (const FName Alias : *FoundAliases)
 		{
-			return (*Aliases)[0];
+			if (const FAliasData* AliasData = AllAliases.Find(FContentBrowserUniqueAlias(ObjectPath, Alias)))
+			{
+				OutAliases.Add(FContentBrowserItemPath(AliasData->PackageName, EContentBrowserPathType::Internal));
+			}
 		}
 	}
-	return NAME_None;
+
+	return OutAliases;
 }
 
 void UContentBrowserAliasDataSource::AddAliases(const FAssetData& Asset, const TArray<FName>& Aliases, bool bInIsFromMetaData)
@@ -707,6 +716,19 @@ bool UContentBrowserAliasDataSource::UpdateThumbnail(const FContentBrowserItemDa
 	return ContentBrowserAssetData::UpdateItemThumbnail(this, InItem, InThumbnail);
 }
 
+bool UContentBrowserAliasDataSource::TryGetCollectionId(const FContentBrowserItemData& InItem, FName& OutCollectionId)
+{
+	if (TSharedPtr<const FContentBrowserAliasItemDataPayload> AliasPayload = StaticCastSharedPtr<const FContentBrowserAliasItemDataPayload>(InItem.GetPayload()))
+	{
+		if (const FAliasData* AliasData = AllAliases.Find(AliasPayload->Alias))
+		{
+			OutCollectionId = AliasData->ObjectPath;
+			return true;
+		}
+	}
+	return false;
+}
+
 bool UContentBrowserAliasDataSource::Legacy_TryGetPackagePath(const FContentBrowserItemData& InItem, FName& OutPackagePath)
 {
 	if (TSharedPtr<const FContentBrowserAssetFolderItemDataPayload> FolderPayload = ContentBrowserAssetData::GetAssetFolderItemPayload(this, InItem))
@@ -756,7 +778,7 @@ FContentBrowserItemData UContentBrowserAliasDataSource::CreateAssetFileItem(cons
 	if (AliasData)
 	{
 		FName VirtualizedPath;
-		TryConvertInternalPathToVirtual(AliasData->AliasID, VirtualizedPath);
+		TryConvertInternalPathToVirtual(AliasData->ObjectPath, VirtualizedPath);
 
 		// Since AliasID is PackagePath/AssetName, AssetName should also be passed as the ItemName here. This provides the functionality of
 		// being able to have multiple aliases with the same display name, while still showing their original asset name in the tooltip.
diff --git a/Engine/Plugins/Editor/ContentBrowser/ContentBrowserAliasDataSoure/Source/ContentBrowserAliasDataSource/Public/ContentBrowserAliasDataSource.h b/Engine/Plugins/Editor/ContentBrowser/ContentBrowserAliasDataSoure/Source/ContentBrowserAliasDataSource/Public/ContentBrowserAliasDataSource.h
--- a/Engine/Plugins/Editor/ContentBrowser/ContentBrowserAliasDataSoure/Source/ContentBrowserAliasDataSource/Public/ContentBrowserAliasDataSource.h
+++ b/Engine/Plugins/Editor/ContentBrowser/ContentBrowserAliasDataSoure/Source/ContentBrowserAliasDataSource/Public/ContentBrowserAliasDataSource.h
@@ -58,6 +58,8 @@ public:
 
 	virtual bool DoesItemPassFilter(const FContentBrowserItemData& InItem, const FContentBrowserDataCompiledFilter& InFilter) override;
 
+	virtual TArray<FContentBrowserItemPath> GetAliasesForPath(const FName InInternalPath) const override;
+
 	virtual bool GetItemAttribute(const FContentBrowserItemData& InItem, const bool InIncludeMetaData, const FName InAttributeKey, FContentBrowserItemDataAttributeValue& OutAttributeValue) override;
 	virtual bool GetItemAttributes(const FContentBrowserItemData& InItem, const bool InIncludeMetaData, FContentBrowserItemDataAttributeValues& OutAttributeValues) override;
 	virtual bool GetItemPhysicalPath(const FContentBrowserItemData& InItem, FString& OutDiskPath) override;
@@ -78,6 +80,8 @@ public:
 	virtual bool AppendItemReference(const FContentBrowserItemData& InItem, FString& InOutStr) override;
 	virtual bool UpdateThumbnail(const FContentBrowserItemData& InItem, FAssetThumbnail& InThumbnail) override;
 
+	virtual bool TryGetCollectionId(const FContentBrowserItemData& InItem, FName& OutCollectionId) override;
+
 	// Legacy functions seem necessary for FrontendFilters to work
 	virtual bool Legacy_TryGetPackagePath(const FContentBrowserItemData& InItem, FName& OutPackagePath) override;
 	virtual bool Legacy_TryGetAssetData(const FContentBrowserItemData& InItem, FAssetData& OutAssetData) override;
@@ -85,8 +89,6 @@ public:
 	virtual bool Legacy_TryConvertAssetDataToVirtualPath(const FAssetData& InAssetData, const bool InUseFolderPaths, FName& OutPath) override;
 	// ~ End UContentBrowserDataSource interface
 
-	FName GetAliasForObjectPath(const FName ObjectPath) const;
-
 	/** Add a list of aliases for a given asset. bInIsFromMetaData should only be true if the list of aliases came from the AliasTagName metadata. */
 	void AddAliases(const FAssetData& Asset, const TArray<FName>& Aliases, bool bInIsFromMetaData = false);
 	/** Add an alias for a given asset. bInIsFromMetaData should only be true if the alias came from the AliasTagName metadata. */
@@ -129,15 +131,24 @@ private:
 		FAliasData(const FAssetData& InAssetData, const FName InPackagePath, const FName InName, const bool bInIsFromMetaData = false)
 			: AssetData(InAssetData), PackagePath(InPackagePath), AliasName(InName), bIsFromMetaData(bInIsFromMetaData)
 		{
-			AliasID = *(InPackagePath.ToString() / InAssetData.AssetName.ToString());
+			FNameBuilder PathBuilder(PackagePath);
+			PathBuilder << TEXT('/');
+			PathBuilder << InAssetData.AssetName;
+			PackageName = FName(PathBuilder.ToView());
+
+			PathBuilder << TEXT('.');
+			PathBuilder << InAssetData.AssetName;
+			ObjectPath = FName(PathBuilder.ToView());
 		}
 
 		/** The source asset for this alias */
 		FAssetData AssetData;
-		/** The package path of the alias */
+		/** The folder path that contains the alias, /MyAliases */
 		FName PackagePath;
-		/** A unique path for this alias, which is PackagePath/SourceAssetName */
-		FName AliasID;
+		/** PackagePath/SourceAssetName, /MyAliases/SomeAsset  */
+		FName PackageName;
+		/** PackageName.SourceAssetName, /MyAliases/SomeAsset.SomeAsset */
+		FName ObjectPath;
 		/** A non-unique display name for this alias */
 		FName AliasName;
 		/** Whether this alias was generated from package metadata or manually through the C++ interface */
diff --git a/Engine/Source/Editor/ContentBrowser/Private/SContentBrowser.cpp b/Engine/Source/Editor/ContentBrowser/Private/SContentBrowser.cpp
--- a/Engine/Source/Editor/ContentBrowser/Private/SContentBrowser.cpp
+++ b/Engine/Source/Editor/ContentBrowser/Private/SContentBrowser.cpp
@@ -82,10 +82,6 @@
 
 const FString SContentBrowser::SettingsIniSection = TEXT("ContentBrowser");
 
-// Workaround to hide Save As Collection button until collections support the AliasDataSource
-bool bHideSaveSearchButton = false;
-FAutoConsoleVariableRef CVarHideSaveSearchButton(TEXT("ContentBrowser.HideSaveSearchButton"), bHideSaveSearchButton, TEXT("Hide the Content Browser button to save search as a dynamic collection."));
-
 namespace ContentBrowserSourcesWidgetSwitcherIndex
 {
 	static const int32 PathView = 0;
@@ -1412,6 +1408,36 @@ void SContentBrowser::SyncToFolders(TArrayView<const FString> FolderList, const
 
 void SContentBrowser::SyncToItems(TArrayView<const FContentBrowserItem> ItemsToSync, const bool bAllowImplicitSync, const bool bDisableFiltersThatHideAssets)
 {
+	UContentBrowserDataSubsystem* ContentBrowserData = IContentBrowserDataModule::Get().GetSubsystem();
+	const TSharedRef<FPathPermissionList>& FolderPermissions = FAssetToolsModule::GetModule().Get().GetFolderPermissionList();
+
+	// If any of the items to sync don't pass the permission filter, try to find an alias that might be visible
+	TArray<FContentBrowserItem> NewItemsToSync;
+	for (const FContentBrowserItem& Item : ItemsToSync)
+	{
+		if (FolderPermissions->PassesStartsWithFilter(Item.GetInternalPath()))
+		{
+			NewItemsToSync.Add(Item);
+		}
+		else
+		{
+			TArray<FContentBrowserItemPath> Aliases = ContentBrowserData->GetAliasesForPath(Item.GetInternalPath());
+			for (const FContentBrowserItemPath& Alias : Aliases)
+			{
+				if (FolderPermissions->PassesStartsWithFilter(Alias.GetInternalPathName()))
+				{
+					FContentBrowserItem AliasedItem = ContentBrowserData->GetItemAtPath(Alias.GetVirtualPathName(), EContentBrowserItemTypeFilter::IncludeFiles);
+					if (AliasedItem.IsValid())
+					{
+						NewItemsToSync.Add(MoveTemp(AliasedItem));
+						break;
+					}
+				}
+			}
+		}
+	}
+
+	ItemsToSync = NewItemsToSync;
 	PrepareToSyncItems(ItemsToSync, bDisableFiltersThatHideAssets);
 
 	// Tell the sources view first so the asset view will be up to date by the time we request the sync
@@ -1422,6 +1448,32 @@ void SContentBrowser::SyncToItems(TArrayView<const FContentBrowserItem> ItemsToS
 
 void SContentBrowser::SyncToVirtualPaths(TArrayView<const FName> VirtualPathsToSync, const bool bAllowImplicitSync, const bool bDisableFiltersThatHideAssets)
 {
+	UContentBrowserDataSubsystem* ContentBrowserData = IContentBrowserDataModule::Get().GetSubsystem();
+	const TSharedRef<FPathPermissionList>& FolderPermissions = FAssetToolsModule::GetModule().Get().GetFolderPermissionList();
+
+	// If any of the items to sync don't pass the permission filter, try to map the item to a different one that might be visible
+	TArray<FName> NewItemsToSync;
+	for (const FName VirtualPath : VirtualPathsToSync)
+	{
+		FName InternalPath;
+		ContentBrowserData->TryConvertVirtualPath(VirtualPath, InternalPath);
+		if (FolderPermissions->PassesStartsWithFilter(InternalPath))
+		{
+			NewItemsToSync.Add(VirtualPath);
+		}
+		else
+		{
+			TArray<FContentBrowserItemPath> Aliases = ContentBrowserData->GetAliasesForPath(InternalPath);
+			for (const FContentBrowserItemPath& Alias : Aliases)
+			{
+				if (FolderPermissions->PassesStartsWithFilter(Alias.GetInternalPathName()))
+				{
+					NewItemsToSync.Add(Alias.GetVirtualPathName());
+				}
+			}
+		}
+	}
+	VirtualPathsToSync = NewItemsToSync;
 	PrepareToSyncVirtualPaths(VirtualPathsToSync, bDisableFiltersThatHideAssets);
 
 	// Tell the sources view first so the asset view will be up to date by the time we request the sync
@@ -1432,6 +1484,34 @@ void SContentBrowser::SyncToVirtualPaths(TArrayView<const FName> VirtualPathsToS
 
 void SContentBrowser::SyncToLegacy(TArrayView<const FAssetData> AssetDataList, TArrayView<const FString> FolderList, const bool bAllowImplicitSync, const bool bDisableFiltersThatHideAssets)
 {
+	UContentBrowserDataSubsystem* ContentBrowserData = IContentBrowserDataModule::Get().GetSubsystem();
+	const TSharedRef<FPathPermissionList>& FolderPermissions = FAssetToolsModule::GetModule().Get().GetFolderPermissionList();
+
+	// If any of the items to sync don't pass the permission filter, try to map the item to a different one that might be visible
+	TArray<FAssetData> NewItemsToSync;
+	for (int32 i = 0; i < AssetDataList.Num(); ++i)
+	{
+		if (FolderPermissions->PassesStartsWithFilter(AssetDataList[i].ObjectPath))
+		{
+			NewItemsToSync.Add(AssetDataList[i]);
+		}
+		else
+		{
+			TArray<FContentBrowserItemPath> Aliases = ContentBrowserData->GetAliasesForPath(AssetDataList[i].ObjectPath);
+			for (const FContentBrowserItemPath& Alias : Aliases)
+			{
+				const FName InternalPath = Alias.GetInternalPathName();
+				if (FolderPermissions->PassesStartsWithFilter(InternalPath))
+				{
+					FAssetData AliasAssetData(InternalPath, *FPackageName::GetLongPackagePath(InternalPath.ToString()), FPackageName::GetShortFName(InternalPath), AssetDataList[i].AssetClassPath, AssetDataList[i].TagsAndValues.CopyMap());
+					NewItemsToSync.Add(MoveTemp(AliasAssetData));
+					break;
+				}
+			}
+		}
+	}
+
+	AssetDataList = NewItemsToSync;
 	PrepareToSyncLegacy(AssetDataList, FolderList, bDisableFiltersThatHideAssets);
 
 	// Tell the sources view first so the asset view will be up to date by the time we request the sync
@@ -1444,21 +1524,11 @@ void SContentBrowser::SyncTo( const FContentBrowserSelection& ItemSelection, con
 {
 	if (ItemSelection.IsLegacy())
 	{
-		PrepareToSyncLegacy(ItemSelection.SelectedAssets, ItemSelection.SelectedFolders, bDisableFiltersThatHideAssets);
-
-		// Tell the sources view first so the asset view will be up to date by the time we request the sync
-		PathViewPtr->SyncToLegacy(ItemSelection.SelectedAssets, ItemSelection.SelectedFolders, bAllowImplicitSync);
-		FavoritePathViewPtr->SyncToLegacy(ItemSelection.SelectedAssets, ItemSelection.SelectedFolders, bAllowImplicitSync);
-		AssetViewPtr->SyncToLegacy(ItemSelection.SelectedAssets, ItemSelection.SelectedFolders);
+		SyncToLegacy(ItemSelection.SelectedAssets, ItemSelection.SelectedFolders, bAllowImplicitSync, bDisableFiltersThatHideAssets);
 	}
 	else
 	{
-		PrepareToSyncItems(ItemSelection.SelectedItems, bDisableFiltersThatHideAssets);
-
-		// Tell the sources view first so the asset view will be up to date by the time we request the sync
-		PathViewPtr->SyncToItems(ItemSelection.SelectedItems, bAllowImplicitSync);
-		FavoritePathViewPtr->SyncToItems(ItemSelection.SelectedItems, bAllowImplicitSync);
-		AssetViewPtr->SyncToItems(ItemSelection.SelectedItems);
+		SyncToItems(ItemSelection.SelectedItems, bAllowImplicitSync, bDisableFiltersThatHideAssets);
 	}
 }
 
diff --git a/Engine/Source/Editor/ContentBrowserData/Private/ContentBrowserDataSource.cpp b/Engine/Source/Editor/ContentBrowserData/Private/ContentBrowserDataSource.cpp
--- a/Engine/Source/Editor/ContentBrowserData/Private/ContentBrowserDataSource.cpp
+++ b/Engine/Source/Editor/ContentBrowserData/Private/ContentBrowserDataSource.cpp
@@ -217,6 +217,11 @@ bool UContentBrowserDataSource::EnumerateItemsForObjects(const TArrayView<UObjec
 	return true;
 }
 
+TArray<FContentBrowserItemPath> UContentBrowserDataSource::GetAliasesForPath(const FName InInternalPath) const
+{
+	return TArray<FContentBrowserItemPath>();
+}
+
 bool UContentBrowserDataSource::IsDiscoveringItems(FText* OutStatus)
 {
 	return false;
diff --git a/Engine/Source/Editor/ContentBrowserData/Private/ContentBrowserDataSubsystem.cpp b/Engine/Source/Editor/ContentBrowserData/Private/ContentBrowserDataSubsystem.cpp
--- a/Engine/Source/Editor/ContentBrowserData/Private/ContentBrowserDataSubsystem.cpp
+++ b/Engine/Source/Editor/ContentBrowserData/Private/ContentBrowserDataSubsystem.cpp
@@ -2,6 +2,7 @@
 
 #include "ContentBrowserDataSubsystem.h"
 #include "ContentBrowserDataSource.h"
+#include "ContentBrowserItemPath.h"
 #include "IContentBrowserDataModule.h"
 #include "Containers/Ticker.h"
 #include "Misc/PackageName.h"
@@ -441,6 +442,24 @@ FContentBrowserItem UContentBrowserDataSubsystem::GetItemAtPath(const FName InPa
 	return FoundItem;
 }
 
+TArray<FContentBrowserItemPath> UContentBrowserDataSubsystem::GetAliasesForPath(const FContentBrowserItemPath InPath) const
+{
+	return GetAliasesForPath(InPath.GetInternalPathName());
+}
+
+TArray<FContentBrowserItemPath> UContentBrowserDataSubsystem::GetAliasesForPath(const FName InInternalPath) const
+{
+	TArray<FContentBrowserItemPath> Aliases;
+
+	for (const auto& ActiveDataSourcePair : ActiveDataSources)
+	{
+		UContentBrowserDataSource* DataSource = ActiveDataSourcePair.Value;
+		Aliases.Append(DataSource->GetAliasesForPath(InInternalPath));
+	}
+
+	return Aliases;
+}
+
 bool UContentBrowserDataSubsystem::IsDiscoveringItems(TArray<FText>* OutStatus) const
 {
 	bool bIsDiscoveringItems = false;
diff --git a/Engine/Source/Editor/ContentBrowserData/Public/ContentBrowserDataSource.h b/Engine/Source/Editor/ContentBrowserData/Public/ContentBrowserDataSource.h
--- a/Engine/Source/Editor/ContentBrowserData/Public/ContentBrowserDataSource.h
+++ b/Engine/Source/Editor/ContentBrowserData/Public/ContentBrowserDataSource.h
@@ -212,6 +212,14 @@ public:
 	 */
 	virtual bool EnumerateItemsForObjects(const TArrayView<UObject*> InObjects, TFunctionRef<bool(FContentBrowserItemData&&)> InCallback);
 
+	/**
+	 * Get a list of other paths that the data source may be using to represent a specific path
+	 *
+	 * @param The internal path (or object path) of an asset to get aliases for
+	 * @return All alternative paths that represent the input path (not including the input path itself)
+	 */
+	virtual TArray<FContentBrowserItemPath> GetAliasesForPath(const FName InInternalPath) const;
+
 	/**
 	 * Query whether this data source instance is currently discovering content, and retrieve an optional status message that can be shown in the UI.
 	 */
diff --git a/Engine/Source/Editor/ContentBrowserData/Public/ContentBrowserDataSubsystem.h b/Engine/Source/Editor/ContentBrowserData/Public/ContentBrowserDataSubsystem.h
--- a/Engine/Source/Editor/ContentBrowserData/Public/ContentBrowserDataSubsystem.h
+++ b/Engine/Source/Editor/ContentBrowserData/Public/ContentBrowserDataSubsystem.h
@@ -12,6 +12,8 @@
 #include "Containers/Ticker.h"
 #include "ContentBrowserDataSubsystem.generated.h"
 
+struct FContentBrowserItemPath;
+
 UENUM(BlueprintType)
 enum class EContentBrowserPathType : uint8
 {
@@ -189,6 +191,15 @@ public:
 	UFUNCTION(BlueprintCallable, Category="ContentBrowser")
 	FContentBrowserItem GetItemAtPath(const FName InPath, const EContentBrowserItemTypeFilter InItemTypeFilter) const;
 
+	/**
+	 * Get a list of other paths that the data source may be using to represent a specific path
+	 *
+	 * @param The internal path (or object path) of an asset to get aliases for
+	 * @return All alternative paths that represent the input path (not including the input path itself)
+	 */
+	TArray<FContentBrowserItemPath> GetAliasesForPath(const FName InInternalPath) const;
+	TArray<FContentBrowserItemPath> GetAliasesForPath(const FContentBrowserItemPath InPath) const;
+
 	/**
 	 * Query whether any data sources are currently discovering content, and retrieve optional status messages that can be shown in the UI.
 	 */
diff --git a/Engine/Source/Editor/UnrealEd/Classes/Editor/EditorEngine.h b/Engine/Source/Editor/UnrealEd/Classes/Editor/EditorEngine.h
--- a/Engine/Source/Editor/UnrealEd/Classes/Editor/EditorEngine.h
+++ b/Engine/Source/Editor/UnrealEd/Classes/Editor/EditorEngine.h
@@ -2127,8 +2127,10 @@ public:
 
 	/**
 	 * Syncs the selected actors objects to the content browser
+	 * 
+	 * @param bAllowOverrideMetadata If true, allows an asset to define "BrowseToAssetOverride" in its metadata to sync to an asset other than itself
 	 */
-	void SyncToContentBrowser();
+	void SyncToContentBrowser(bool bAllowOverrideMetadata = true);
 
 	/**
 	 * Syncs the selected actors' levels to the content browser
@@ -2445,8 +2447,17 @@ public:
 	 * Gets all objects which can be synced to in content browser for current selection
 	 *
 	 * @param Objects	Array to be filled with objects which can be browsed to
+	 * @param bAllowOverrideMetadata If true, allows an asset to define "BrowseToAssetOverride" in its metadata to sync to an asset other than itself
+	 */
+	UE_DEPRECATED(5.1, "Use GetAssetsToSyncToContentBrowser instead")
+	void GetObjectsToSyncToContentBrowser(TArray<UObject*>& Objects, bool bAllowBrowseToAssetOverride = true);
+	/**
+	 * Gets all assets which can be synced to in content browser for current selection
+	 *
+	 * @param Assets	Array to be filled with assets which can be browsed to
+	 * @param bAllowOverrideMetadata If true, allows an asset to define "BrowseToAssetOverride" in its metadata to sync to an asset other than itself
 	 */
-	void GetObjectsToSyncToContentBrowser( TArray<UObject*>& Objects );
+	void GetAssetsToSyncToContentBrowser(TArray<FAssetData>& Assets, bool bAllowBrowseToAssetOverride = true);
 
 	/**
 	 * Gets all levels which can be synced to in content browser for current selection
diff --git a/Engine/Source/Editor/UnrealEd/Private/EditorEngine.cpp b/Engine/Source/Editor/UnrealEd/Private/EditorEngine.cpp
--- a/Engine/Source/Editor/UnrealEd/Private/EditorEngine.cpp
+++ b/Engine/Source/Editor/UnrealEd/Private/EditorEngine.cpp
@@ -2904,64 +2904,98 @@ void UEditorEngine::SyncBrowserToObjects( TArray<struct FAssetData>& InAssetsToS
 
 bool UEditorEngine::CanSyncToContentBrowser()
 {
-	TArray< UObject*> Objects;
-	GetObjectsToSyncToContentBrowser( Objects );
-	return Objects.Num() > 0;
+	TArray<FAssetData> Assets;
+	GetAssetsToSyncToContentBrowser(Assets);
+	return Assets.Num() > 0;
 }
 
 
-void UEditorEngine::GetObjectsToSyncToContentBrowser( TArray<UObject*>& Objects )
+void UEditorEngine::GetObjectsToSyncToContentBrowser(TArray<UObject*>& Objects, bool bAllowBrowseToAssetOverride)
+{
+	TArray<FAssetData> Assets;
+	GetAssetsToSyncToContentBrowser(Assets, bAllowBrowseToAssetOverride);
+	for (const FAssetData& Asset : Assets)
+	{
+		Objects.Add(Asset.GetAsset());
+	}
+}
+
+void UEditorEngine::GetAssetsToSyncToContentBrowser(TArray<FAssetData>& Assets, bool bAllowBrowseToAssetOverride)
 {
 	// If the user has any BSP surfaces selected, sync to the materials on them.
 	bool bFoundSurfaceMaterial = false;
 
-	for ( TSelectedSurfaceIterator<> It(GWorld) ; It ; ++It )
+	for (TSelectedSurfaceIterator<> It(GWorld); It; ++It)
 	{
 		FBspSurf* Surf = *It;
 		UMaterialInterface* Material = Surf->Material;
-		if( Material )
+		if (Material)
 		{
-			Objects.AddUnique( Material );
+			Assets.AddUnique(FAssetData(Material));
 			bFoundSurfaceMaterial = true;
 		}
 	}
 
 	// Otherwise, assemble a list of resources from selected actors.
-	if( !bFoundSurfaceMaterial )
+	if (!bFoundSurfaceMaterial)
 	{
-		for ( FSelectionIterator It( GetSelectedActorIterator() ) ; It ; ++It )
+		for (FSelectionIterator It(GetSelectedActorIterator()); It; ++It)
 		{
-			AActor* Actor = static_cast<AActor*>( *It );
-			checkSlow( Actor->IsA(AActor::StaticClass()) );
+			AActor* Actor = static_cast<AActor*>(*It);
+			checkSlow(Actor->IsA(AActor::StaticClass()));
 
-			// If the actor is an instance of a blueprint, just add the blueprint.
-			UBlueprint* GeneratingBP = Cast<UBlueprint>(It->GetClass()->ClassGeneratedBy);
-			if ( GeneratingBP != NULL )
-			{
-				Objects.Add(GeneratingBP);
-			}
-			// Cooked editor sometimes only contains UBlueprintGeneratedClass with no UBlueprint
-			else if (UBlueprintGeneratedClass* BlueprintGeneratedClass = Cast<UBlueprintGeneratedClass>(It->GetClass()))
+			bool bFoundOverride = false;
+			if (bAllowBrowseToAssetOverride)
 			{
-				Objects.Add(BlueprintGeneratedClass);
+				// If BrowseToAssetOverride is set, then use the asset it points to instead of the selected asset
+				const FString& BrowseToAssetOverride = Actor->GetBrowseToAssetOverride();
+				if (!BrowseToAssetOverride.IsEmpty())
+				{
+					if (IAssetRegistry* AssetRegistry = IAssetRegistry::Get())
+					{
+						TArray<FAssetData> FoundAssets;
+						if (AssetRegistry->GetAssetsByPackageName(*BrowseToAssetOverride, FoundAssets) && FoundAssets.Num() > 0)
+						{
+							Assets.Add(FoundAssets[0]);
+							bFoundOverride = true;
+						}
+					}
+				}
 			}
-			// Otherwise, add the results of the GetReferencedContentObjects call
-			else
+
+			if (!bFoundOverride)
 			{
-				Actor->GetReferencedContentObjects(Objects);
+				// If the actor is an instance of a blueprint, just add the blueprint.
+				UBlueprint* GeneratingBP = Cast<UBlueprint>(It->GetClass()->ClassGeneratedBy);
+				if (GeneratingBP != NULL)
+				{
+					Assets.Add(FAssetData(GeneratingBP));
+				}
+				// Cooked editor sometimes only contains UBlueprintGeneratedClass with no UBlueprint
+				else if (UBlueprintGeneratedClass* BlueprintGeneratedClass = Cast<UBlueprintGeneratedClass>(It->GetClass()))
+				{
+					Assets.Add(FAssetData(BlueprintGeneratedClass));
+				}
+				// Otherwise, add the results of the GetReferencedContentObjects call
+				else
+				{
+					TArray<UObject*> Objects;
+					Actor->GetReferencedContentObjects(Objects);
+					for (UObject* Object : Objects)
+					{
+						Assets.Add(FAssetData(Object));
+					}
+				}
 			}
 		}
 	}
 }
 
-void UEditorEngine::SyncToContentBrowser()
+void UEditorEngine::SyncToContentBrowser(bool bAllowOverrideMetadata)
 {
-	TArray<UObject*> Objects;
-
-	GetObjectsToSyncToContentBrowser( Objects );
-
-	// Sync the content browser to the object list.
-	SyncBrowserToObjects(Objects);
+	TArray<FAssetData> Assets;
+	GetAssetsToSyncToContentBrowser(Assets, bAllowOverrideMetadata);
+	SyncBrowserToObjects(Assets);
 }
 
 void UEditorEngine::GetLevelsToSyncToContentBrowser(TArray<UObject*>& Objects)
diff --git a/Engine/Source/Runtime/Engine/Classes/GameFramework/Actor.h b/Engine/Source/Runtime/Engine/Classes/GameFramework/Actor.h
--- a/Engine/Source/Runtime/Engine/Classes/GameFramework/Actor.h
+++ b/Engine/Source/Runtime/Engine/Classes/GameFramework/Actor.h
@@ -2468,6 +2468,11 @@ public:
 	{
 		return bOptimizeBPComponentData;
 	}
+
+	/** Sets metadata on the Actor that allows it to point to a different asset than its source when browsing in the Content Browser */
+	void SetBrowseToAssetOverride(const FString& PackageName);
+	/** Gets metadata on the Actor that says which package to show in the Content Browser when browsing to it, or an empty string if there is no override */
+	const FString& GetBrowseToAssetOverride() const;
 #endif		// WITH_EDITOR
 
 	/**
diff --git a/Engine/Source/Runtime/Engine/Private/Actor.cpp b/Engine/Source/Runtime/Engine/Private/Actor.cpp
--- a/Engine/Source/Runtime/Engine/Private/Actor.cpp
+++ b/Engine/Source/Runtime/Engine/Private/Actor.cpp
@@ -573,6 +573,16 @@ bool AActor::NeedsLoadForTargetPlatform(const ITargetPlatform* TargetPlatform) c
 	return true;
 }
 
+void AActor::SetBrowseToAssetOverride(const FString& PackageName)
+{
+	GetPackage()->GetMetaData()->SetValue(this, "BrowseToAssetOverride", *PackageName);
+}
+
+const FString& AActor::GetBrowseToAssetOverride() const
+{
+	return GetPackage()->GetMetaData()->GetValue(this, "BrowseToAssetOverride");
+}
+
 #endif
 
 UWorld* AActor::GetWorld() const
