Content Browser: Move Filter Display section to the Settings Cog

diff --git a/Engine/Source/Editor/ContentBrowser/Private/SAssetView.cpp b/Engine/Source/Editor/ContentBrowser/Private/SAssetView.cpp
--- a/Engine/Source/Editor/ContentBrowser/Private/SAssetView.cpp
+++ b/Engine/Source/Editor/ContentBrowser/Private/SAssetView.cpp
@@ -59,6 +59,7 @@
 #include "ContentBrowserDataSubsystem.h"
 #include "ContentBrowserDataLegacyBridge.h"
 #include "ContentBrowserDataDragDropOp.h"
+#include "SFilterList.h"
 #include "SPrimaryButton.h"
 
 #define LOCTEXT_NAMESPACE "ContentBrowser"
@@ -2458,6 +2459,47 @@ void SAssetView::PopulateViewButtonMenu(UToolMenu* Menu)
 			);
 	}
 
+	if(TSharedPtr<SFilterList> FilterBarPinned = FilterBar.Pin())
+	{
+		FToolMenuSection& Section = Menu->AddSection("FilterBar", LOCTEXT("FilterBarHeading", "Filter Display"));
+
+		Section.AddMenuEntry(
+		"VerticalLayout",
+		LOCTEXT("FilterListVerticalLayout", "Vertical"),
+		LOCTEXT("FilterListVerticalLayoutToolTip", "Swap to a vertical layout for the filter bar"),
+		FSlateIcon(),
+		FUIAction(
+			FExecuteAction::CreateLambda([FilterBarPinned]()
+			{
+				if(FilterBarPinned->GetFilterLayout() != EFilterBarLayout::Vertical)
+				{
+					FilterBarPinned->SetFilterLayout(EFilterBarLayout::Vertical);
+				}
+			}),
+			FCanExecuteAction(),
+			FIsActionChecked::CreateLambda([FilterBarPinned]() { return FilterBarPinned->GetFilterLayout() == EFilterBarLayout::Vertical; })),
+		EUserInterfaceActionType::RadioButton
+	);
+
+		Section.AddMenuEntry(
+			"HorizontalLayout",
+			LOCTEXT("FilterListHorizontalLayout", "Horizontal"),
+			LOCTEXT("FilterListHorizontalLayoutToolTip", "Swap to a Horizontal layout for the filter bar"),
+			FSlateIcon(),
+			FUIAction(
+			FExecuteAction::CreateLambda([FilterBarPinned]()
+				{
+					if(FilterBarPinned->GetFilterLayout() != EFilterBarLayout::Horizontal)
+					{
+						FilterBarPinned->SetFilterLayout(EFilterBarLayout::Horizontal);
+					}
+				}),
+				FCanExecuteAction(),
+				FIsActionChecked::CreateLambda([FilterBarPinned]() { return FilterBarPinned->GetFilterLayout() == EFilterBarLayout::Horizontal; })),
+			EUserInterfaceActionType::RadioButton
+		);
+	}
+
 	{
 		FToolMenuSection& Section = Menu->AddSection("View", LOCTEXT("ViewHeading", "View"));
 
@@ -4606,6 +4648,11 @@ void SAssetView::HandleItemDataDiscoveryComplete()
 	}
 }
 
+void SAssetView::SetFilterBar(TSharedPtr<SFilterList> InFilterBar)
+{
+	FilterBar = InFilterBar;
+}
+
 #undef checkAssetList
 #undef ASSET_VIEW_PARANOIA_LIST_CHECKS
 
diff --git a/Engine/Source/Editor/ContentBrowser/Private/SContentBrowser.cpp b/Engine/Source/Editor/ContentBrowser/Private/SContentBrowser.cpp
--- a/Engine/Source/Editor/ContentBrowser/Private/SContentBrowser.cpp
+++ b/Engine/Source/Editor/ContentBrowser/Private/SContentBrowser.cpp
@@ -869,6 +869,8 @@ TSharedRef<SWidget> SContentBrowser::CreateAssetView(const FContentBrowserConfig
 	FilterComboButton->AddMetadata(FilterComboButtonMetaData.ToSharedRef());
 	FilterComboButton->SetVisibility((Config != nullptr ? Config->bCanShowFilters : true) ? EVisibility::Visible : EVisibility::Collapsed);
 
+	AssetViewPtr->SetFilterBar(FilterListPtr);
+
 	SearchBoxPtr = SNew(SAssetSearchBox)
 							.HintText(this, &SContentBrowser::GetSearchAssetsHintText)
 							.ShowSearchHistory(true)
diff --git a/Engine/Source/Editor/ContentBrowser/Private/SFilterList.cpp b/Engine/Source/Editor/ContentBrowser/Private/SFilterList.cpp
--- a/Engine/Source/Editor/ContentBrowser/Private/SFilterList.cpp
+++ b/Engine/Source/Editor/ContentBrowser/Private/SFilterList.cpp
@@ -306,47 +306,6 @@ bool IsFilteredByPicker(const TArray<UClass*>& FilterClassList, UClass* TestClas
 	return true;
 }
 
-void SFilterList::PopulateFilterDisplayMenu(UToolMenu* Menu)
-{
-	FToolMenuSection& Section = Menu->AddSection("FilterListFilterDisplay", LOCTEXT("FilterDisplay", "Filter Display"));
-
-	Section.AddMenuEntry(
-		"VerticalLayout",
-		LOCTEXT("FilterListVerticalLayout", "Vertical"),
-		LOCTEXT("FilterListVerticalLayoutToolTip", "Swap to a vertical layout for the filter bar"),
-		FSlateIcon(),
-		FUIAction(
-			FExecuteAction::CreateLambda([this]()
-			{
-				if(this->FilterBarLayout != EFilterBarLayout::Vertical)
-				{
-					this->SetFilterLayout(EFilterBarLayout::Vertical);
-				}
-			}),
-			FCanExecuteAction(),
-			FIsActionChecked::CreateLambda([this]() { return FilterBarLayout == EFilterBarLayout::Vertical; })),
-		EUserInterfaceActionType::RadioButton
-	);
-
-	Section.AddMenuEntry(
-		"HorizontalLayout",
-		LOCTEXT("FilterListHorizontalLayout", "Horizontal"),
-		LOCTEXT("FilterListHorizontalLayoutToolTip", "Swap to a Horizontal layout for the filter bar"),
-		FSlateIcon(),
-		FUIAction(
-		FExecuteAction::CreateLambda([this]()
-			{
-				if(this->FilterBarLayout != EFilterBarLayout::Horizontal)
-				{
-					this->SetFilterLayout(EFilterBarLayout::Horizontal);
-				}
-			}),
-			FCanExecuteAction(),
-			FIsActionChecked::CreateLambda([this]() { return FilterBarLayout == EFilterBarLayout::Horizontal; })),
-		EUserInterfaceActionType::RadioButton
-	);
-}
-
 void SFilterList::PopulateAddFilterMenu_Internal(UToolMenu* Menu)
 {
 	EAssetTypeCategories::Type MenuExpansion = EAssetTypeCategories::Basic;
@@ -360,8 +319,6 @@ void SFilterList::PopulateAddFilterMenu_Internal(UToolMenu* Menu)
 		return !IsFilteredByPicker(this->InitialClassFilters, TestClass);
 	}));
 	
-	PopulateFilterDisplayMenu(Menu);
-
 	Menu->AddSection("ContentBrowserFilterMiscAsset", LOCTEXT("MiscAssetsMenuHeading", "Misc Options") );
 }
 
@@ -760,10 +717,11 @@ void SFilterList::SetFilterLayout(EFilterBarLayout InFilterBarLayout)
 // FFilterListCustomTextFilter
 /////////////////////////////////////////
 
-
+/** Returns the system name for this filter */
 FString FFrontendFilter_CustomText::GetName() const
 {
-	return TEXT("FFrontendFilter_CustomText");
+	// Todo: Find some way to enforce this on all custom text filter interfaces
+	return FCustomTextFilter<FAssetFilterType>::GetFilterTypeName().ToString();
 }
 
 FText FFrontendFilter_CustomText::GetDisplayName() const
diff --git a/Engine/Source/Editor/ContentBrowser/Private/SFilterList.h b/Engine/Source/Editor/ContentBrowser/Private/SFilterList.h
--- a/Engine/Source/Editor/ContentBrowser/Private/SFilterList.h
+++ b/Engine/Source/Editor/ContentBrowser/Private/SFilterList.h
@@ -156,9 +156,6 @@ private:
 	/** Find the custom text filter corresponding to the specified state, and restore it's state to what is specified */
 	void RestoreCustomTextFilterState(const FCustomTextFilterState& InFilterState);
 	
-	/** Function to populate the options to change the filter layout */
-	void PopulateFilterDisplayMenu(UToolMenu* Menu);
-	
 private:
 
 	/** List of classes that our filters must match */
diff --git a/Engine/Source/Editor/ContentBrowser/Public/SAssetView.h b/Engine/Source/Editor/ContentBrowser/Public/SAssetView.h
--- a/Engine/Source/Editor/ContentBrowser/Public/SAssetView.h
+++ b/Engine/Source/Editor/ContentBrowser/Public/SAssetView.h
@@ -36,6 +36,7 @@ class UFactory;
 struct FPropertyChangedEvent;
 struct FToolMenuContext;
 class SContentBrowser;
+class SFilterList;
 
 /** Fires whenever the asset view is asked to start to create a temporary item */
 DECLARE_DELEGATE_OneParam(FOnAssetViewNewItemRequested, const FContentBrowserItem& /*NewItem*/);
@@ -379,6 +380,9 @@ public:
 	/** Gets text name for given thumbnail */
 	static FText ThumbnailSizeToDisplayName(EThumbnailSize InSize);
 
+	/** Set the filter list attached to this asset view - allows toggling of the the filter bar layout from the view options */
+	void SetFilterBar(TSharedPtr<SFilterList> InFilterBar);
+
 private:
 
 	/** Sets the pending selection to the current selection (used when changing views or refreshing the view). */
@@ -852,6 +856,9 @@ private:
 	/** The content browser that created this asset view if any */
 	TWeakPtr<SContentBrowser> OwningContentBrowser;
 
+	/** The Filter Bar attached to this asset view if any */
+	TWeakPtr<SFilterList> FilterBar;
+
 	/** The current base source filter for the view */
 	FSourcesData SourcesData;
 	FARFilter BackendFilter;
