SAssetView: Change the column view to have a similar wokflow to other column views in the editor (like the Outliner)

diff --git a/Engine/Source/Editor/ContentBrowser/Private/SAssetView.cpp b/Engine/Source/Editor/ContentBrowser/Private/SAssetView.cpp
--- a/Engine/Source/Editor/ContentBrowser/Private/SAssetView.cpp
+++ b/Engine/Source/Editor/ContentBrowser/Private/SAssetView.cpp
@@ -330,8 +330,6 @@ void SAssetView::Construct( const FArguments& InArgs )
 	bPendingFocusOnSync = false;
 	bWereItemsRecursivelyFiltered = false;
 
-	NumVisibleColumns = 0;
-
 	OwningContentBrowser = InArgs._OwningContentBrowser;
 
 	FAssetToolsModule& AssetToolsModule = FModuleManager::LoadModuleChecked<FAssetToolsModule>("AssetTools");
@@ -938,6 +936,16 @@ void SAssetView::LoadSettings(const FString& IniFilename, const FString& IniSect
 	if (LoadedHiddenColumnNames.Num() > 0)
 	{
 		HiddenColumnNames = LoadedHiddenColumnNames;
+
+		// Also update the visibility of the columns we just loaded in (unless this is called before creation and ColumnView doesn't exist)
+		if(ColumnView)
+		{
+			for (const SHeaderRow::FColumn& Column : ColumnView->GetHeaderRow()->GetColumns())
+			{
+				ColumnView->GetHeaderRow()->SetShowGeneratedColumn(Column.ColumnId, !HiddenColumnNames.Contains(Column.ColumnId.ToString()));
+			}
+		}
+		
 	}
 }
 
@@ -1692,24 +1700,19 @@ TSharedRef<SAssetColumnView> SAssetView::CreateColumnView()
 		(
 			SNew(SHeaderRow)
 			.ResizeMode(ESplitterResizeMode::FixedSize)
+			.CanSelectGeneratedColumn(true)
+			.OnHiddenColumnsListChanged(this, &SAssetView::OnHiddenColumnsChanged)
 			+ SHeaderRow::Column(SortManager.NameColumnId)
 			.FillWidth(300)
 			.SortMode( TAttribute< EColumnSortMode::Type >::Create( TAttribute< EColumnSortMode::Type >::FGetter::CreateSP( this, &SAssetView::GetColumnSortMode, SortManager.NameColumnId ) ) )
 			.SortPriority(TAttribute< EColumnSortPriority::Type >::Create(TAttribute< EColumnSortPriority::Type >::FGetter::CreateSP(this, &SAssetView::GetColumnSortPriority, SortManager.NameColumnId)))
 			.OnSort( FOnSortModeChanged::CreateSP( this, &SAssetView::OnSortColumnHeader ) )
 			.DefaultLabel( LOCTEXT("Column_Name", "Name") )
-			.ShouldGenerateWidget(TAttribute<bool>::Create(TAttribute<bool>::FGetter::CreateSP(this, &SAssetView::ShouldColumnGenerateWidget, SortManager.NameColumnId.ToString())))
-			.MenuContent()
-			[
-				CreateRowHeaderMenuContent(SortManager.NameColumnId.ToString())
-			]
+			.ShouldGenerateWidget(true) // Can't hide name column, so at least one column is visible
 		);
 
 	NewColumnView->GetHeaderRow()->SetOnGetMaxRowSizeForColumn(FOnGetMaxRowSizeForColumn::CreateRaw(NewColumnView.Get(), &SAssetColumnView::GetMaxRowSizeForColumn));
 
-
-	NumVisibleColumns = HiddenColumnNames.Contains(SortManager.NameColumnId.ToString()) ? 0 : 1;
-
 	if(bShowTypeInColumnView)
 	{
 		NewColumnView->GetHeaderRow()->AddColumn(
@@ -1719,17 +1722,13 @@ TSharedRef<SAssetColumnView> SAssetView::CreateColumnView()
 				.SortPriority(TAttribute< EColumnSortPriority::Type >::Create(TAttribute< EColumnSortPriority::Type >::FGetter::CreateSP(this, &SAssetView::GetColumnSortPriority, SortManager.ClassColumnId)))
 				.OnSort(FOnSortModeChanged::CreateSP(this, &SAssetView::OnSortColumnHeader))
 				.DefaultLabel(LOCTEXT("Column_Class", "Type"))
-				.ShouldGenerateWidget(TAttribute<bool>::Create(TAttribute<bool>::FGetter::CreateSP(this, &SAssetView::ShouldColumnGenerateWidget, SortManager.ClassColumnId.ToString())))
-				.MenuContent()
-				[
-					CreateRowHeaderMenuContent(SortManager.ClassColumnId.ToString())
-				]
 			);
 
-		NumVisibleColumns += HiddenColumnNames.Contains(SortManager.ClassColumnId.ToString()) ? 0 : 1;
+		bool bIsColumnVisible = !HiddenColumnNames.Contains(SortManager.ClassColumnId.ToString());
+		
+		NewColumnView->GetHeaderRow()->SetShowGeneratedColumn(SortManager.ClassColumnId, bIsColumnVisible);
 	}
 
-
 	if (bShowPathInColumnView)
 	{
 		NewColumnView->GetHeaderRow()->AddColumn(
@@ -1739,15 +1738,11 @@ TSharedRef<SAssetColumnView> SAssetView::CreateColumnView()
 				.SortPriority(TAttribute< EColumnSortPriority::Type >::Create(TAttribute< EColumnSortPriority::Type >::FGetter::CreateSP(this, &SAssetView::GetColumnSortPriority, SortManager.PathColumnId)))
 				.OnSort(FOnSortModeChanged::CreateSP(this, &SAssetView::OnSortColumnHeader))
 				.DefaultLabel(LOCTEXT("Column_Path", "Path"))
-				.ShouldGenerateWidget(TAttribute<bool>::Create(TAttribute<bool>::FGetter::CreateSP(this, &SAssetView::ShouldColumnGenerateWidget, SortManager.PathColumnId.ToString())))
-				.MenuContent()
-				[
-					CreateRowHeaderMenuContent(SortManager.PathColumnId.ToString())
-				]
 			);
 
-
-		NumVisibleColumns += HiddenColumnNames.Contains(SortManager.PathColumnId.ToString()) ? 0 : 1;
+		bool bIsColumnVisible = !HiddenColumnNames.Contains(SortManager.PathColumnId.ToString());
+		
+		NewColumnView->GetHeaderRow()->SetShowGeneratedColumn(SortManager.PathColumnId, bIsColumnVisible);
 	}
 
 	return NewColumnView.ToSharedRef();
@@ -2138,14 +2133,11 @@ void SAssetView::SetMajorityAssetType(FName NewMajorityAssetType)
 				.OnSort(FOnSortModeChanged::CreateSP(this, &SAssetView::OnSortColumnHeader))
 				.DefaultLabel(Column.DisplayName)
 				.DefaultTooltip(Column.TooltipText)
-				.FillWidth(180)
-				.ShouldGenerateWidget(TAttribute<bool>::Create(TAttribute<bool>::FGetter::CreateSP(this, &SAssetView::ShouldColumnGenerateWidget, TagName.ToString())))
-				.MenuContent()
-				[
-					CreateRowHeaderMenuContent(TagName.ToString())
-				]);
+				.FillWidth(180));
 
-			NumVisibleColumns += HiddenColumnNames.Contains(TagName.ToString()) ? 0 : 1;
+			bool bIsColumnVisible = !HiddenColumnNames.Contains(TagName.ToString());
+		
+			ColumnView->GetHeaderRow()->SetShowGeneratedColumn(TagName, bIsColumnVisible);
 
 			// If we found a tag the matches the column we are currently sorting on, there will be no need to change the column
 			for (int32 SortIdx = 0; SortIdx < CurrentSortOrder.Num(); SortIdx++)
@@ -2218,15 +2210,12 @@ void SAssetView::SetMajorityAssetType(FName NewMajorityAssetType)
 							.OnSort(FOnSortModeChanged::CreateSP(this, &SAssetView::OnSortColumnHeader))
 							.DefaultLabel(TagPair.Value.GetMetaData().DisplayName)
 							.DefaultTooltip(TagPair.Value.GetMetaData().TooltipText)
-							.FillWidth(180)
-							.ShouldGenerateWidget(TAttribute<bool>::Create(TAttribute<bool>::FGetter::CreateSP(this, &SAssetView::ShouldColumnGenerateWidget, TagPair.Key.ToString())))
-							.MenuContent()
-							[
-								CreateRowHeaderMenuContent(TagPair.Key.ToString())
-							]);
-
-						NumVisibleColumns += HiddenColumnNames.Contains(TagPair.Key.ToString()) ? 0 : 1;
+							.FillWidth(180));
 
+						bool bIsColumnVisible = !HiddenColumnNames.Contains(TagPair.Key.ToString());
+		
+						ColumnView->GetHeaderRow()->SetShowGeneratedColumn(TagPair.Key, bIsColumnVisible);
+						
 						// If we found a tag the matches the column we are currently sorting on, there will be no need to change the column
 						for (int32 SortIdx = 0; SortIdx < CurrentSortOrder.Num(); SortIdx++)
 						{
@@ -2762,15 +2751,6 @@ void SAssetView::PopulateViewButtonMenu(UToolMenu* Menu)
 	{
 		{
 			FToolMenuSection& Section = Menu->AddSection("AssetColumns", LOCTEXT("ToggleColumnsHeading", "Columns"));
-			Section.AddSubMenu(
-				"ToggleColumns",
-				LOCTEXT("ToggleColumnsMenu", "Toggle columns"),
-				LOCTEXT("ToggleColumnsMenuTooltip", "Show or hide specific columns."),
-				FNewMenuDelegate::CreateSP(this, &SAssetView::FillToggleColumnsMenu),
-				false,
-				FSlateIcon(),
-				false
-				);
 
 			Section.AddMenuEntry(
 				"ResetColumns",
@@ -4290,38 +4270,15 @@ bool SAssetView::PerformQuickJump(const bool bWasJumping)
 	return ValidMatch;
 }
 
-void SAssetView::FillToggleColumnsMenu(FMenuBuilder& MenuBuilder)
+void SAssetView::ResetColumns()
 {
-	// Column view may not be valid if we toggled off columns view while the columns menu was open
-	if(ColumnView.IsValid())
+	for (const SHeaderRow::FColumn &Column : ColumnView->GetHeaderRow()->GetColumns())
 	{
-		const TIndirectArray<SHeaderRow::FColumn> Columns = ColumnView->GetHeaderRow()->GetColumns();
-
-		for (int32 ColumnIndex = 0; ColumnIndex < Columns.Num(); ++ColumnIndex)
-		{
-			const FString ColumnName = Columns[ColumnIndex].ColumnId.ToString();
-
-			MenuBuilder.AddMenuEntry(
-				Columns[ColumnIndex].DefaultText,
-				LOCTEXT("ShowHideColumnTooltip", "Show or hide column"),
-				FSlateIcon(),
-				FUIAction(
-					FExecuteAction::CreateSP(this, &SAssetView::ToggleColumn, ColumnName),
-					FCanExecuteAction::CreateSP(this, &SAssetView::CanToggleColumn, ColumnName),
-					FIsActionChecked::CreateSP(this, &SAssetView::IsColumnVisible, ColumnName),
-					EUIActionRepeatMode::RepeatEnabled
-				),
-				NAME_None,
-				EUserInterfaceActionType::Check
-			);
-		}
+		ColumnView->GetHeaderRow()->SetShowGeneratedColumn(Column.ColumnId, !DefaultHiddenColumnNames.Contains(Column.ColumnId.ToString()));
 	}
-}
 
-void SAssetView::ResetColumns()
-{
+	// This is set after updating the column visibilties, because SetShowGeneratedColumn calls OnHiddenColumnsChanged indirectly which can mess up the list
 	HiddenColumnNames = DefaultHiddenColumnNames;
-	NumVisibleColumns = ColumnView->GetHeaderRow()->GetColumns().Num();
 	ColumnView->GetHeaderRow()->RefreshColumns();
 	ColumnView->RebuildList();
 }
@@ -4363,37 +4320,32 @@ void SAssetView::ExportColumns()
 	}
 }
 
-void SAssetView::ToggleColumn(const FString ColumnName)
+void SAssetView::OnHiddenColumnsChanged()
 {
-	SetColumnVisibility(ColumnName, HiddenColumnNames.Contains(ColumnName));
-}
-
-void SAssetView::SetColumnVisibility(const FString ColumnName, const bool bShow)
-{
-	if (!bShow)
-	{
-		--NumVisibleColumns;
-		HiddenColumnNames.Add(ColumnName);
-	}
-	else
+	// Early out if this is called before creation (due to loading config etc)
+	if(!ColumnView)
 	{
-		++NumVisibleColumns;
-		check(HiddenColumnNames.Contains(ColumnName));
-		HiddenColumnNames.Remove(ColumnName);
+		return;
 	}
 
-	ColumnView->GetHeaderRow()->RefreshColumns();
-	ColumnView->RebuildList();
-}
+	// We can't directly update the hidden columns list, because some columns maybe hidden, but not created yet
+	TArray<FName> NewHiddenColumns = ColumnView->GetHeaderRow()->GetHiddenColumnIds();
 
-bool SAssetView::CanToggleColumn(const FString ColumnName) const
-{
-	return (HiddenColumnNames.Contains(ColumnName) || NumVisibleColumns > 1);
-}
+	// So instead for each column that currently exists, we update its visibility state in the HiddenColumnNames array
+	for (const SHeaderRow::FColumn& Column : ColumnView->GetHeaderRow()->GetColumns())
+	{
+		bool bIsColumnVisible = NewHiddenColumns.Contains(Column.ColumnId);
 
-bool SAssetView::IsColumnVisible(const FString ColumnName) const
-{
-	return !HiddenColumnNames.Contains(ColumnName);
+		if(bIsColumnVisible)
+		{
+			HiddenColumnNames.AddUnique(Column.ColumnId.ToString());
+		}
+		else
+		{
+			HiddenColumnNames.Remove(Column.ColumnId.ToString());
+		}
+	}
+	
 }
 
 bool SAssetView::ShouldColumnGenerateWidget(const FString ColumnName) const
@@ -4401,21 +4353,6 @@ bool SAssetView::ShouldColumnGenerateWidget(const FString ColumnName) const
 	return !HiddenColumnNames.Contains(ColumnName);
 }
 
-TSharedRef<SWidget> SAssetView::CreateRowHeaderMenuContent(const FString ColumnName)
-{
-	FMenuBuilder MenuBuilder(true, NULL);
-
-	MenuBuilder.AddMenuEntry(
-		LOCTEXT("HideColumn", "Hide Column"),
-		LOCTEXT("HideColumnToolTip", "Hides this column."),
-		FSlateIcon(),
-		FUIAction(FExecuteAction::CreateSP(this, &SAssetView::SetColumnVisibility, ColumnName, false), FCanExecuteAction::CreateSP(this, &SAssetView::CanToggleColumn, ColumnName)),
-		NAME_None,
-		EUserInterfaceActionType::Button);
-
-	return MenuBuilder.MakeWidget();
-}
-
 void SAssetView::ForceShowPluginFolder(bool bEnginePlugin)
 {
 	if (bEnginePlugin && !IsShowingEngineContent())
diff --git a/Engine/Source/Editor/ContentBrowser/Public/SAssetView.h b/Engine/Source/Editor/ContentBrowser/Public/SAssetView.h
--- a/Engine/Source/Editor/ContentBrowser/Public/SAssetView.h
+++ b/Engine/Source/Editor/ContentBrowser/Public/SAssetView.h
@@ -785,28 +785,14 @@ private:
 	 */
 	bool PerformQuickJump(const bool bWasJumping);
 	
-	/** Generates the column filtering menu */
-	void FillToggleColumnsMenu(FMenuBuilder& MenuBuilder);
-
 	/** Resets the column filtering state to make them all visible */
 	void ResetColumns();
 
 	/** Export columns to CSV */
 	void ExportColumns();
 
-	/** Toggle the column at ColumnIndex */
-	void ToggleColumn(const FString ColumnName);
-	/** Sets the column visibility by removing/inserting the column*/
-	void SetColumnVisibility(const FString ColumnName, const bool bShow);
-
-	/** Whether or not a column can be toggled, has to be valid column and mandatory minimum number of columns = 1*/
-	bool CanToggleColumn(const FString ColumnName) const;
-
-	/** Whether or not a column is visible to show it's state in the filtering menu */
-	bool IsColumnVisible(const FString ColumnName) const;
-
-	/** Creates the row header context menu allowing for hiding individually clicked columns*/
-	TSharedRef<SWidget> CreateRowHeaderMenuContent(const FString ColumnName);
+	/** Called when a column is shown/hidden in the column view */
+	void OnHiddenColumnsChanged();
 
 	/** Will compute the max row size from all its children for the specified column id*/
 	FVector2D GetMaxRowSizeForColumn(const FName& ColumnId);
@@ -1137,7 +1123,6 @@ private:
 	/** Column filtering state */
 	TArray<FString> DefaultHiddenColumnNames;
 	TArray<FString> HiddenColumnNames;
-	int32 NumVisibleColumns;
 
 	TArray<FAssetViewCustomColumn> CustomColumns;
 public:
