Content Browser
- Fixed the Other Developers frontend filter not working when the All folder is being showned.
- Converted the virtual paths to internal paths before checking against the base developers folder.

diff --git a/Engine/Source/Editor/ContentBrowser/Private/FrontendFilters.cpp b/Engine/Source/Editor/ContentBrowser/Private/FrontendFilters.cpp
--- a/Engine/Source/Editor/ContentBrowser/Private/FrontendFilters.cpp
+++ b/Engine/Source/Editor/ContentBrowser/Private/FrontendFilters.cpp
@@ -27,6 +27,7 @@
 #include "TextFilterValueHandlers.h"
 #include "AssetCompilingManager.h"
 #include "ContentBrowserDataSource.h"
+#include "IContentBrowserDataModule.h"
 
 /** Helper functions for frontend filters */
 namespace FrontendFilterHelper
@@ -1064,8 +1065,10 @@ void FFrontendFilter_ShowOtherDevelopers::SetCurrentFilter(TArrayView<const FNam
 {
 	if ( InSourcePaths.Num() == 1 )
 	{
-		const FString PackagePath = InSourcePaths[0].ToString() + TEXT("/");
-		
+		FString PackagePath;
+		IContentBrowserDataModule::Get().GetSubsystem()->TryConvertVirtualPath(InSourcePaths[0].ToString(), PackagePath);
+		PackagePath += TEXT("/");
+
 		// If the path starts with the base developer path, and is not the path itself then only one developer path is selected
 		bIsOnlyOneDeveloperPathSelected = PackagePath.StartsWith(BaseDeveloperPath) && PackagePath.Len() != BaseDeveloperPath.Len();
 	}
@@ -1086,12 +1089,12 @@ bool FFrontendFilter_ShowOtherDevelopers::PassesFilter(FAssetFilterType InItem)
 		{
 			// TODO: Have attribute flags for this so you can tell from the item whether it's a developer folder, and also whether it's yours
 			// If selecting multiple folders, the Developers folder/parent folder, or "All Assets", hide assets which are found in the development folder unless they are in the current user's folder
-			bool bPackageInDeveloperFolder = !TextFilterUtils::NameStrincmp(InItem.GetVirtualPath(), BaseDeveloperPath, BaseDeveloperPathAnsi, BaseDeveloperPath.Len());
+			bool bPackageInDeveloperFolder = !TextFilterUtils::NameStrincmp(InItem.GetInternalPath(), BaseDeveloperPath, BaseDeveloperPathAnsi, BaseDeveloperPath.Len());
 			if ( bPackageInDeveloperFolder )
 			{
 				// Test again using only the path part to avoid filtering files directly in the Developers folder
 				// This happens after the above check to avoid string manipulation when not required
-				FString PackagePath = FPaths::GetPath(InItem.GetVirtualPath().ToString());
+				FString PackagePath = FPaths::GetPath(InItem.GetInternalPath().ToString());
 				bPackageInDeveloperFolder = PackagePath.StartsWith(BaseDeveloperPath);
 				if ( bPackageInDeveloperFolder )
 				{
