SFilterList: Temporarily disable saving/loading of filters using EditorConfig to fix a crash when docking the Content Drawer

#localization none
#tests Tested the repro on the JIRA to dock the Content Drawer - no crash. Tested adding filters to the Content Browser and closing and reopening it - no crash and filters are lost because saving is disabled.

diff --git a/Engine/Source/Editor/ContentBrowser/Private/SFilterList.cpp b/Engine/Source/Editor/ContentBrowser/Private/SFilterList.cpp
--- a/Engine/Source/Editor/ContentBrowser/Private/SFilterList.cpp
+++ b/Engine/Source/Editor/ContentBrowser/Private/SFilterList.cpp
@@ -532,6 +532,11 @@ void SFilterList::SaveSettings()
 		return;
 	}
 
+	// Temporarily disable saving of filters as a fail-safe (LoadSettings is also disabled)
+	// TODO: Remove after fixing the root cause in EditorConfig
+	UE_LOG(LogSlate, Display, TEXT("%s Filters were not saved, saving filters is currently force disabled."), *FilterBarIdentifier.ToString());
+	return;
+
 	// Get the settings unique to this instance and the common settings
 	FFilterBarSettings* InstanceSettings = &UFilterBarConfig::Get()->FilterBars.FindOrAdd(FilterBarIdentifier);
 	FFilterBarSettings* SharedSettings = &UFilterBarConfig::Get()->FilterBars.FindOrAdd(SharedIdentifier);
@@ -608,6 +613,11 @@ void SFilterList::LoadSettings(const FName& InInstanceName)
 		UE_LOG(LogSlate, Error, TEXT("SFilterList Requires that you specify a FilterBarIdentifier to load settings"));
 		return;
 	}
+		
+	// Temporarily disable loading of filters to fix a crash
+    // TODO: Remove after fixing the root cause in EditorConfig
+    UE_LOG(LogSlate, Display, TEXT("Couldn't load saved filters for %s, saving filters is currently force disabled."), *FilterBarIdentifier.ToString());
+    return;
 
 	// Get the settings unique to this instance and the common settings
 	const FFilterBarSettings* InstanceSettings = UFilterBarConfig::Get()->FilterBars.Find(InInstanceName);
