Content Browser: Add a minimum size for the vertical filters

diff --git a/Engine/Source/Editor/ContentBrowser/Private/SContentBrowser.cpp b/Engine/Source/Editor/ContentBrowser/Private/SContentBrowser.cpp
--- a/Engine/Source/Editor/ContentBrowser/Private/SContentBrowser.cpp
+++ b/Engine/Source/Editor/ContentBrowser/Private/SContentBrowser.cpp
@@ -763,6 +763,7 @@ void SContentBrowser::OnFilterBarLayoutChanging(EFilterBarLayout NewLayout)
 			/* Filters in an SScrollBox */
 			+ SSplitter::Slot()
 			.Value(0.12f)
+			.MinSize(95.f)
 			[
 				SNew(SVerticalBox)
 				// Don't take up space when there are no filters
