Typing a character while the asset picker menu is open will now focus the search box if it is not already focused, such as if you've dragged the scrollbar.

diff --git a/Engine/Source/Editor/ContentBrowser/Private/ContentBrowserSingleton.cpp b/Engine/Source/Editor/ContentBrowser/Private/ContentBrowserSingleton.cpp
--- a/Engine/Source/Editor/ContentBrowser/Private/ContentBrowserSingleton.cpp
+++ b/Engine/Source/Editor/ContentBrowser/Private/ContentBrowserSingleton.cpp
@@ -109,6 +109,12 @@ TSharedRef<SWidget> FContentBrowserSingleton::CreateAssetPicker(const FAssetPick
 		.AssetPickerConfig(AssetPickerConfig);
 }
 
+TSharedPtr<SWidget> FContentBrowserSingleton::GetAssetPickerSearchBox(const TSharedRef<SWidget>& AssetPickerWidget)
+{
+	TSharedRef<SAssetPicker> AssetPicker = StaticCastSharedRef<SAssetPicker>(AssetPickerWidget);
+	return AssetPicker->GetSearchBox();
+}
+
 TSharedRef<SWidget> FContentBrowserSingleton::CreatePathPicker(const FPathPickerConfig& PathPickerConfig)
 {
 	return SNew( SPathPicker )
diff --git a/Engine/Source/Editor/ContentBrowser/Private/ContentBrowserSingleton.h b/Engine/Source/Editor/ContentBrowser/Private/ContentBrowserSingleton.h
--- a/Engine/Source/Editor/ContentBrowser/Private/ContentBrowserSingleton.h
+++ b/Engine/Source/Editor/ContentBrowser/Private/ContentBrowserSingleton.h
@@ -13,6 +13,7 @@ class FSpawnTabArgs;
 class FTabManager;
 class FViewport;
 class SContentBrowser;
+class SWidget;
 class UFactory;
 class UToolMenu;
 class FWorkspaceItem;
@@ -54,11 +55,12 @@ public:
 	virtual ~FContentBrowserSingleton();
 
 	// IContentBrowserSingleton interface
-	virtual TSharedRef<class SWidget> CreateContentBrowser( const FName InstanceName, TSharedPtr<SDockTab> ContainingTab, const FContentBrowserConfig* ContentBrowserConfig ) override;
-	virtual TSharedRef<class SWidget> CreateAssetPicker(const FAssetPickerConfig& AssetPickerConfig) override;
-	virtual TSharedRef<class SWidget> CreatePathPicker(const FPathPickerConfig& PathPickerConfig) override;
-	virtual TSharedRef<class SWidget> CreateCollectionPicker(const FCollectionPickerConfig& CollectionPickerConfig) override;
-	virtual TSharedRef<class SWidget> CreateContentBrowserDrawer(const FContentBrowserConfig& ContentBrowserConfig, TFunction<TSharedPtr<SDockTab>()> InOnGetTabForDrawer) override;
+	virtual TSharedRef<SWidget> CreateContentBrowser( const FName InstanceName, TSharedPtr<SDockTab> ContainingTab, const FContentBrowserConfig* ContentBrowserConfig ) override;
+	virtual TSharedRef<SWidget> CreateAssetPicker(const FAssetPickerConfig& AssetPickerConfig) override;
+	virtual TSharedPtr<SWidget> GetAssetPickerSearchBox(const TSharedRef<SWidget>& AssetPickerWidget) override;
+	virtual TSharedRef<SWidget> CreatePathPicker(const FPathPickerConfig& PathPickerConfig) override;
+	virtual TSharedRef<SWidget> CreateCollectionPicker(const FCollectionPickerConfig& CollectionPickerConfig) override;
+	virtual TSharedRef<SWidget> CreateContentBrowserDrawer(const FContentBrowserConfig& ContentBrowserConfig, TFunction<TSharedPtr<SDockTab>()> InOnGetTabForDrawer) override;
 	virtual void CreateOpenAssetDialog(const FOpenAssetDialogConfig& OpenAssetConfig, const FOnAssetsChosenForOpen& OnAssetsChosenForOpen, const FOnAssetDialogCancelled& OnAssetDialogCancelled) override;
 	virtual TArray<FAssetData> CreateModalOpenAssetDialog(const FOpenAssetDialogConfig& InConfig) override;
 	virtual void CreateSaveAssetDialog(const FSaveAssetDialogConfig& SaveAssetConfig, const FOnObjectPathChosenForSave& OnAssetNameChosenForSave, const FOnAssetDialogCancelled& OnAssetDialogCancelled) override;
diff --git a/Engine/Source/Editor/ContentBrowser/Private/SAssetPicker.cpp b/Engine/Source/Editor/ContentBrowser/Private/SAssetPicker.cpp
--- a/Engine/Source/Editor/ContentBrowser/Private/SAssetPicker.cpp
+++ b/Engine/Source/Editor/ContentBrowser/Private/SAssetPicker.cpp
@@ -363,15 +363,19 @@ void SAssetPicker::Construct( const FArguments& InArgs )
 	AssetViewPtr->RequestSlowFullListRefresh();
 }
 
+TSharedPtr<SWidget> SAssetPicker::GetSearchBox() const
+{
+	return SearchBoxPtr;
+}
+
 EActiveTimerReturnType SAssetPicker::SetFocusPostConstruct( double InCurrentTime, float InDeltaTime )
 {
-	if ( SearchBoxPtr.IsValid() )
+	if (SearchBoxPtr.IsValid())
 	{
 		FWidgetPath WidgetToFocusPath;
-		FSlateApplication::Get().GeneratePathToWidgetUnchecked( SearchBoxPtr.ToSharedRef(), WidgetToFocusPath );
-		FSlateApplication::Get().SetKeyboardFocus( WidgetToFocusPath, EFocusCause::SetDirectly );
+		FSlateApplication::Get().GeneratePathToWidgetUnchecked(SearchBoxPtr.ToSharedRef(), WidgetToFocusPath);
+		FSlateApplication::Get().SetKeyboardFocus(WidgetToFocusPath, EFocusCause::SetDirectly);
 		WidgetToFocusPath.GetWindow()->SetWidgetToFocusOnActivate(SearchBoxPtr);
-
 		return EActiveTimerReturnType::Stop;
 	}
 
diff --git a/Engine/Source/Editor/ContentBrowser/Private/SAssetPicker.h b/Engine/Source/Editor/ContentBrowser/Private/SAssetPicker.h
--- a/Engine/Source/Editor/ContentBrowser/Private/SAssetPicker.h
+++ b/Engine/Source/Editor/ContentBrowser/Private/SAssetPicker.h
@@ -49,6 +49,8 @@ public:
 	/** Function that's called when someone wants to execute a Rename Command*/
 	void ExecuteRenameCommand();
 
+	TSharedPtr<SWidget> GetSearchBox() const;
+
 private:
 	/** Focuses the search box post-construct */
 	EActiveTimerReturnType SetFocusPostConstruct( double InCurrentTime, float InDeltaTime );
diff --git a/Engine/Source/Editor/ContentBrowser/Public/IContentBrowserSingleton.h b/Engine/Source/Editor/ContentBrowser/Public/IContentBrowserSingleton.h
--- a/Engine/Source/Editor/ContentBrowser/Public/IContentBrowserSingleton.h
+++ b/Engine/Source/Editor/ContentBrowser/Public/IContentBrowserSingleton.h
@@ -19,8 +19,9 @@
 #include "Misc/NamePermissionList.h"
 
 class FViewport;
-class UFactory;
 class IPlugin;
+class SWidget;
+class UFactory;
 
 typedef const FContentBrowserItem& FAssetFilterType;
 typedef TFilterCollection<FAssetFilterType> FAssetFilterCollectionType;
@@ -28,9 +29,6 @@ typedef TFilterCollection<FAssetFilterType> FAssetFilterCollectionType;
 typedef const TSharedRef<IPlugin>& FPluginFilterType;
 typedef TFilterCollection<FPluginFilterType> FPluginFilterCollectionType;
 
-class UFactory;
-
-
 /** The view modes used in SAssetView */
 namespace EAssetViewType
 {
@@ -484,7 +482,7 @@ public:
 	 *
 	 * @return The newly created content browser widget
 	 */
-	virtual TSharedRef<class SWidget> CreateContentBrowser( const FName InstanceName, TSharedPtr<SDockTab> ContainingTab, const FContentBrowserConfig* ContentBrowserConfig ) = 0;
+	virtual TSharedRef<SWidget> CreateContentBrowser( const FName InstanceName, TSharedPtr<SDockTab> ContainingTab, const FContentBrowserConfig* ContentBrowserConfig ) = 0;
 
 	/**
 	 * Generates an asset picker widget locked to the specified FARFilter.
@@ -492,7 +490,10 @@ public:
 	 * @param AssetPickerConfig		A struct containing details about how the asset picker should behave				
 	 * @return The asset picker widget
 	 */
-	virtual TSharedRef<class SWidget> CreateAssetPicker(const FAssetPickerConfig& AssetPickerConfig) = 0;
+	virtual TSharedRef<SWidget> CreateAssetPicker(const FAssetPickerConfig& AssetPickerConfig) = 0;
+
+	/** Focus the search box of the given asset picker widget. */
+	virtual TSharedPtr<SWidget> GetAssetPickerSearchBox(const TSharedRef<SWidget>& AssetPickerWidget) = 0;
 
 	/**
 	 * Generates a path picker widget.
@@ -500,7 +501,7 @@ public:
 	 * @param PathPickerConfig		A struct containing details about how the path picker should behave				
 	 * @return The path picker widget
 	 */
-	virtual TSharedRef<class SWidget> CreatePathPicker(const FPathPickerConfig& PathPickerConfig) = 0;
+	virtual TSharedRef<SWidget> CreatePathPicker(const FPathPickerConfig& PathPickerConfig) = 0;
 
 	/**
 	 * Generates a collection picker widget.
@@ -508,7 +509,7 @@ public:
 	 * @param CollectionPickerConfig		A struct containing details about how the collection picker should behave				
 	 * @return The collection picker widget
 	 */
-	virtual TSharedRef<class SWidget> CreateCollectionPicker(const FCollectionPickerConfig& CollectionPickerConfig) = 0;
+	virtual TSharedRef<SWidget> CreateCollectionPicker(const FCollectionPickerConfig& CollectionPickerConfig) = 0;
 
 	/**
 	 * Generates a content browser for use in a drawer. This content browser is a singleton and is reused among all drawers.
@@ -517,7 +518,7 @@ public:
 	 *
 	 * @return The content browser drawer widget
 	 */
-	virtual TSharedRef<class SWidget> CreateContentBrowserDrawer(const FContentBrowserConfig& ContentBrowserConfig, TFunction<TSharedPtr<SDockTab>()> InOnGetTabForDrawer) = 0;
+	virtual TSharedRef<SWidget> CreateContentBrowserDrawer(const FContentBrowserConfig& ContentBrowserConfig, TFunction<TSharedPtr<SDockTab>()> InOnGetTabForDrawer) = 0;
 
 	/**
 	 * Opens the Open Asset dialog in a non-modal window
diff --git a/Engine/Source/Editor/PropertyEditor/Private/UserInterface/PropertyEditor/SPropertyMenuAssetPicker.cpp b/Engine/Source/Editor/PropertyEditor/Private/UserInterface/PropertyEditor/SPropertyMenuAssetPicker.cpp
--- a/Engine/Source/Editor/PropertyEditor/Private/UserInterface/PropertyEditor/SPropertyMenuAssetPicker.cpp
+++ b/Engine/Source/Editor/PropertyEditor/Private/UserInterface/PropertyEditor/SPropertyMenuAssetPicker.cpp
@@ -36,8 +36,9 @@ void SPropertyMenuAssetPicker::Construct( const FArguments& InArgs )
 
 	const bool bInShouldCloseWindowAfterMenuSelection = true;
 	const bool bCloseSelfOnly = true;
+	const bool bSearchable = false;
 	
-	FMenuBuilder MenuBuilder(bInShouldCloseWindowAfterMenuSelection, nullptr, nullptr, bCloseSelfOnly);
+	FMenuBuilder MenuBuilder(bInShouldCloseWindowAfterMenuSelection, nullptr, nullptr, bCloseSelfOnly, &FCoreStyle::Get(), bSearchable);
 
 	if (NewAssetFactories.Num() > 0)
 	{
@@ -105,8 +106,6 @@ void SPropertyMenuAssetPicker::Construct( const FArguments& InArgs )
 
 	MenuBuilder.BeginSection(NAME_None, LOCTEXT("BrowseHeader", "Browse"));
 	{
-		TSharedPtr<SWidget> MenuContent;
-
 		FContentBrowserModule& ContentBrowserModule = FModuleManager::Get().LoadModuleChecked<FContentBrowserModule>(TEXT("ContentBrowser"));
 
 		FAssetPickerConfig AssetPickerConfig;
@@ -157,15 +156,17 @@ void SPropertyMenuAssetPicker::Construct( const FArguments& InArgs )
 		// Force show plugin content if meta data says so
 		AssetPickerConfig.bForceShowPluginContent = bForceShowPluginContent;
 
-		MenuContent =
+		AssetPickerWidget = ContentBrowserModule.Get().CreateAssetPicker(AssetPickerConfig);
+
+		TSharedRef<SWidget> MenuContent =
 			SNew(SBox)
 			.WidthOverride(PropertyEditorAssetConstants::ContentBrowserWindowSize.X)
 			.HeightOverride(PropertyEditorAssetConstants::ContentBrowserWindowSize.Y)
 			[
-				ContentBrowserModule.Get().CreateAssetPicker(AssetPickerConfig)
+				AssetPickerWidget.ToSharedRef()
 			];
 
-		MenuBuilder.AddWidget(MenuContent.ToSharedRef(), FText::GetEmpty(), true);
+		MenuBuilder.AddWidget(MenuContent, FText::GetEmpty(), true);
 	}
 	MenuBuilder.EndSection();
 
@@ -175,6 +176,39 @@ void SPropertyMenuAssetPicker::Construct( const FArguments& InArgs )
 	];
 }
 
+FReply SPropertyMenuAssetPicker::OnPreviewKeyDown(const FGeometry& MyGeometry, const FKeyEvent& InKeyEvent)
+{
+	if (!AssetPickerWidget.IsValid())
+	{
+		return FReply::Unhandled();
+	}
+	
+	// only give the search box focus if it's not a command like Ctrl+C
+	if (InKeyEvent.GetCharacter() == 0 || 
+		InKeyEvent.IsAltDown() ||
+		InKeyEvent.IsControlDown() ||
+		InKeyEvent.IsCommandDown())
+	{
+		return FReply::Unhandled();
+	}
+
+	const FWidgetPath* Path = InKeyEvent.GetEventPath();
+	if (Path != nullptr)
+	{
+		FContentBrowserModule& ContentBrowserModule = FModuleManager::Get().LoadModuleChecked<FContentBrowserModule>(TEXT("ContentBrowser"));
+		TSharedPtr<SWidget> SearchBox = ContentBrowserModule.Get().GetAssetPickerSearchBox(AssetPickerWidget.ToSharedRef());
+		if (SearchBox.IsValid())
+		{
+			if (!Path->ContainsWidget(SearchBox.Get()))
+			{
+				return FReply::Unhandled().SetUserFocus(SearchBox.ToSharedRef());
+			}
+		}
+	}
+
+	return FReply::Unhandled();
+}
+
 void SPropertyMenuAssetPicker::OnEdit()
 {
 	if( CurrentObject.IsValid() )
diff --git a/Engine/Source/Editor/PropertyEditor/Private/UserInterface/PropertyEditor/SPropertyMenuAssetPicker.h b/Engine/Source/Editor/PropertyEditor/Private/UserInterface/PropertyEditor/SPropertyMenuAssetPicker.h
--- a/Engine/Source/Editor/PropertyEditor/Private/UserInterface/PropertyEditor/SPropertyMenuAssetPicker.h
+++ b/Engine/Source/Editor/PropertyEditor/Private/UserInterface/PropertyEditor/SPropertyMenuAssetPicker.h
@@ -14,7 +14,7 @@ class SPropertyMenuAssetPicker : public SCompoundWidget
 {
 public:
 	SLATE_BEGIN_ARGS( SPropertyMenuAssetPicker )
-		: _InitialObject(NULL)
+		: _InitialObject(nullptr)
 		, _AllowClear(true)
 	{}
 		SLATE_ARGUMENT( FAssetData, InitialObject )
@@ -88,12 +88,16 @@ private:
 	 */
 	void SetValue( const FAssetData& AssetData );
 
+	virtual FReply OnPreviewKeyDown(const FGeometry& MyGeometry, const FKeyEvent& InKeyEvent) override;
+
 private:
 	FAssetData CurrentObject;
 	
 	/** The property this asset picker will modify (if any) */
 	TSharedPtr<IPropertyHandle> PropertyHandle;
 
+	TSharedPtr<SWidget> AssetPickerWidget;
+
 	/** Whether the asset can be 'None' in this case */
 	bool bAllowClear;
 
