Tiny Content Browser tiles will now display text for clarity

#p4v-preflight-copy 21734858

diff --git a/Engine/Source/Editor/ContentBrowser/Private/AssetViewWidgets.cpp b/Engine/Source/Editor/ContentBrowser/Private/AssetViewWidgets.cpp
--- a/Engine/Source/Editor/ContentBrowser/Private/AssetViewWidgets.cpp
+++ b/Engine/Source/Editor/ContentBrowser/Private/AssetViewWidgets.cpp
@@ -1873,7 +1873,6 @@ void SAssetTileItem::Construct( const FArguments& InArgs )
 							.Padding(2.0f,0.0f, 0.0f, 2.0f)
 							[
 								SAssignNew(ClassTextWidget, STextBlock)
-								.Visibility(this, &SAssetTileItem::GetAssetClassLabelVisibility)
 								.TextStyle(FAppStyle::Get(), "ContentBrowser.ClassFont")
 								.OverflowPolicy(ETextOverflowPolicy::Ellipsis)
 								.Text(this, &SAssetTileItem::GetAssetClassText)
diff --git a/Engine/Source/Editor/ContentBrowser/Public/AssetViewWidgets.h b/Engine/Source/Editor/ContentBrowser/Public/AssetViewWidgets.h
--- a/Engine/Source/Editor/ContentBrowser/Public/AssetViewWidgets.h
+++ b/Engine/Source/Editor/ContentBrowser/Public/AssetViewWidgets.h
@@ -601,6 +601,7 @@ protected:
 	FOptionalSize GetThumbnailBoxSize() const;
 
 	/** Gets the visibility of the asset class label in thumbnails */
+	UE_DEPRECATED(5.1, "GetAssetClassLabelVisibility is deprecated, as there is no longer a need to identify a visual difference in the label based on various states.")
 	EVisibility GetAssetClassLabelVisibility() const;
 
 	/** Gets the color of the asset class label in thumbnails */
