Fix for short class paths encountered during redirector fixup

diff --git a/Engine/Source/Editor/ContentBrowser/Private/CollectionAssetRegistryBridge.cpp b/Engine/Source/Editor/ContentBrowser/Private/CollectionAssetRegistryBridge.cpp
--- a/Engine/Source/Editor/ContentBrowser/Private/CollectionAssetRegistryBridge.cpp
+++ b/Engine/Source/Editor/ContentBrowser/Private/CollectionAssetRegistryBridge.cpp
@@ -40,7 +40,31 @@ public:
 	{
 		OutNewObjectPath = NAME_None;
 
-		if (InObjectPath.ToString().StartsWith(TEXT("/Script/")))
+		FString InObjectPathString = InObjectPath.ToString();
+		if (!InObjectPathString.StartsWith(TEXT("/")))
+		{
+			if (!GIsSavingPackage)
+			{
+				if (UClass* Class = UClass::TryFindTypeSlow<UClass>(InObjectPathString))
+				{
+					const FString ClassPathStr = Class->GetPathName();
+					// Use the linker to search for class name redirects (from the loaded ActiveClassRedirects)
+					const FString NewClassName = FLinkerLoad::FindNewPathNameForClass(ClassPathStr, false);
+
+					if (!NewClassName.IsEmpty())
+					{
+						check(FPackageName::IsValidObjectPath(NewClassName));
+						// Our new class name might be lacking the path, so try and find it so we can use the full path in the collection
+						UClass* FoundClass = FindObject<UClass>(nullptr, *NewClassName);
+						if (FoundClass)
+						{
+							OutNewObjectPath = *FoundClass->GetPathName();
+						}
+					}
+				}
+			}
+		}
+		else if (InObjectPathString.StartsWith(TEXT("/Script/"))) 
 		{
 			// We can't use FindObject while we're saving
 			if (!GIsSavingPackage)
