Use UClass::TryConvertShortTypeNameToPathName to avoid warning when handling known short path.

diff --git a/Engine/Source/Editor/ContentBrowser/Private/CollectionAssetRegistryBridge.cpp b/Engine/Source/Editor/ContentBrowser/Private/CollectionAssetRegistryBridge.cpp
--- a/Engine/Source/Editor/ContentBrowser/Private/CollectionAssetRegistryBridge.cpp
+++ b/Engine/Source/Editor/ContentBrowser/Private/CollectionAssetRegistryBridge.cpp
@@ -45,15 +45,13 @@ public:
 		{
 			if (!GIsSavingPackage)
 			{
-				if (UClass* Class = UClass::TryFindTypeSlow<UClass>(InObjectPathString))
+				FTopLevelAssetPath FullPath = UClass::TryConvertShortTypeNameToPathName(UClass::StaticClass(), InObjectPathString);
+				if (FullPath.IsValid())
 				{
-					const FString ClassPathStr = Class->GetPathName();
-					// Use the linker to search for class name redirects (from the loaded ActiveClassRedirects)
+					FString ClassPathStr = FullPath.ToString();
 					const FString NewClassName = FLinkerLoad::FindNewPathNameForClass(ClassPathStr, false);
-
 					if (!NewClassName.IsEmpty())
 					{
-						check(FPackageName::IsValidObjectPath(NewClassName));
 						// Our new class name might be lacking the path, so try and find it so we can use the full path in the collection
 						UClass* FoundClass = FindObject<UClass>(nullptr, *NewClassName);
 						if (FoundClass)
