Fix mismatched deprecation warning macros

diff --git a/Engine/Source/Editor/ContentBrowser/Private/ContentBrowserUtils.cpp b/Engine/Source/Editor/ContentBrowser/Private/ContentBrowserUtils.cpp
--- a/Engine/Source/Editor/ContentBrowser/Private/ContentBrowserUtils.cpp
+++ b/Engine/Source/Editor/ContentBrowser/Private/ContentBrowserUtils.cpp
@@ -584,7 +584,7 @@ void ContentBrowserUtils::AppendAssetFilterToContentBrowserFilter(const FARFilte
 PRAGMA_DISABLE_DEPRECATION_WARNINGS
 		// TODO: Modify this API to also use FSoftObjectPath with deprecation
 		ObjectFilter.ObjectNamesToInclude = UE::SoftObjectPath::Private::ConvertSoftObjectPaths(InAssetFilter.SoftObjectPaths);
-PRAGMA_DISABLE_DEPRECATION_WARNINGS
+PRAGMA_ENABLE_DEPRECATION_WARNINGS
 		ObjectFilter.TagsAndValuesToInclude = InAssetFilter.TagsAndValues;
 		ObjectFilter.bOnDiskObjectsOnly = InAssetFilter.bIncludeOnlyOnDiskAssets;
 	}
