Adding SWarningOrErrorBox styling to the error message for selecting an unusable path in the asset dialog.

diff --git a/Engine/Source/Editor/ContentBrowser/Private/SAssetDialog.cpp b/Engine/Source/Editor/ContentBrowser/Private/SAssetDialog.cpp
--- a/Engine/Source/Editor/ContentBrowser/Private/SAssetDialog.cpp
+++ b/Engine/Source/Editor/ContentBrowser/Private/SAssetDialog.cpp
@@ -51,6 +51,7 @@
 #include "SlateOptMacros.h"
 #include "SlotBase.h"
 #include "SourcesData.h"
+#include "SWarningOrErrorBox.h"
 #include "Styling/AppStyle.h"
 #include "Templates/Tuple.h"
 #include "Templates/UnrealTemplate.h"
@@ -261,19 +262,12 @@ void SAssetDialog::Construct(const FArguments& InArgs, const FSharedAssetDialogC
 		MainVerticalBox->AddSlot()
 		.AutoHeight()
 		[
-			// Constant height, whether the label is visible or not
-			SNew(SBox).HeightOverride(18)
-			[
-				SNew(SBorder)
-				.Visibility( this, &SAssetDialog::GetNameErrorLabelVisibility )
-				.BorderImage( FAppStyle::GetBrush("AssetDialog.ErrorLabelBorder") )
-				.Content()
-				[
-					SNew(STextBlock)
-					.Text( this, &SAssetDialog::GetNameErrorLabelText )
-					.ToolTipText(this, &SAssetDialog::GetNameErrorLabelText)
-				]
-			]
+			SNew(SWarningOrErrorBox)
+			.Padding(FMargin(8.0f, 4.0f, 4.0f, 4.0f))
+			.IconSize(FVector2D(16,16))
+			.MessageStyle(EMessageStyle::Error)
+			.Message(this, &SAssetDialog::GetNameErrorLabelText)
+			.Visibility(this, &SAssetDialog::GetNameErrorLabelVisibility)
 		];
 	}
 
@@ -968,7 +962,7 @@ void SAssetDialog::UpdateInputValidity()
 
 			if (!bIsMountedInternalPath)
 			{
-				LastInputValidityErrorText = LOCTEXT("AssetDialog_VirtualPathSelected", "You must select a non virtual path.");
+				LastInputValidityErrorText = LOCTEXT("AssetDialog_VirtualPathSelected", "The selected folder cannot be modified.");
 				bLastInputValidityCheckSuccessful = false;
 			}
 		}
