Content Browser: Add a minimum size to the Sources Panel headers to prevent overlap, and fix an issue causing drag on the splitters to fail in certain situations

diff --git a/Engine/Source/Editor/ContentBrowser/Private/SContentBrowser.cpp b/Engine/Source/Editor/ContentBrowser/Private/SContentBrowser.cpp
--- a/Engine/Source/Editor/ContentBrowser/Private/SContentBrowser.cpp
+++ b/Engine/Source/Editor/ContentBrowser/Private/SContentBrowser.cpp
@@ -478,23 +478,41 @@ void SContentBrowser::Construct( const FArguments& InArgs, const FName& InInstan
 							.Visibility( this, &SContentBrowser::GetSourcesViewVisibility )
 							+SSplitter::Slot()
 							.SizeRule(TAttribute<SSplitter::ESizeRule>(this, &SContentBrowser::GetFavoritesAreaSizeRule))
+							.MinSize(29.0f)
 							.Value(0.2f)
 							[
-								CreateFavoritesView(Config)
+								SNew(SBorder)
+								.BorderImage(FAppStyle::Get().GetBrush("Brushes.Header"))
+								.Padding(0.0f, 2.0f, 0.0f, 0.0f)
+								[
+									CreateFavoritesView(Config)
+								]
 							]
 								
 							+SSplitter::Slot()
 							.SizeRule(TAttribute<SSplitter::ESizeRule>(this, &SContentBrowser::GetPathAreaSizeRule))
+							.MinSize(29.0f)
 							.Value(0.8f)
 							[
-								CreatePathView(Config)
+								SNew(SBorder)
+								.BorderImage(FAppStyle::Get().GetBrush("Brushes.Header"))
+								.Padding(0.0f, 2.0f, 0.0f, 0.0f)
+								[
+									CreatePathView(Config)
+								]
 							]
 
 							+SSplitter::Slot()
 							.SizeRule(TAttribute<SSplitter::ESizeRule>(this, &SContentBrowser::GetCollectionsAreaSizeRule))
+							.MinSize(29.0f)
 							.Value(0.4f)
 							[
-								CreateDockedCollectionsView(Config)
+								SNew(SBorder)
+								.BorderImage(FAppStyle::Get().GetBrush("Brushes.Header"))
+								.Padding(0.0f, 2.0f, 0.0f, 0.0f)
+								[
+									CreateDockedCollectionsView(Config)
+								]
 							]
 						]
 
@@ -975,7 +993,7 @@ TSharedRef<SWidget> SContentBrowser::CreateFavoritesView(const FContentBrowserCo
 		SAssignNew(FavoritesArea, SExpandableArea)
 		.BorderImage(FAppStyle::Get().GetBrush("Brushes.Header"))
 		.BodyBorderImage(FAppStyle::Get().GetBrush("Brushes.Recessed"))
-		.HeaderPadding(FMargin(4.0f, 2.0f))
+		.HeaderPadding(FMargin(4.0f, 0.0f))
 		.Visibility(this, &SContentBrowser::GetFavoriteFolderVisibility)
 		.Padding(0)
 		.AllowAnimatedTransition(false)
@@ -1036,7 +1054,7 @@ TSharedRef<SWidget> SContentBrowser::CreatePathView(const FContentBrowserConfig*
 		SAssignNew(PathArea, SExpandableArea)
 		.BorderImage(FAppStyle::Get().GetBrush("Brushes.Header"))
 		.BodyBorderImage(FAppStyle::Get().GetBrush("Brushes.Recessed"))
-		.HeaderPadding(FMargin(4.0f, 2.0f))
+		.HeaderPadding(FMargin(4.0f, 0.0f))
 		.Padding(0)
 		.AllowAnimatedTransition(false)
 		.OnAreaExpansionChanged_Lambda([this](bool bIsExpanded) { if (!bIsExpanded) PathSearchToggleButton->SetExpanded(false); })
@@ -1101,7 +1119,7 @@ TSharedRef<SWidget> SContentBrowser::CreateDockedCollectionsView(const FContentB
 		SAssignNew(CollectionArea, SExpandableArea)
 		.BorderImage(FAppStyle::Get().GetBrush("Brushes.Header"))
 		.BodyBorderImage(FAppStyle::Get().GetBrush("Brushes.Recessed"))
-		.HeaderPadding(FMargin(4.0f, 2.0f))
+		.HeaderPadding(FMargin(4.0f, 0.0f))
 		.Padding(0)
 		.Visibility(this, &SContentBrowser::GetDockedCollectionsVisibility)
 		.AllowAnimatedTransition(false)
