Content Browser: Fix typo in collection view context menu

diff --git a/Engine/Source/Editor/ContentBrowser/Private/CollectionContextMenu.cpp b/Engine/Source/Editor/ContentBrowser/Private/CollectionContextMenu.cpp
--- a/Engine/Source/Editor/ContentBrowser/Private/CollectionContextMenu.cpp
+++ b/Engine/Source/Editor/ContentBrowser/Private/CollectionContextMenu.cpp
@@ -117,7 +117,7 @@ TSharedPtr<SWidget> FCollectionContextMenu::MakeCollectionTreeContextMenu(TShare
 			// New... (submenu)
 			MenuBuilder.AddSubMenu(
 				LOCTEXT("NewChildCollection", "New..."),
-				LOCTEXT("NewChildCollectionTooltip", "Create child a collection."),
+				LOCTEXT("NewChildCollectionTooltip", "Create a child collection."),
 				FNewMenuDelegate::CreateRaw( this, &FCollectionContextMenu::MakeNewCollectionSubMenu, ECollectionStorageMode::Static, SCollectionView::FCreateCollectionPayload( ParentCollection ) ),
 				FUIAction(
 					FExecuteAction(),
