Content Browser: Fix an issue where hiding favorites would not collapse the area, and just leave behind empty and unusable space

diff --git a/Engine/Source/Editor/ContentBrowser/Private/SContentBrowser.cpp b/Engine/Source/Editor/ContentBrowser/Private/SContentBrowser.cpp
--- a/Engine/Source/Editor/ContentBrowser/Private/SContentBrowser.cpp
+++ b/Engine/Source/Editor/ContentBrowser/Private/SContentBrowser.cpp
@@ -478,7 +478,7 @@ void SContentBrowser::Construct( const FArguments& InArgs, const FName& InInstan
 							.Visibility( this, &SContentBrowser::GetSourcesViewVisibility )
 							+SSplitter::Slot()
 							.SizeRule(TAttribute<SSplitter::ESizeRule>(this, &SContentBrowser::GetFavoritesAreaSizeRule))
-							.MinSize(29.0f)
+							.MinSize(TAttribute<float>(this, &SContentBrowser::GetFavoritesAreaMinSize))
 							.Value(0.2f)
 							[
 								SNew(SBorder)
@@ -504,7 +504,7 @@ void SContentBrowser::Construct( const FArguments& InArgs, const FName& InInstan
 
 							+SSplitter::Slot()
 							.SizeRule(TAttribute<SSplitter::ESizeRule>(this, &SContentBrowser::GetCollectionsAreaSizeRule))
-							.MinSize(29.0f)
+							.MinSize(TAttribute<float>(this, &SContentBrowser::GetCollectionsAreaMinSize))
 							.Value(0.4f)
 							[
 								SNew(SBorder)
@@ -1356,7 +1356,8 @@ void SContentBrowser::RegisterContentBrowserToolBar()
 
 SSplitter::ESizeRule SContentBrowser::GetFavoritesAreaSizeRule() const
 {
-	return FavoritesArea->IsExpanded() ? SSplitter::ESizeRule::FractionOfParent : SSplitter::ESizeRule::SizeToContent;
+	// Make sure the area is expanded and visible 
+	return FavoritesArea->IsExpanded() && GetFavoriteFolderVisibility() == EVisibility::Visible ? SSplitter::ESizeRule::FractionOfParent : SSplitter::ESizeRule::SizeToContent;
 }
 
 SSplitter::ESizeRule SContentBrowser::GetPathAreaSizeRule() const
@@ -1366,7 +1367,18 @@ SSplitter::ESizeRule SContentBrowser::GetPathAreaSizeRule() const
 
 SSplitter::ESizeRule SContentBrowser::GetCollectionsAreaSizeRule() const
 {
-	return CollectionArea->IsExpanded() ? SSplitter::ESizeRule::FractionOfParent : SSplitter::ESizeRule::SizeToContent;
+	// Make sure the area is expanded and visible 
+	return CollectionArea->IsExpanded() && GetDockedCollectionsVisibility() == EVisibility::Visible ? SSplitter::ESizeRule::FractionOfParent : SSplitter::ESizeRule::SizeToContent;
+}
+
+float SContentBrowser::GetFavoritesAreaMinSize() const
+{
+	return GetFavoriteFolderVisibility() == EVisibility::Visible ? 29.0f : 0.0f;
+}
+
+float SContentBrowser::GetCollectionsAreaMinSize() const
+{
+	return GetDockedCollectionsVisibility() == EVisibility::Visible ? 29.0f : 0.0f;
 }
 
 FText SContentBrowser::GetHighlightedText() const
diff --git a/Engine/Source/Editor/ContentBrowser/Private/SContentBrowser.h b/Engine/Source/Editor/ContentBrowser/Private/SContentBrowser.h
--- a/Engine/Source/Editor/ContentBrowser/Private/SContentBrowser.h
+++ b/Engine/Source/Editor/ContentBrowser/Private/SContentBrowser.h
@@ -581,6 +581,10 @@ private:
 	SSplitter::ESizeRule GetPathAreaSizeRule() const;
 	SSplitter::ESizeRule GetCollectionsAreaSizeRule() const;
 
+	/** Gets the min size for various areas. When areas are not visible the min size is 0, otherwise there is a minimum size to prevent overlap */
+	float GetFavoritesAreaMinSize() const;
+	float GetCollectionsAreaMinSize() const;
+
 	/** Called when the layout of the SFilterList is changing */
 	void OnFilterBarLayoutChanging(EFilterBarLayout NewLayout);
 
