Editor - Beginning the process to no longer make loading all assets you right click in the content browser.  The first step, is setting up the code to have a deprecation flag DEPRECATE_ASSET_TYPE_ACTIONS_NEEDING_LOADED_OBJECTS, this will make it illegal to implement GetActions on all AssetTypeAction implementations.  Checked in disabled, enable it locally to begin trying to implement the recommended method detailed in the IAssetTypeActions header.  WIP.  One of the things this change does is outlaw (regardless of deprecation) the HasActions() function, this function is no longer used at all, and so we may as well remove its implementation, it also works well as a canary in finding other samples and things that need to be upgraded.

diff --git a/Engine/Plugins/2D/Paper2D/Source/Paper2DEditor/Private/SpriteAssetTypeActions.h b/Engine/Plugins/2D/Paper2D/Source/Paper2DEditor/Private/SpriteAssetTypeActions.h
--- a/Engine/Plugins/2D/Paper2D/Source/Paper2DEditor/Private/SpriteAssetTypeActions.h
+++ b/Engine/Plugins/2D/Paper2D/Source/Paper2DEditor/Private/SpriteAssetTypeActions.h
@@ -17,7 +17,6 @@ public:
 	virtual FText GetName() const override;
 	virtual FColor GetTypeColor() const override;
 	virtual UClass* GetSupportedClass() const override;
-	virtual bool HasActions(const TArray<UObject*>& InObjects) const override { return true; }
 	virtual void GetActions(const TArray<UObject*>& InObjects, struct FToolMenuSection& Section) override;
 	virtual void OpenAssetEditor(const TArray<UObject*>& InObjects, TSharedPtr<class IToolkitHost> EditWithinLevelEditor = TSharedPtr<IToolkitHost>()) override;
 	virtual uint32 GetCategories() override;
diff --git a/Engine/Plugins/2D/Paper2D/Source/Paper2DEditor/Private/TileSetAssetTypeActions.h b/Engine/Plugins/2D/Paper2D/Source/Paper2DEditor/Private/TileSetAssetTypeActions.h
--- a/Engine/Plugins/2D/Paper2D/Source/Paper2DEditor/Private/TileSetAssetTypeActions.h
+++ b/Engine/Plugins/2D/Paper2D/Source/Paper2DEditor/Private/TileSetAssetTypeActions.h
@@ -19,7 +19,6 @@ public:
 	virtual UClass* GetSupportedClass() const override;
 	virtual void OpenAssetEditor(const TArray<UObject*>& InObjects, TSharedPtr<class IToolkitHost> EditWithinLevelEditor = TSharedPtr<IToolkitHost>()) override;
 	virtual uint32 GetCategories() override;
-	virtual bool HasActions(const TArray<UObject*>& InObjects) const override { return true; }
 	virtual void GetActions(const TArray<UObject*>& InObjects, struct FToolMenuSection& Section) override;
 	// End of IAssetTypeActions interface
 
diff --git a/Engine/Plugins/2D/Paper2D/Source/PaperSpriteSheetImporter/Private/PaperSpriteSheetAssetTypeActions.h b/Engine/Plugins/2D/Paper2D/Source/PaperSpriteSheetImporter/Private/PaperSpriteSheetAssetTypeActions.h
--- a/Engine/Plugins/2D/Paper2D/Source/PaperSpriteSheetImporter/Private/PaperSpriteSheetAssetTypeActions.h
+++ b/Engine/Plugins/2D/Paper2D/Source/PaperSpriteSheetImporter/Private/PaperSpriteSheetAssetTypeActions.h
@@ -14,7 +14,6 @@ public:
 	virtual FText GetName() const override;
 	virtual FColor GetTypeColor() const override;
 	virtual UClass* GetSupportedClass() const override;
-	virtual bool HasActions(const TArray<UObject*>& InObjects) const override { return true; }
 	virtual void GetActions(const TArray<UObject*>& InObjects, struct FToolMenuSection& Section) override;
 	virtual uint32 GetCategories() override;
 	virtual bool IsImportedAsset() const override;
diff --git a/Engine/Plugins/Animation/ControlRig/Source/ControlRigEditor/Private/ControlRigGizmoLibraryActions.h b/Engine/Plugins/Animation/ControlRig/Source/ControlRigEditor/Private/ControlRigGizmoLibraryActions.h
--- a/Engine/Plugins/Animation/ControlRig/Source/ControlRigEditor/Private/ControlRigGizmoLibraryActions.h
+++ b/Engine/Plugins/Animation/ControlRig/Source/ControlRigEditor/Private/ControlRigGizmoLibraryActions.h
@@ -12,10 +12,6 @@ class FControlRigShapeLibraryActions : public FAssetTypeActions_Base
 {
 public:
 	// IAssetTypeActions Implementation
-	virtual bool HasActions(const TArray<UObject*>& InObjects) const override
-	{
-		return true;
-	}
 	virtual void GetActions(const TArray<UObject*>& InObjects, FMenuBuilder& MenuBuilder) override;
 	virtual FText GetName() const override { return NSLOCTEXT("AssetTypeActions", "AssetTypeActions_ControlRigShapeLibrary", "Control Rig Shape Library"); }
 	virtual FColor GetTypeColor() const override { return FColor(100,100,255); }
diff --git a/Engine/Plugins/Animation/ControlRig/Source/ControlRigEditor/Private/Editor/SControlRigFunctionBulkEditWidget.cpp b/Engine/Plugins/Animation/ControlRig/Source/ControlRigEditor/Private/Editor/SControlRigFunctionBulkEditWidget.cpp
--- a/Engine/Plugins/Animation/ControlRig/Source/ControlRigEditor/Private/Editor/SControlRigFunctionBulkEditWidget.cpp
+++ b/Engine/Plugins/Animation/ControlRig/Source/ControlRigEditor/Private/Editor/SControlRigFunctionBulkEditWidget.cpp
@@ -102,7 +102,6 @@ TSharedRef<SWidget> SControlRigFunctionBulkEditWidget::MakeAssetViewForReference
 	AssetPickerConfig.bAllowNullSelection = false;
 	AssetPickerConfig.bShowBottomToolbar = false;
 	AssetPickerConfig.bAutohideSearchBar = true;
-	AssetPickerConfig.bPreloadAssetsForContextMenu = false;
 	AssetPickerConfig.Filter.ClassPaths.Add(FTopLevelAssetPath(TEXT("/Script/Engine"), TEXT("InvalidClass")));
 	AssetPickerConfig.Filter.bIncludeOnlyOnDiskAssets = false;
 	AssetPickerConfig.OnAssetsActivated = FOnAssetsActivated::CreateSP(this, &SControlRigFunctionBulkEditWidget::OnAssetsActivated);
diff --git a/Engine/Plugins/Animation/ControlRig/Source/ControlRigEditor/Private/Tools/AssetTypeActions_ControlRigPose.cpp b/Engine/Plugins/Animation/ControlRig/Source/ControlRigEditor/Private/Tools/AssetTypeActions_ControlRigPose.cpp
--- a/Engine/Plugins/Animation/ControlRig/Source/ControlRigEditor/Private/Tools/AssetTypeActions_ControlRigPose.cpp
+++ b/Engine/Plugins/Animation/ControlRig/Source/ControlRigEditor/Private/Tools/AssetTypeActions_ControlRigPose.cpp
@@ -29,11 +29,6 @@ UClass* FAssetTypeActions_ControlRigPose::GetSupportedClass() const
 	return UControlRigPoseAsset::StaticClass();
 }
 
-bool FAssetTypeActions_ControlRigPose::HasActions(const TArray<UObject*>& InObjects) const
-{
-	return true;
-}
-
 void FAssetTypeActions_ControlRigPose::OpenAssetEditor(const TArray<UObject*>& InObjects, TSharedPtr<class IToolkitHost> EditWithinLevelEditor /*= TSharedPtr<IToolkitHost>()*/)
 {
 	FAssetTypeActions_Base::OpenAssetEditor(InObjects, EditWithinLevelEditor);
@@ -44,22 +39,6 @@ uint32 FAssetTypeActions_ControlRigPose::GetCategories()
 	return EAssetTypeCategories::Animation;
 }
 
-class UThumbnailInfo* FAssetTypeActions_ControlRigPose::GetThumbnailInfo(UObject* Asset) const
-{
-	return nullptr; //mz todo it looks like this is not needed if we use the default. Will remove based on feedback.
-	/*
-	UControlRigPoseAsset* PoseAsset = CastChecked<UControlRigPoseAsset>(Asset);
-	UThumbnailInfo* ThumbnailInfo = PoseAsset->ThumbnailInfo;
-	if (ThumbnailInfo == NULL)
-	{
-		ThumbnailInfo = NewObject<USceneThumbnailInfo>(PoseAsset, NAME_None, RF_Transactional);
-		PoseAsset->ThumbnailInfo = ThumbnailInfo;
-	}
-	
-	return ThumbnailInfo;
-	*/
-}
-
 bool FAssetTypeActions_ControlRigPose::IsImportedAsset() const
 {
 	return false;
diff --git a/Engine/Plugins/Animation/ControlRig/Source/ControlRigEditor/Private/Tools/AssetTypeActions_ControlRigPose.h b/Engine/Plugins/Animation/ControlRig/Source/ControlRigEditor/Private/Tools/AssetTypeActions_ControlRigPose.h
--- a/Engine/Plugins/Animation/ControlRig/Source/ControlRigEditor/Private/Tools/AssetTypeActions_ControlRigPose.h
+++ b/Engine/Plugins/Animation/ControlRig/Source/ControlRigEditor/Private/Tools/AssetTypeActions_ControlRigPose.h
@@ -18,11 +18,9 @@ public:
 	virtual FText GetName() const override;
 	virtual FColor GetTypeColor() const override;
 	virtual UClass* GetSupportedClass() const override;
-	virtual bool HasActions(const TArray<UObject*>& InObjects) const override;
 	virtual void GetActions(const TArray<UObject*>& InObjects, struct FToolMenuSection& Section) override;
 	virtual void OpenAssetEditor(const TArray<UObject*>& InObjects, TSharedPtr<class IToolkitHost> EditWithinLevelEditor = TSharedPtr<IToolkitHost>()) override;
 	virtual uint32 GetCategories() override;
-	virtual class UThumbnailInfo* GetThumbnailInfo(UObject* Asset) const override;
 	virtual bool IsImportedAsset() const override;
 	// End IAssetTypeActions
 };
\ No newline at end of file
diff --git a/Engine/Plugins/Animation/IKRig/Source/IKRigEditor/Private/RetargetEditor/AssetTypeActions_IKRetargeter.cpp b/Engine/Plugins/Animation/IKRig/Source/IKRigEditor/Private/RetargetEditor/AssetTypeActions_IKRetargeter.cpp
--- a/Engine/Plugins/Animation/IKRig/Source/IKRigEditor/Private/RetargetEditor/AssetTypeActions_IKRetargeter.cpp
+++ b/Engine/Plugins/Animation/IKRig/Source/IKRigEditor/Private/RetargetEditor/AssetTypeActions_IKRetargeter.cpp
@@ -18,11 +18,6 @@ UClass* FAssetTypeActions_IKRetargeter::GetSupportedClass() const
 	return UIKRetargeter::StaticClass();
 }
 
-void FAssetTypeActions_IKRetargeter::GetActions(const TArray<UObject*>& InObjects, FToolMenuSection& Section)
-{
-	FAssetTypeActions_Base::GetActions(InObjects, Section);
-}
-
 void FAssetTypeActions_IKRetargeter::OpenAssetEditor(
 	const TArray<UObject*>& InObjects,
 	TSharedPtr<IToolkitHost> EditWithinLevelEditor)
@@ -61,10 +56,7 @@ void FAssetTypeActions_IKRetargeter::ExtendAnimSequenceToolMenu()
 	for (const FName& MenuName : MenusToExtend)
 	{
 		UToolMenu* Menu = UToolMenus::Get()->ExtendMenu(MenuName);
-		if (Menu == nullptr)
-		{
-			continue;
-		}
+		check(Menu);
 
 		FToolMenuSection& Section = Menu->FindOrAddSection("GetAssetActions");
 		Section.AddSubMenu(
diff --git a/Engine/Plugins/Animation/IKRig/Source/IKRigEditor/Private/RigEditor/AssetTypeActions_IKRigDefinition.cpp b/Engine/Plugins/Animation/IKRig/Source/IKRigEditor/Private/RigEditor/AssetTypeActions_IKRigDefinition.cpp
--- a/Engine/Plugins/Animation/IKRig/Source/IKRigEditor/Private/RigEditor/AssetTypeActions_IKRigDefinition.cpp
+++ b/Engine/Plugins/Animation/IKRig/Source/IKRigEditor/Private/RigEditor/AssetTypeActions_IKRigDefinition.cpp
@@ -19,12 +19,6 @@ UThumbnailInfo* FAssetTypeActions_IKRigDefinition::GetThumbnailInfo(UObject* Ass
 	return NewObject<USceneThumbnailInfo>(IKRig, NAME_None, RF_Transactional);
 }
 
-void FAssetTypeActions_IKRigDefinition::GetActions(const TArray<UObject*>& InObjects, FToolMenuSection& Section)
-{
-	FAssetTypeActions_Base::GetActions(InObjects, Section);
-}
-
-
 void FAssetTypeActions_IKRigDefinition::OpenAssetEditor(
 	const TArray<UObject*>& InObjects,
 	TSharedPtr<IToolkitHost> EditWithinLevelEditor)
diff --git a/Engine/Plugins/Animation/IKRig/Source/IKRigEditor/Public/RetargetEditor/AssetTypeActions_IKRetargeter.h b/Engine/Plugins/Animation/IKRig/Source/IKRigEditor/Public/RetargetEditor/AssetTypeActions_IKRetargeter.h
--- a/Engine/Plugins/Animation/IKRig/Source/IKRigEditor/Public/RetargetEditor/AssetTypeActions_IKRetargeter.h
+++ b/Engine/Plugins/Animation/IKRig/Source/IKRigEditor/Public/RetargetEditor/AssetTypeActions_IKRetargeter.h
@@ -14,12 +14,11 @@ class FAssetTypeActions_IKRetargeter : public FAssetTypeActions_Base
 public:
 	// IAssetTypeActions Implementation
 	virtual FText GetName() const override { return NSLOCTEXT("AssetTypeActions", "AssetTypeActions_IKRetargeter", "IK Retargeter"); }
-	virtual void GetActions(const TArray<UObject*>& InObjects, struct FToolMenuSection& Section) override;
 	virtual FColor GetTypeColor() const override { return FColor(0, 128, 255); }
 	virtual UClass* GetSupportedClass() const override;
 	virtual void OpenAssetEditor(const TArray<UObject*>& InObjects, TSharedPtr<class IToolkitHost> EditWithinLevelEditor = TSharedPtr<IToolkitHost>()) override;
 	virtual uint32 GetCategories() override { return EAssetTypeCategories::Animation; }
-	const TArray<FText>& GetSubMenus() const override;
+	virtual const TArray<FText>& GetSubMenus() const override;
 	virtual UThumbnailInfo* GetThumbnailInfo(UObject* Asset) const override;
 
 	static void ExtendAnimSequenceToolMenu();
diff --git a/Engine/Plugins/Animation/IKRig/Source/IKRigEditor/Public/RigEditor/AssetTypeActions_IKRigDefinition.h b/Engine/Plugins/Animation/IKRig/Source/IKRigEditor/Public/RigEditor/AssetTypeActions_IKRigDefinition.h
--- a/Engine/Plugins/Animation/IKRig/Source/IKRigEditor/Public/RigEditor/AssetTypeActions_IKRigDefinition.h
+++ b/Engine/Plugins/Animation/IKRig/Source/IKRigEditor/Public/RigEditor/AssetTypeActions_IKRigDefinition.h
@@ -11,7 +11,6 @@ class FAssetTypeActions_IKRigDefinition : public FAssetTypeActions_Base
 public:
 	// IAssetTypeActions Implementation
 	virtual FText GetName() const override { return NSLOCTEXT("AssetTypeActions", "AssetTypeActions_IKRigDefinition", "IK Rig"); }
-	virtual void GetActions(const TArray<UObject*>& InObjects, struct FToolMenuSection& Section) override;
 	virtual FColor GetTypeColor() const override { return FColor(255, 215, 0); }
 	virtual UClass* GetSupportedClass() const override;
 	virtual void OpenAssetEditor(const TArray<UObject*>& InObjects, TSharedPtr<class IToolkitHost> EditWithinLevelEditor = TSharedPtr<IToolkitHost>()) override;
diff --git a/Engine/Plugins/Editor/AssetManagerEditor/Source/AssetManagerEditor/Private/AssetManagerEditorModule.cpp b/Engine/Plugins/Editor/AssetManagerEditor/Source/AssetManagerEditor/Private/AssetManagerEditorModule.cpp
--- a/Engine/Plugins/Editor/AssetManagerEditor/Source/AssetManagerEditor/Private/AssetManagerEditorModule.cpp
+++ b/Engine/Plugins/Editor/AssetManagerEditor/Source/AssetManagerEditor/Private/AssetManagerEditorModule.cpp
@@ -935,7 +935,7 @@ void FAssetManagerEditorModule::ExtendContentBrowserAssetSelectionMenu()
 	FToolMenuEntry& Entry = Section.AddDynamicEntry("AssetManagerEditorViewCommands", FNewToolMenuSectionDelegate::CreateLambda([](FToolMenuSection& InSection)
 	{
 		UContentBrowserAssetContextMenuContext* Context = InSection.FindContext<UContentBrowserAssetContextMenuContext>();
-		if (Context && Context->bCanBeModified && Context->SelectedObjects.Num() > 0)
+		if (Context && Context->bCanBeModified && Context->SelectedAssets.Num() > 0)
 		{
 			FAssetManagerEditorModule::CreateAssetContextMenu(InSection);
 		}
diff --git a/Engine/Plugins/Editor/AssetManagerEditor/Source/AssetManagerEditor/Private/SAssetAuditBrowser.cpp b/Engine/Plugins/Editor/AssetManagerEditor/Source/AssetManagerEditor/Private/SAssetAuditBrowser.cpp
--- a/Engine/Plugins/Editor/AssetManagerEditor/Source/AssetManagerEditor/Private/SAssetAuditBrowser.cpp
+++ b/Engine/Plugins/Editor/AssetManagerEditor/Source/AssetManagerEditor/Private/SAssetAuditBrowser.cpp
@@ -289,7 +289,6 @@ void SAssetAuditBrowser::Construct(const FArguments& InArgs)
 	Config.GetCurrentSelectionDelegates.Add(&GetCurrentSelectionDelegate);
 	Config.SetFilterDelegates.Add(&SetFilterDelegate);
 	Config.bFocusSearchBoxWhenOpened = false;
-	Config.bPreloadAssetsForContextMenu = false;
 
 	Config.SaveSettingsName = SettingsIniSection;
 	
diff --git a/Engine/Plugins/Editor/ContentBrowser/ContentBrowserAssetDataSource/Source/ContentBrowserAssetDataSource/Private/AssetFileContextMenu.cpp b/Engine/Plugins/Editor/ContentBrowser/ContentBrowserAssetDataSource/Source/ContentBrowserAssetDataSource/Private/AssetFileContextMenu.cpp
--- a/Engine/Plugins/Editor/ContentBrowser/ContentBrowserAssetDataSource/Source/ContentBrowserAssetDataSource/Private/AssetFileContextMenu.cpp
+++ b/Engine/Plugins/Editor/ContentBrowser/ContentBrowserAssetDataSource/Source/ContentBrowserAssetDataSource/Private/AssetFileContextMenu.cpp
@@ -454,18 +454,15 @@ void FAssetFileContextMenu::MakeAssetActionsSubMenu(UToolMenu* Menu)
 		);
 
 		// Replace References
-		if (CanExecuteConsolidate())
-		{
-			Section.AddMenuEntry(
-				"ReplaceReferences",
-				LOCTEXT("ReplaceReferences", "Replace References"),
-				LOCTEXT("ConsolidateTooltip", "Replace references to the selected assets."),
-				FSlateIcon(),
-				FUIAction(
-				FExecuteAction::CreateSP(this, &FAssetFileContextMenu::ExecuteConsolidate)
-				)
-				);
-		}
+		Section.AddMenuEntry(
+			"ReplaceReferences",
+			LOCTEXT("ReplaceReferences", "Replace References"),
+			LOCTEXT("ConsolidateTooltip", "Replace references to the selected assets."),
+			FSlateIcon(),
+			FUIAction(
+			FExecuteAction::CreateSP(this, &FAssetFileContextMenu::ExecuteConsolidate)
+			)
+		);
 
 		// Property Matrix
 		bool bCanUsePropertyMatrix = FModuleManager::LoadModuleChecked<FPropertyEditorModule>( "PropertyEditor" ).GetCanUsePropertyMatrix();
@@ -1535,9 +1532,7 @@ void FAssetFileContextMenu::GetSelectedAssets(TArray<UObject*>& Assets, bool Ski
 			continue;
 		}
 
-		UObject* Object = SelectedAssets[AssetIdx].FastGetAsset();
-
-		if (Object)
+		if (UObject* Object = SelectedAssets[AssetIdx].GetAsset())
 		{
 			Assets.Add(Object);
 		}
@@ -2377,32 +2372,6 @@ bool FAssetFileContextMenu::CanExecuteSCCSync() const
 	return bCanExecuteSCCSync;
 }
 
-bool FAssetFileContextMenu::CanExecuteConsolidate() const
-{
-	TArray<UObject*> ProposedObjects;
-	for (int32 AssetIdx = 0; AssetIdx < SelectedAssets.Num(); ++AssetIdx)
-	{
-		// Don't load assets here. Only operate on already loaded assets.
-		if ( SelectedAssets[AssetIdx].IsAssetLoaded() )
-		{
-			UObject* Object = SelectedAssets[AssetIdx].GetAsset();
-
-			if ( Object )
-			{
-				ProposedObjects.Add(Object);
-			}
-		}
-	}
-	
-	if ( ProposedObjects.Num() > 0 )
-	{
-		TArray<UObject*> CompatibleObjects;
-		return FConsolidateToolWindow::DetermineAssetCompatibility(ProposedObjects, CompatibleObjects);
-	}
-
-	return false;
-}
-
 bool FAssetFileContextMenu::CanExecuteDiffSelected() const
 {
 	bool bCanDiffSelected = false;
diff --git a/Engine/Plugins/Editor/ContentBrowser/ContentBrowserAssetDataSource/Source/ContentBrowserAssetDataSource/Private/ContentBrowserAssetDataCore.cpp b/Engine/Plugins/Editor/ContentBrowser/ContentBrowserAssetDataSource/Source/ContentBrowserAssetDataSource/Private/ContentBrowserAssetDataCore.cpp
--- a/Engine/Plugins/Editor/ContentBrowser/ContentBrowserAssetDataSource/Source/ContentBrowserAssetDataSource/Private/ContentBrowserAssetDataCore.cpp
+++ b/Engine/Plugins/Editor/ContentBrowser/ContentBrowserAssetDataSource/Source/ContentBrowserAssetDataSource/Private/ContentBrowserAssetDataCore.cpp
@@ -19,6 +19,8 @@
 #include "ContentBrowserDataUtils.h"
 #include "Settings/ContentBrowserSettings.h"
 #include "Engine/Level.h"
+#include "Framework/Notifications/NotificationManager.h"
+#include "Widgets/Notifications/SNotificationList.h"
 
 #define LOCTEXT_NAMESPACE "ContentBrowserAssetDataSource"
 
@@ -1607,7 +1609,7 @@ void PopulateAssetFileContextMenu(UContentBrowserDataSource* InOwnerDataSource,
 		InMenu,
 		SelectedAssets,
 		OnShowAssetsInPathsView
-		);
+	);
 }
 
 }
diff --git a/Engine/Plugins/Editor/ContentBrowser/ContentBrowserAssetDataSource/Source/ContentBrowserAssetDataSource/Public/AssetFileContextMenu.h b/Engine/Plugins/Editor/ContentBrowser/ContentBrowserAssetDataSource/Source/ContentBrowserAssetDataSource/Public/AssetFileContextMenu.h
--- a/Engine/Plugins/Editor/ContentBrowser/ContentBrowserAssetDataSource/Source/ContentBrowserAssetDataSource/Public/AssetFileContextMenu.h
+++ b/Engine/Plugins/Editor/ContentBrowser/ContentBrowserAssetDataSource/Source/ContentBrowserAssetDataSource/Public/AssetFileContextMenu.h
@@ -265,9 +265,6 @@ private:
 	/** Handler to check to see if "Diff Against Depot" can be executed */
 	bool CanExecuteSCCDiffAgainstDepot() const;
 
-	/** Handler to check to see if "Consolidate" can be executed */
-	bool CanExecuteConsolidate() const;
-
 	/** Handler to check to see if "Diff Selected" can be executed */
 	bool CanExecuteDiffSelected() const;
 
diff --git a/Engine/Plugins/Editor/GameplayTagsEditor/Source/GameplayTagsEditor/Private/AssetTypeActions_GameplayTagAssetBase.cpp b/Engine/Plugins/Editor/GameplayTagsEditor/Source/GameplayTagsEditor/Private/AssetTypeActions_GameplayTagAssetBase.cpp
--- a/Engine/Plugins/Editor/GameplayTagsEditor/Source/GameplayTagsEditor/Private/AssetTypeActions_GameplayTagAssetBase.cpp
+++ b/Engine/Plugins/Editor/GameplayTagsEditor/Source/GameplayTagsEditor/Private/AssetTypeActions_GameplayTagAssetBase.cpp
@@ -15,11 +15,6 @@ FAssetTypeActions_GameplayTagAssetBase::FAssetTypeActions_GameplayTagAssetBase(F
 	: OwnedGameplayTagPropertyName(InTagPropertyName)
 {}
 
-bool FAssetTypeActions_GameplayTagAssetBase::HasActions(const TArray<UObject*>& InObjects) const
-{
-	return true;
-}
-
 void FAssetTypeActions_GameplayTagAssetBase::GetActions(const TArray<UObject*>& InObjects, FToolMenuSection& Section)
 {
 	TArray<UObject*> ContainerObjectOwners;
diff --git a/Engine/Plugins/Editor/GameplayTagsEditor/Source/GameplayTagsEditor/Public/AssetTypeActions_GameplayTagAssetBase.h b/Engine/Plugins/Editor/GameplayTagsEditor/Source/GameplayTagsEditor/Public/AssetTypeActions_GameplayTagAssetBase.h
--- a/Engine/Plugins/Editor/GameplayTagsEditor/Source/GameplayTagsEditor/Public/AssetTypeActions_GameplayTagAssetBase.h
+++ b/Engine/Plugins/Editor/GameplayTagsEditor/Source/GameplayTagsEditor/Public/AssetTypeActions_GameplayTagAssetBase.h
@@ -14,9 +14,6 @@ public:
 	/** Constructor */
 	FAssetTypeActions_GameplayTagAssetBase(FName InTagPropertyName);
 
-	/** Overridden to specify that the gameplay tag base has actions */
-	virtual bool HasActions(const TArray<UObject*>& InObjects) const override;
-
 	/** Overridden to offer the gameplay tagging options */
 	virtual void GetActions(const TArray<UObject*>& InObjects, struct FToolMenuSection& Section) override;
 
diff --git a/Engine/Plugins/EnhancedInput/Source/InputEditor/Private/InputEditorModule.cpp b/Engine/Plugins/EnhancedInput/Source/InputEditor/Private/InputEditorModule.cpp
--- a/Engine/Plugins/EnhancedInput/Source/InputEditor/Private/InputEditorModule.cpp
+++ b/Engine/Plugins/EnhancedInput/Source/InputEditor/Private/InputEditorModule.cpp
@@ -257,7 +257,6 @@ public:
 	virtual FColor GetTypeColor() const override { return FColor(127, 255, 255); }
 	virtual FText GetAssetDescription(const FAssetData& AssetData) const override { return NSLOCTEXT("AssetTypeActions", "AssetTypeActions_InputActionDesc", "Represents an abstract game action that can be mapped to arbitrary hardware input devices."); }
 	virtual UClass* GetSupportedClass() const override { return UInputAction::StaticClass(); }
-	virtual bool HasActions(const TArray<UObject*>& InObjects) const override { return true; }
 	virtual void GetActions(const TArray<UObject*>& InObjects, struct FToolMenuSection& Section) override;
 
 private:
diff --git a/Engine/Plugins/Enterprise/GLTFExporter/Source/GLTFExporter/Private/Actions/GLTFProxyAssetActions.cpp b/Engine/Plugins/Enterprise/GLTFExporter/Source/GLTFExporter/Private/Actions/GLTFProxyAssetActions.cpp
--- a/Engine/Plugins/Enterprise/GLTFExporter/Source/GLTFExporter/Private/Actions/GLTFProxyAssetActions.cpp
+++ b/Engine/Plugins/Enterprise/GLTFExporter/Source/GLTFExporter/Private/Actions/GLTFProxyAssetActions.cpp
@@ -16,11 +16,6 @@ FGLTFProxyAssetActions::FGLTFProxyAssetActions(const TSharedRef<IAssetTypeAction
 {
 }
 
-bool FGLTFProxyAssetActions::HasActions(const TArray<UObject*>& InObjects) const
-{
-	return true;
-}
-
 void FGLTFProxyAssetActions::GetActions(const TArray<UObject*>& InObjects, FToolMenuSection& Section)
 {
 	OriginalActions->GetActions(InObjects, Section);
diff --git a/Engine/Plugins/Enterprise/GLTFExporter/Source/GLTFExporter/Private/Actions/GLTFProxyAssetActions.h b/Engine/Plugins/Enterprise/GLTFExporter/Source/GLTFExporter/Private/Actions/GLTFProxyAssetActions.h
--- a/Engine/Plugins/Enterprise/GLTFExporter/Source/GLTFExporter/Private/Actions/GLTFProxyAssetActions.h
+++ b/Engine/Plugins/Enterprise/GLTFExporter/Source/GLTFExporter/Private/Actions/GLTFProxyAssetActions.h
@@ -11,8 +11,7 @@ class FGLTFProxyAssetActions final : public IAssetTypeActions
 public:
 
 	FGLTFProxyAssetActions(const TSharedRef<IAssetTypeActions>& OriginalActions);
-
-	virtual bool HasActions(const TArray<UObject*>& InObjects) const override;
+	
 	virtual void GetActions(const TArray<UObject*>& InObjects, FToolMenuSection& Section) override;
 
 	virtual FText GetName() const override { return OriginalActions->GetName(); }
diff --git a/Engine/Plugins/Enterprise/VariantManagerContent/Source/VariantManagerContentEditor/Private/LevelVariantSetsAssetActions.h b/Engine/Plugins/Enterprise/VariantManagerContent/Source/VariantManagerContentEditor/Private/LevelVariantSetsAssetActions.h
--- a/Engine/Plugins/Enterprise/VariantManagerContent/Source/VariantManagerContentEditor/Private/LevelVariantSetsAssetActions.h
+++ b/Engine/Plugins/Enterprise/VariantManagerContent/Source/VariantManagerContentEditor/Private/LevelVariantSetsAssetActions.h
@@ -16,7 +16,6 @@ public:
 	virtual uint32 GetCategories() override;
 	virtual FText GetName() const override;
 	virtual UClass* GetSupportedClass() const override;
-	virtual bool HasActions(const TArray<UObject*>& InObjects) const override { return true; }
 	virtual void GetActions(const TArray<UObject*>& InObjects, struct FToolMenuSection& Section) override;
 	virtual FColor GetTypeColor() const override;
 	virtual void OpenAssetEditor(const TArray<UObject*>& InObjects, TSharedPtr<class IToolkitHost> EditWithinLevelEditor = TSharedPtr<IToolkitHost>()) override;
diff --git a/Engine/Plugins/Experimental/ChaosSolverPlugin/Source/ChaosSolverEditor/Private/Chaos/AssetTypeActions_ChaosSolver.cpp b/Engine/Plugins/Experimental/ChaosSolverPlugin/Source/ChaosSolverEditor/Private/Chaos/AssetTypeActions_ChaosSolver.cpp
--- a/Engine/Plugins/Experimental/ChaosSolverPlugin/Source/ChaosSolverEditor/Private/Chaos/AssetTypeActions_ChaosSolver.cpp
+++ b/Engine/Plugins/Experimental/ChaosSolverPlugin/Source/ChaosSolverEditor/Private/Chaos/AssetTypeActions_ChaosSolver.cpp
@@ -19,11 +19,6 @@ UThumbnailInfo* FAssetTypeActions_ChaosSolver::GetThumbnailInfo(UObject* Asset)
 	return NewObject<USceneThumbnailInfo>(ChaosSolver, NAME_None, RF_Transactional);
 }
 
-void FAssetTypeActions_ChaosSolver::GetActions(const TArray<UObject*>& InObjects, FToolMenuSection& Section)
-{
-	FAssetTypeActions_Base::GetActions(InObjects, Section);
-}
-
 void FAssetTypeActions_ChaosSolver::OpenAssetEditor(const TArray<UObject*>& InObjects, TSharedPtr<IToolkitHost> EditWithinLevelEditor)
 {
 }
diff --git a/Engine/Plugins/Experimental/ChaosSolverPlugin/Source/ChaosSolverEditor/Public/Chaos/AssetTypeActions_ChaosSolver.h b/Engine/Plugins/Experimental/ChaosSolverPlugin/Source/ChaosSolverEditor/Public/Chaos/AssetTypeActions_ChaosSolver.h
--- a/Engine/Plugins/Experimental/ChaosSolverPlugin/Source/ChaosSolverEditor/Public/Chaos/AssetTypeActions_ChaosSolver.h
+++ b/Engine/Plugins/Experimental/ChaosSolverPlugin/Source/ChaosSolverEditor/Public/Chaos/AssetTypeActions_ChaosSolver.h
@@ -11,7 +11,6 @@ class FAssetTypeActions_ChaosSolver : public FAssetTypeActions_Base
 public:
 	// IAssetTypeActions Implementation
 	virtual FText GetName() const override { return NSLOCTEXT("AssetTypeActions", "AssetTypeActions_ChaosSolver", "Chaos Solver"); }
-	virtual void GetActions(const TArray<UObject*>& InObjects, struct FToolMenuSection& Section) override;
 	virtual FColor GetTypeColor() const override { return FColor(255, 192, 128); }
 	virtual UClass* GetSupportedClass() const override;
 	virtual void OpenAssetEditor(const TArray<UObject*>& InObjects, TSharedPtr<class IToolkitHost> EditWithinLevelEditor = TSharedPtr<IToolkitHost>()) override;
diff --git a/Engine/Plugins/Experimental/GeometryCollectionPlugin/Source/GeometryCollectionEditor/Private/GeometryCollection/AssetTypeActions_GeometryCollection.cpp b/Engine/Plugins/Experimental/GeometryCollectionPlugin/Source/GeometryCollectionEditor/Private/GeometryCollection/AssetTypeActions_GeometryCollection.cpp
--- a/Engine/Plugins/Experimental/GeometryCollectionPlugin/Source/GeometryCollectionEditor/Private/GeometryCollection/AssetTypeActions_GeometryCollection.cpp
+++ b/Engine/Plugins/Experimental/GeometryCollectionPlugin/Source/GeometryCollectionEditor/Private/GeometryCollection/AssetTypeActions_GeometryCollection.cpp
@@ -23,7 +23,7 @@ UThumbnailInfo* FAssetTypeActions_GeometryCollection::GetThumbnailInfo(UObject*
 {
 	UGeometryCollection * GeometryCollection = CastChecked<UGeometryCollection>(Asset);
 	UThumbnailInfo* ThumbnailInfo = GeometryCollection->ThumbnailInfo;
-	if (ThumbnailInfo == NULL)
+	if (ThumbnailInfo == nullptr)
 	{
 		ThumbnailInfo = NewObject<USceneThumbnailInfo>(GeometryCollection, NAME_None, RF_Transactional);
 		GeometryCollection->ThumbnailInfo = ThumbnailInfo;
@@ -32,14 +32,6 @@ UThumbnailInfo* FAssetTypeActions_GeometryCollection::GetThumbnailInfo(UObject*
 	return ThumbnailInfo;
 }
 
-void FAssetTypeActions_GeometryCollection::GetActions(const TArray<UObject*>& InObjects, FToolMenuSection& Section)
-{
-	FAssetTypeActions_Base::GetActions(InObjects, Section);
-	//Set("ClassIcon.GeometryCollection", new IMAGE_BRUSH("Icons/AssetIcons/GeometryCollectionComponentAtlasGroup_16x", Icon16x16));
-	//Set("ClassThumbnail.GeometryCollection", new IMAGE_BRUSH("Icons/AssetIcons/GeometryCollectionComponentAtlasGroup_64x", Icon64x64));
-	// IconPath = Plugin->GetBaseDir() / TEXT("Resources/Icon128.png");
-}
-
 void FAssetTypeActions_GeometryCollection::OpenAssetEditor(const TArray<UObject*>& InObjects, TSharedPtr<class IToolkitHost> EditWithinLevelEditor)
 {
 	if (bGeometryCollectionDataflowEditor)
diff --git a/Engine/Plugins/Experimental/GeometryCollectionPlugin/Source/GeometryCollectionEditor/Public/GeometryCollection/AssetTypeActions_GeometryCollection.h b/Engine/Plugins/Experimental/GeometryCollectionPlugin/Source/GeometryCollectionEditor/Public/GeometryCollection/AssetTypeActions_GeometryCollection.h
--- a/Engine/Plugins/Experimental/GeometryCollectionPlugin/Source/GeometryCollectionEditor/Public/GeometryCollection/AssetTypeActions_GeometryCollection.h
+++ b/Engine/Plugins/Experimental/GeometryCollectionPlugin/Source/GeometryCollectionEditor/Public/GeometryCollection/AssetTypeActions_GeometryCollection.h
@@ -11,7 +11,6 @@ class FAssetTypeActions_GeometryCollection : public FAssetTypeActions_Base
 public:
 	// IAssetTypeActions Implementation
 	virtual FText GetName() const override { return NSLOCTEXT("AssetTypeActions", "AssetTypeActions_GeometryCollection", "Geometry Collection"); }
-	virtual void GetActions(const TArray<UObject*>& InObjects, struct FToolMenuSection& Section) override;
 	virtual FColor GetTypeColor() const override { return FColor(255, 192, 128); }
 	virtual UClass* GetSupportedClass() const override;
 	virtual void OpenAssetEditor(const TArray<UObject*>& InObjects, TSharedPtr<class IToolkitHost> EditWithinLevelEditor = TSharedPtr<IToolkitHost>()) override;
diff --git a/Engine/Plugins/Experimental/ImagePlate/Source/ImagePlateEditor/Private/AssetTypeActions_ImagePlateFileSequence.h b/Engine/Plugins/Experimental/ImagePlate/Source/ImagePlateEditor/Private/AssetTypeActions_ImagePlateFileSequence.h
--- a/Engine/Plugins/Experimental/ImagePlate/Source/ImagePlateEditor/Private/AssetTypeActions_ImagePlateFileSequence.h
+++ b/Engine/Plugins/Experimental/ImagePlate/Source/ImagePlateEditor/Private/AssetTypeActions_ImagePlateFileSequence.h
@@ -13,5 +13,4 @@ public:
 	virtual FColor GetTypeColor() const override { return FColor(105, 0, 60); }
 	virtual UClass* GetSupportedClass() const override { return UImagePlateFileSequence::StaticClass(); }
 	virtual uint32 GetCategories() override { return EAssetTypeCategories::Media; }
-	virtual bool HasActions ( const TArray<UObject*>& InObjects ) const override { return false; }
 };
diff --git a/Engine/Plugins/Experimental/MeshModelingToolsetExp/Source/ModelingEditorUI/Private/ModelingWidgets/SToolInputAssetPicker.cpp b/Engine/Plugins/Experimental/MeshModelingToolsetExp/Source/ModelingEditorUI/Private/ModelingWidgets/SToolInputAssetPicker.cpp
--- a/Engine/Plugins/Experimental/MeshModelingToolsetExp/Source/ModelingEditorUI/Private/ModelingWidgets/SToolInputAssetPicker.cpp
+++ b/Engine/Plugins/Experimental/MeshModelingToolsetExp/Source/ModelingEditorUI/Private/ModelingWidgets/SToolInputAssetPicker.cpp
@@ -195,7 +195,6 @@ void SToolInputAssetPicker::Construct( const FArguments& InArgs )
 		.CanShowDevelopersFolder( InArgs._AssetPickerConfig.bCanShowDevelopersFolder )
 		.ForceShowEngineContent( InArgs._AssetPickerConfig.bForceShowEngineContent )
 		.ForceShowPluginContent( InArgs._AssetPickerConfig.bForceShowPluginContent )
-		.PreloadAssetsForContextMenu( InArgs._AssetPickerConfig.bPreloadAssetsForContextMenu )
 		.HighlightedText( HighlightText )
 		.ThumbnailLabel( ThumbnailLabel )
 		.AssetShowWarningText( InArgs._AssetPickerConfig.AssetShowWarningText)
diff --git a/Engine/Plugins/Experimental/MotoSynth/Source/MotoSynthEditor/Private/SoundWaveAssetActionExtenderMotoSynth.cpp b/Engine/Plugins/Experimental/MotoSynth/Source/MotoSynthEditor/Private/SoundWaveAssetActionExtenderMotoSynth.cpp
--- a/Engine/Plugins/Experimental/MotoSynth/Source/MotoSynthEditor/Private/SoundWaveAssetActionExtenderMotoSynth.cpp
+++ b/Engine/Plugins/Experimental/MotoSynth/Source/MotoSynthEditor/Private/SoundWaveAssetActionExtenderMotoSynth.cpp
@@ -29,14 +29,14 @@ void FMotoSynthExtension::RegisterMenus()
 	Section.AddDynamicEntry("SoundWaveAssetConversion", FNewToolMenuSectionDelegate::CreateLambda([](FToolMenuSection& InSection)
 	{
 		UContentBrowserAssetContextMenuContext* Context = InSection.FindContext<UContentBrowserAssetContextMenuContext>();
-		if (!Context || Context->SelectedObjects.IsEmpty())
+		if (!Context || Context->SelectedAssets.IsEmpty())
 		{
 			return;
 		}
 
-		for (const TWeakObjectPtr<UObject>& Object : Context->SelectedObjects)
+		for (const FAssetData& Object : Context->SelectedAssets)
 		{
-			if (Object->IsA<USoundWaveProcedural>())
+			if (Object.IsInstanceOf(USoundWaveProcedural::StaticClass()))
 			{
 				return;
 			}
@@ -54,7 +54,7 @@ void FMotoSynthExtension::RegisterMenus()
 void FMotoSynthExtension::ExecuteCreateMotoSynthSource(const FToolMenuContext& MenuContext)
 {
 	UContentBrowserAssetContextMenuContext* Context = MenuContext.FindContext<UContentBrowserAssetContextMenuContext>();
-	if (!Context || Context->SelectedObjects.Num() == 0)
+	if (!Context || Context->SelectedAssets.Num() == 0)
 	{
 		return;
 	}
@@ -65,21 +65,28 @@ void FMotoSynthExtension::ExecuteCreateMotoSynthSource(const FToolMenuContext& M
 	// Create the factory used to generate the asset
 	UMotoSynthSourceFactory* Factory = NewObject<UMotoSynthSourceFactory>();
 	
-	for (const TWeakObjectPtr<UObject>& Object : Context->SelectedObjects)
+	FMessageLog EditorErrors("EditorErrors");
+	for (const FAssetData& Asset : Context->SelectedAssets)
 	{
 		// stage the soundwave on the factory to be used during asset creation
-		USoundWave* Wave = Cast<USoundWave>(Object);
-		check(Wave);
-		Factory->StagedSoundWave = Wave; // WeakPtr gets reset by the Factory after it is consumed
+		if (USoundWave* Wave = Cast<USoundWave>(Asset.GetAsset()))
+		{
+			Factory->StagedSoundWave = Wave; // WeakPtr gets reset by the Factory after it is consumed
 
-		// Determine an appropriate name
-		FString Name;
-		FString PackagePath;
-		AssetToolsModule.Get().CreateUniqueAssetName(Wave->GetOutermost()->GetName(), DefaultSuffix, PackagePath, Name);
+			// Determine an appropriate name
+			FString Name;
+			FString PackagePath;
+			AssetToolsModule.Get().CreateUniqueAssetName(Wave->GetOutermost()->GetName(), DefaultSuffix, PackagePath, Name);
 
-		// create new asset
-		AssetToolsModule.Get().CreateAsset(Name, FPackageName::GetLongPackagePath(PackagePath), UMotoSynthSource::StaticClass(), Factory);
+			// create new asset
+			AssetToolsModule.Get().CreateAsset(Name, FPackageName::GetLongPackagePath(PackagePath), UMotoSynthSource::StaticClass(), Factory);
+		}
+		else
+		{
+			EditorErrors.Error(LOCTEXT("ExpectedUSoundWave", "Expected SoundWave"))->AddToken(FAssetNameToken::Create(Asset.PackageName.ToString()));
+		}
 	}
+	EditorErrors.Notify();
 }
 
 #undef LOCTEXT_NAMESPACE
\ No newline at end of file
diff --git a/Engine/Plugins/Experimental/Mutable/Source/CustomizableObjectEditor/Private/MuCOE/AssetTypeActions_CustomizableObject.h b/Engine/Plugins/Experimental/Mutable/Source/CustomizableObjectEditor/Private/MuCOE/AssetTypeActions_CustomizableObject.h
--- a/Engine/Plugins/Experimental/Mutable/Source/CustomizableObjectEditor/Private/MuCOE/AssetTypeActions_CustomizableObject.h
+++ b/Engine/Plugins/Experimental/Mutable/Source/CustomizableObjectEditor/Private/MuCOE/AssetTypeActions_CustomizableObject.h
@@ -33,7 +33,6 @@ public:
 	FText GetName() const override { return NSLOCTEXT("AssetTypeActions", "AssetTypeActions_CustomizableObject", "Customizable Object"); }
 	FColor GetTypeColor() const override { return FColor(234, 255, 0); }
 	UClass* GetSupportedClass() const override { return UCustomizableObject::StaticClass(); }
-	bool HasActions ( const TArray<UObject*>& InObjects ) const override { return true; }
 	void GetActions( const TArray<UObject*>& InObjects, FMenuBuilder& MenuBuilder ) override;
 	void OpenAssetEditor( const TArray<UObject*>& InObjects, TSharedPtr<class IToolkitHost> EditWithinLevelEditor = TSharedPtr<IToolkitHost>() ) override;
 	uint32 GetCategories() override;
diff --git a/Engine/Plugins/Experimental/Mutable/Source/CustomizableObjectEditor/Private/MuCOE/AssetTypeActions_CustomizableObjectInstance.h b/Engine/Plugins/Experimental/Mutable/Source/CustomizableObjectEditor/Private/MuCOE/AssetTypeActions_CustomizableObjectInstance.h
--- a/Engine/Plugins/Experimental/Mutable/Source/CustomizableObjectEditor/Private/MuCOE/AssetTypeActions_CustomizableObjectInstance.h
+++ b/Engine/Plugins/Experimental/Mutable/Source/CustomizableObjectEditor/Private/MuCOE/AssetTypeActions_CustomizableObjectInstance.h
@@ -26,12 +26,10 @@ public:
 	FText GetName() const override { return NSLOCTEXT("AssetTypeActions", "AssetTypeActions_CustomizableObjectInstance", "Customizable Object Instance"); }
 	FColor GetTypeColor() const override { return FColor(255, 82, 49); }
 	UClass* GetSupportedClass() const override { return UCustomizableObjectInstance::StaticClass(); }
-	bool HasActions ( const TArray<UObject*>& InObjects ) const override { return true; }
 	void GetActions( const TArray<UObject*>& InObjects, FMenuBuilder& MenuBuilder ) override;
 	void OpenAssetEditor( const TArray<UObject*>& InObjects, TSharedPtr<class IToolkitHost> EditWithinLevelEditor = TSharedPtr<IToolkitHost>() ) override;
 	uint32 GetCategories() override;
-
-
+	
 	bool AssetsActivatedOverride(const TArray<UObject*>& InObjects, EAssetTypeActivationMethod::Type ActivationType) override;
 	bool CanFilter() override { return true; }
 	bool ShouldForceWorldCentric() override { return false; }
diff --git a/Engine/Plugins/Experimental/Mutable/Source/CustomizableObjectPopulationEditor/Private/MuCOPE/AssetTypeActions_CustomizableObjectPopulation.h b/Engine/Plugins/Experimental/Mutable/Source/CustomizableObjectPopulationEditor/Private/MuCOPE/AssetTypeActions_CustomizableObjectPopulation.h
--- a/Engine/Plugins/Experimental/Mutable/Source/CustomizableObjectPopulationEditor/Private/MuCOPE/AssetTypeActions_CustomizableObjectPopulation.h
+++ b/Engine/Plugins/Experimental/Mutable/Source/CustomizableObjectPopulationEditor/Private/MuCOPE/AssetTypeActions_CustomizableObjectPopulation.h
@@ -25,7 +25,6 @@ class FAssetTypeActions_CustomizableObjectPopulation : public FAssetTypeActions_
 	UClass* GetSupportedClass() const override;
 	uint32 GetCategories();
 
-	bool HasActions(const TArray<UObject*>& InObjects) const override { return true; }
 	void GetActions(const TArray<UObject*>& InObjects, FMenuBuilder& MenuBuilder) override;
 
 	void OpenAssetEditor(const TArray<UObject*>& InObjects, TSharedPtr<IToolkitHost> EditWithinLevelEditor) override;
diff --git a/Engine/Plugins/Experimental/NNI/Source/NeuralNetworkInferenceEditor/Private/NeuralNetworkAssetTypeActions.cpp b/Engine/Plugins/Experimental/NNI/Source/NeuralNetworkInferenceEditor/Private/NeuralNetworkAssetTypeActions.cpp
--- a/Engine/Plugins/Experimental/NNI/Source/NeuralNetworkInferenceEditor/Private/NeuralNetworkAssetTypeActions.cpp
+++ b/Engine/Plugins/Experimental/NNI/Source/NeuralNetworkInferenceEditor/Private/NeuralNetworkAssetTypeActions.cpp
@@ -21,12 +21,6 @@ UClass* FNeuralNetworkAssetTypeActions::GetSupportedClass() const
 	return UNeuralNetwork::StaticClass();
 }
 
-bool FNeuralNetworkAssetTypeActions::HasActions(const TArray<UObject*>& InObjects) const
-{
-	// Returns true if this class can supply actions for InObjects
-	return false;
-}
-
 uint32 FNeuralNetworkAssetTypeActions::GetCategories()
 {
 	const INeuralNetworkInferenceEditorModule& NeuralNetworkInferenceEditorModule = FModuleManager::GetModuleChecked<INeuralNetworkInferenceEditorModule>("NeuralNetworkInferenceEditor");
diff --git a/Engine/Plugins/Experimental/NNI/Source/NeuralNetworkInferenceEditor/Private/NeuralNetworkAssetTypeActions.h b/Engine/Plugins/Experimental/NNI/Source/NeuralNetworkInferenceEditor/Private/NeuralNetworkAssetTypeActions.h
--- a/Engine/Plugins/Experimental/NNI/Source/NeuralNetworkInferenceEditor/Private/NeuralNetworkAssetTypeActions.h
+++ b/Engine/Plugins/Experimental/NNI/Source/NeuralNetworkInferenceEditor/Private/NeuralNetworkAssetTypeActions.h
@@ -12,7 +12,6 @@ public:
 	virtual FText GetName() const override;
 	virtual FColor GetTypeColor() const override;
 	virtual UClass* GetSupportedClass() const override;
-	virtual bool HasActions(const TArray<UObject*>& InObjects) const override;
 	virtual uint32 GetCategories() override;
 	virtual bool IsImportedAsset() const override;
 	virtual void GetResolvedSourceFilePaths(const TArray<UObject*>& TypeAssets, TArray<FString>& OutSourceFilePaths) const override;
diff --git a/Engine/Plugins/Experimental/NNI/Source/NeuralNetworkInferenceEditor/Private/QA/NeuralNetworkInferenceQAAssetTypeActions.cpp b/Engine/Plugins/Experimental/NNI/Source/NeuralNetworkInferenceEditor/Private/QA/NeuralNetworkInferenceQAAssetTypeActions.cpp
--- a/Engine/Plugins/Experimental/NNI/Source/NeuralNetworkInferenceEditor/Private/QA/NeuralNetworkInferenceQAAssetTypeActions.cpp
+++ b/Engine/Plugins/Experimental/NNI/Source/NeuralNetworkInferenceEditor/Private/QA/NeuralNetworkInferenceQAAssetTypeActions.cpp
@@ -20,12 +20,6 @@ UClass* FNeuralNetworkInferenceQAAssetTypeActions::GetSupportedClass() const
 	return UNeuralNetworkInferenceQAAsset::StaticClass();
 }
 
-bool FNeuralNetworkInferenceQAAssetTypeActions::HasActions(const TArray<UObject*>& InObjects) const
-{
-	// Returns true if this class can supply actions for InObjects
-	return false;
-}
-
 uint32 FNeuralNetworkInferenceQAAssetTypeActions::GetCategories()
 {
 	const INeuralNetworkInferenceEditorModule& NeuralNetworkInferenceEditorModule = FModuleManager::GetModuleChecked<INeuralNetworkInferenceEditorModule>("NeuralNetworkInferenceEditor");
diff --git a/Engine/Plugins/Experimental/NNI/Source/NeuralNetworkInferenceEditor/Private/QA/NeuralNetworkInferenceQAAssetTypeActions.h b/Engine/Plugins/Experimental/NNI/Source/NeuralNetworkInferenceEditor/Private/QA/NeuralNetworkInferenceQAAssetTypeActions.h
--- a/Engine/Plugins/Experimental/NNI/Source/NeuralNetworkInferenceEditor/Private/QA/NeuralNetworkInferenceQAAssetTypeActions.h
+++ b/Engine/Plugins/Experimental/NNI/Source/NeuralNetworkInferenceEditor/Private/QA/NeuralNetworkInferenceQAAssetTypeActions.h
@@ -12,7 +12,6 @@ public:
 	virtual FText GetName() const override;
 	virtual FColor GetTypeColor() const override;
 	virtual UClass* GetSupportedClass() const override;
-	virtual bool HasActions(const TArray<UObject*>& InObjects) const override;
 	virtual uint32 GetCategories() override;
 	virtual bool IsImportedAsset() const override;
 };
diff --git a/Engine/Plugins/FX/Niagara/Source/NiagaraEditor/Private/AssetTypeActions/AssetTypeActions_NiagaraEmitter.h b/Engine/Plugins/FX/Niagara/Source/NiagaraEditor/Private/AssetTypeActions/AssetTypeActions_NiagaraEmitter.h
--- a/Engine/Plugins/FX/Niagara/Source/NiagaraEditor/Private/AssetTypeActions/AssetTypeActions_NiagaraEmitter.h
+++ b/Engine/Plugins/FX/Niagara/Source/NiagaraEditor/Private/AssetTypeActions/AssetTypeActions_NiagaraEmitter.h
@@ -16,7 +16,6 @@ public:
 	virtual FText GetName() const override { return NSLOCTEXT("AssetTypeActions", "AssetTypeActions_NiagaraEmitter", "Niagara Emitter"); }
 	virtual FColor GetTypeColor() const override;
 	virtual UClass* GetSupportedClass() const override;
-	virtual bool HasActions(const TArray<UObject*>& InObjects) const override { return true; }
 	virtual void GetActions(const TArray<UObject*>& InObjects, struct FToolMenuSection& Section) override;
 	virtual void OpenAssetEditor(const TArray<UObject*>& InObjects, TSharedPtr<class IToolkitHost> EditWithinLevelEditor = TSharedPtr<IToolkitHost>()) override;
 	virtual uint32 GetCategories() override { return FNiagaraEditorModule::GetAssetCategory(); }
diff --git a/Engine/Plugins/FX/Niagara/Source/NiagaraEditor/Private/AssetTypeActions/AssetTypeActions_NiagaraParameterCollection.h b/Engine/Plugins/FX/Niagara/Source/NiagaraEditor/Private/AssetTypeActions/AssetTypeActions_NiagaraParameterCollection.h
--- a/Engine/Plugins/FX/Niagara/Source/NiagaraEditor/Private/AssetTypeActions/AssetTypeActions_NiagaraParameterCollection.h
+++ b/Engine/Plugins/FX/Niagara/Source/NiagaraEditor/Private/AssetTypeActions/AssetTypeActions_NiagaraParameterCollection.h
@@ -16,7 +16,6 @@ public:
 	virtual UClass* GetSupportedClass() const override;
 	virtual void OpenAssetEditor(const TArray<UObject*>& InObjects, TSharedPtr<class IToolkitHost> EditWithinLevelEditor = TSharedPtr<IToolkitHost>()) override;
 	virtual uint32 GetCategories() override { return FNiagaraEditorModule::GetAssetCategory(); }
-	virtual bool HasActions(const TArray<UObject*>& InObjects) const override { return true; }
 	virtual void GetActions(const TArray<UObject*>& InObjects, struct FToolMenuSection& Section) override;
 
 	void ExecuteNewNPC(TArray<TWeakObjectPtr<class UNiagaraParameterCollection>> Objects);
diff --git a/Engine/Plugins/FX/Niagara/Source/NiagaraEditor/Private/AssetTypeActions/AssetTypeActions_NiagaraScript.cpp b/Engine/Plugins/FX/Niagara/Source/NiagaraEditor/Private/AssetTypeActions/AssetTypeActions_NiagaraScript.cpp
--- a/Engine/Plugins/FX/Niagara/Source/NiagaraEditor/Private/AssetTypeActions/AssetTypeActions_NiagaraScript.cpp
+++ b/Engine/Plugins/FX/Niagara/Source/NiagaraEditor/Private/AssetTypeActions/AssetTypeActions_NiagaraScript.cpp
@@ -62,18 +62,6 @@ FText FAssetTypeActions_NiagaraScript::GetDisplayNameFromAssetData(const FAssetD
 	return FAssetTypeActions_NiagaraScript::GetName();
 }
 
-bool FAssetTypeActions_NiagaraScript::HasActions(const TArray<UObject*>& InObjects) const
-{
-	for (UObject* ActionObject : InObjects)
-	{
-		if (FNiagaraEditorUtilities::IsCompilableAssetClass(ActionObject->GetClass()) == false)
-		{
-			return false;
-		}
-	}
-	return true;
-}
-
 void FAssetTypeActions_NiagaraScript::GetActions(const TArray<UObject*>& InObjects, FToolMenuSection& Section)
 {
 	Section.AddMenuEntry(
diff --git a/Engine/Plugins/FX/Niagara/Source/NiagaraEditor/Private/AssetTypeActions/AssetTypeActions_NiagaraScript.h b/Engine/Plugins/FX/Niagara/Source/NiagaraEditor/Private/AssetTypeActions/AssetTypeActions_NiagaraScript.h
--- a/Engine/Plugins/FX/Niagara/Source/NiagaraEditor/Private/AssetTypeActions/AssetTypeActions_NiagaraScript.h
+++ b/Engine/Plugins/FX/Niagara/Source/NiagaraEditor/Private/AssetTypeActions/AssetTypeActions_NiagaraScript.h
@@ -17,7 +17,6 @@ public:
 	virtual void OpenAssetEditor(const TArray<UObject*>& InObjects, TSharedPtr<class IToolkitHost> EditWithinLevelEditor = TSharedPtr<IToolkitHost>()) override;
 	virtual uint32 GetCategories() override { return FNiagaraEditorModule::GetAssetCategory(); }
 	virtual FText GetDisplayNameFromAssetData(const FAssetData& AssetData) const override;
-	virtual bool HasActions(const TArray<UObject*>& InObjects) const override;
 	virtual void GetActions(const TArray<UObject*>& InObjects, struct FToolMenuSection& Section) override;
 	// End IAssetTypeActions Interface
 
diff --git a/Engine/Plugins/Media/ImgMedia/Source/ImgMediaEditor/Private/ImgMediaEditorModule.cpp b/Engine/Plugins/Media/ImgMedia/Source/ImgMediaEditor/Private/ImgMediaEditorModule.cpp
--- a/Engine/Plugins/Media/ImgMedia/Source/ImgMediaEditor/Private/ImgMediaEditorModule.cpp
+++ b/Engine/Plugins/Media/ImgMedia/Source/ImgMediaEditor/Private/ImgMediaEditorModule.cpp
@@ -201,55 +201,51 @@ protected:
 
 	void ExtendContentMenu()
 	{
-		if (UToolMenu* WorldAssetMenu = UToolMenus::Get()->ExtendMenu("ContentBrowser.AssetContextMenu.ImgMediaSource"))
+		UToolMenu* Menu = UToolMenus::Get()->ExtendMenu("ContentBrowser.AssetContextMenu.ImgMediaSource");
+		
+		FToolMenuSection& Section = Menu->AddDynamicSection("ImgMediaSource", FNewToolMenuDelegate::CreateLambda([this](UToolMenu* ToolMenu)
 		{
-			FToolMenuSection& Section = WorldAssetMenu->AddDynamicSection("ImgMediaSource", FNewToolMenuDelegate::CreateLambda([this](UToolMenu* ToolMenu)
+			if (const UContentBrowserAssetContextMenuContext* Context = UContentBrowserAssetContextMenuContext::FindContextWithAssets(*ToolMenu))
 			{
-				if (ToolMenu)
+				if (Context->SelectedAssets.Num() != 1)
 				{
-					if (UContentBrowserAssetContextMenuContext* AssetMenuContext = ToolMenu->Context.FindContext<UContentBrowserAssetContextMenuContext>())
+					return;
+				}
+
+				// Add menu entry.
+				if (UImgMediaSource* ImgMediaSource = Cast<UImgMediaSource>(Context->SelectedAssets[0].GetAsset()))
+				{
+					// Get path from media source.
+					// Make sure it ends with / so it opens the directory.
+					FString InputPath = ImgMediaSource->GetSequencePath();
+					if (InputPath.EndsWith(TEXT("/")) == false)
+					{
+						InputPath += TEXT("/");
+					}
+
+					const TAttribute<FText> Label = LOCTEXT("ProcessEXR", "Process EXRs");
+					const TAttribute<FText> ToolTip = LOCTEXT("ProcessEXR_Tooltip", "Open a tab to allow adding tiles and mips to an EXR sequence.");
+					const FSlateIcon Icon = FSlateIcon(FAppStyle::GetAppStyleSetName(),
+						"LevelEditor.Tabs.Viewports");
+					const FToolMenuExecuteAction UIAction =
+						FToolMenuExecuteAction::CreateLambda([InputPath]
+						(const FToolMenuContext& InContext)
+					{
+						ProcessEXRInputPath = InputPath;
+						OpenProcesEXRTab();
+					});
+				
+					const FName SectionName = TEXT("ImgMedia");
+					FToolMenuSection* Section = ToolMenu->FindSection(SectionName);
+					if (!Section)
 					{
-						if (AssetMenuContext->SelectedObjects.Num() != 1)
-						{
-							return;
-						}
-
-						// Add menu entry.
-						if (UImgMediaSource* WorldAsset = Cast<UImgMediaSource>(AssetMenuContext->SelectedObjects[0].Get()))
-						{
-							// Get path from media source.
-							// Make sure it ends with / so it opens the directory.
-							FString InputPath = WorldAsset->GetSequencePath();
-							if (InputPath.EndsWith(TEXT("/")) == false)
-							{
-								InputPath += TEXT("/");
-							}
-
-							const TAttribute<FText> Label = LOCTEXT("ProcessEXR", "Process EXRs");
-							const TAttribute<FText> ToolTip = LOCTEXT("ProcessEXR_Tooltip", "Open a tab to allow adding tiles and mips to an EXR sequence.");
-							const FSlateIcon Icon = FSlateIcon(FAppStyle::GetAppStyleSetName(),
-								"LevelEditor.Tabs.Viewports");
-							const FToolMenuExecuteAction UIAction =
-								FToolMenuExecuteAction::CreateLambda([InputPath]
-								(const FToolMenuContext& InContext)
-							{
-								ProcessEXRInputPath = InputPath;
-								OpenProcesEXRTab();
-							});
-						
-							const FName SectionName = TEXT("ImgMedia");
-							FToolMenuSection* Section = ToolMenu->FindSection(SectionName);
-							if (!Section)
-							{
-								Section = &(ToolMenu->AddSection(SectionName, LOCTEXT("ImgMediaSectionLabel", "ImgMedia")));
-							}
-
-							Section->AddMenuEntry("ImgMedia_ProcessEXR", Label, ToolTip, Icon, UIAction);
-						}
+						Section = &(ToolMenu->AddSection(SectionName, LOCTEXT("ImgMediaSectionLabel", "ImgMedia")));
 					}
+
+					Section->AddMenuEntry("ImgMedia_ProcessEXR", Label, ToolTip, Icon, UIAction);
 				}
-			}), FToolMenuInsert(NAME_None, EToolMenuInsertType::Default));
-		}
+			}
+		}), FToolMenuInsert(NAME_None, EToolMenuInsertType::Default));
 	}
 
 	/**
diff --git a/Engine/Plugins/Media/MediaFrameworkUtilities/Source/MediaFrameworkUtilitiesEditor/Private/AssetTypeActions/AssetTypeActions_MediaBundle.h b/Engine/Plugins/Media/MediaFrameworkUtilities/Source/MediaFrameworkUtilitiesEditor/Private/AssetTypeActions/AssetTypeActions_MediaBundle.h
--- a/Engine/Plugins/Media/MediaFrameworkUtilities/Source/MediaFrameworkUtilitiesEditor/Private/AssetTypeActions/AssetTypeActions_MediaBundle.h
+++ b/Engine/Plugins/Media/MediaFrameworkUtilities/Source/MediaFrameworkUtilitiesEditor/Private/AssetTypeActions/AssetTypeActions_MediaBundle.h
@@ -9,7 +9,6 @@ class FAssetTypeActions_MediaBundle : public FAssetTypeActions_Base
 {
 public:
 	virtual FColor GetTypeColor() const override { return FColor(140, 62, 35); }
-	virtual bool HasActions(const TArray<UObject*>& InObjects) const override { return false; }
 	virtual uint32 GetCategories() override { return EAssetTypeCategories::Media; }
 	virtual bool IsImportedAsset() const override { return false; }
 
diff --git a/Engine/Plugins/Media/MediaFrameworkUtilities/Source/MediaFrameworkUtilitiesEditor/Private/AssetTypeActions/AssetTypeActions_MediaProfile.h b/Engine/Plugins/Media/MediaFrameworkUtilities/Source/MediaFrameworkUtilitiesEditor/Private/AssetTypeActions/AssetTypeActions_MediaProfile.h
--- a/Engine/Plugins/Media/MediaFrameworkUtilities/Source/MediaFrameworkUtilitiesEditor/Private/AssetTypeActions/AssetTypeActions_MediaProfile.h
+++ b/Engine/Plugins/Media/MediaFrameworkUtilities/Source/MediaFrameworkUtilitiesEditor/Private/AssetTypeActions/AssetTypeActions_MediaProfile.h
@@ -10,7 +10,6 @@ class FAssetTypeActions_MediaProfile : public FAssetTypeActions_Base
 public:
 	// IAssetTypeActions Implementation
 	virtual FColor GetTypeColor() const override { return FColor(140, 62, 35); }
-	virtual bool HasActions(const TArray<UObject*>& InObjects) const override { return false; }
 	virtual uint32 GetCategories() override { return EAssetTypeCategories::Media; }
 	virtual bool IsImportedAsset() const override { return false; }
 	virtual FText GetName() const override;
diff --git a/Engine/Plugins/Media/MediaPlayerEditor/Source/MediaPlayerEditor/Private/AssetTools/MediaTextureActions.cpp b/Engine/Plugins/Media/MediaPlayerEditor/Source/MediaPlayerEditor/Private/AssetTools/MediaTextureActions.cpp
--- a/Engine/Plugins/Media/MediaPlayerEditor/Source/MediaPlayerEditor/Private/AssetTools/MediaTextureActions.cpp
+++ b/Engine/Plugins/Media/MediaPlayerEditor/Source/MediaPlayerEditor/Private/AssetTools/MediaTextureActions.cpp
@@ -63,37 +63,26 @@ void FMediaTextureActions::GetActions(const TArray<UObject*>& InObjects, FToolMe
 	}*/
 }
 
-
 uint32 FMediaTextureActions::GetCategories()
 {
 	return EAssetTypeCategories::Textures | EAssetTypeCategories::Media;
 }
 
-
 FText FMediaTextureActions::GetName() const
 {
 	return NSLOCTEXT("AssetTypeActions", "AssetTypeActions_MediaTexture", "Media Texture");
 }
 
-
 UClass* FMediaTextureActions::GetSupportedClass() const
 {
 	return UMediaTexture::StaticClass();
 }
 
-
 FColor FMediaTextureActions::GetTypeColor() const
 {
 	return FColor::Red;
 }
 
-
-bool FMediaTextureActions::HasActions(const TArray<UObject*>& InObjects) const
-{
-	return true;
-}
-
-
 void FMediaTextureActions::OpenAssetEditor(const TArray<UObject*>& InObjects, TSharedPtr<IToolkitHost> EditWithinLevelEditor)
 {
 	EToolkitMode::Type Mode = EditWithinLevelEditor.IsValid() ? EToolkitMode::WorldCentric : EToolkitMode::Standalone;
diff --git a/Engine/Plugins/Media/MediaPlayerEditor/Source/MediaPlayerEditor/Private/AssetTools/MediaTextureActions.h b/Engine/Plugins/Media/MediaPlayerEditor/Source/MediaPlayerEditor/Private/AssetTools/MediaTextureActions.h
--- a/Engine/Plugins/Media/MediaPlayerEditor/Source/MediaPlayerEditor/Private/AssetTools/MediaTextureActions.h
+++ b/Engine/Plugins/Media/MediaPlayerEditor/Source/MediaPlayerEditor/Private/AssetTools/MediaTextureActions.h
@@ -26,7 +26,6 @@ public:
 	virtual FText GetName() const override;
 	virtual UClass* GetSupportedClass() const override;
 	virtual FColor GetTypeColor() const override;
-	virtual bool HasActions(const TArray<UObject*>& InObjects) const override;
 	virtual void OpenAssetEditor(const TArray<UObject*>& InObjects, TSharedPtr<class IToolkitHost> EditWithinLevelEditor = TSharedPtr<IToolkitHost>()) override;
 
 private:
diff --git a/Engine/Plugins/Media/TimecodeSynchronizer/Source/TimecodeSynchronizerEditor/Private/AssetTypeActions/AssetTypeActions_TimecodeSynchronizer.h b/Engine/Plugins/Media/TimecodeSynchronizer/Source/TimecodeSynchronizerEditor/Private/AssetTypeActions/AssetTypeActions_TimecodeSynchronizer.h
--- a/Engine/Plugins/Media/TimecodeSynchronizer/Source/TimecodeSynchronizerEditor/Private/AssetTypeActions/AssetTypeActions_TimecodeSynchronizer.h
+++ b/Engine/Plugins/Media/TimecodeSynchronizer/Source/TimecodeSynchronizerEditor/Private/AssetTypeActions/AssetTypeActions_TimecodeSynchronizer.h
@@ -9,7 +9,6 @@ class FAssetTypeActions_TimecodeSynchronizer : public FAssetTypeActions_Base
 {
 public:
 	virtual FColor GetTypeColor() const override { return FColor(140, 62, 35); }
-	virtual bool HasActions(const TArray<UObject*>& InObjects) const override { return false; }
 	virtual uint32 GetCategories() override { return EAssetTypeCategories::Media; }
 	virtual bool IsImportedAsset() const override { return false; }
 
diff --git a/Engine/Plugins/MovieScene/LevelSequenceEditor/Source/LevelSequenceEditor/Private/AssetTools/LevelSequenceActions.h b/Engine/Plugins/MovieScene/LevelSequenceEditor/Source/LevelSequenceEditor/Private/AssetTools/LevelSequenceActions.h
--- a/Engine/Plugins/MovieScene/LevelSequenceEditor/Source/LevelSequenceEditor/Private/AssetTools/LevelSequenceActions.h
+++ b/Engine/Plugins/MovieScene/LevelSequenceEditor/Source/LevelSequenceEditor/Private/AssetTools/LevelSequenceActions.h
@@ -27,7 +27,6 @@ public:
 public:
 	
 	// IAssetTypeActions interface
-
 	virtual uint32 GetCategories() override;
 	virtual FText GetName() const override;
 	virtual UClass* GetSupportedClass() const override;
@@ -35,7 +34,6 @@ public:
 	virtual void OpenAssetEditor(const TArray<UObject*>& InObjects, TSharedPtr<class IToolkitHost> EditWithinLevelEditor = TSharedPtr<IToolkitHost>()) override;
 	virtual bool ShouldForceWorldCentric() override;
 	virtual bool CanLocalize() const override { return false; }
-	virtual bool HasActions(const TArray<UObject*>& InObjects) const { return false; }
 
 private:
 
diff --git a/Engine/Plugins/MovieScene/TemplateSequence/Source/TemplateSequenceEditor/Private/AssetTools/TemplateSequenceActions.h b/Engine/Plugins/MovieScene/TemplateSequence/Source/TemplateSequenceEditor/Private/AssetTools/TemplateSequenceActions.h
--- a/Engine/Plugins/MovieScene/TemplateSequence/Source/TemplateSequenceEditor/Private/AssetTools/TemplateSequenceActions.h
+++ b/Engine/Plugins/MovieScene/TemplateSequence/Source/TemplateSequenceEditor/Private/AssetTools/TemplateSequenceActions.h
@@ -34,7 +34,6 @@ public:
 	virtual void OpenAssetEditor(const TArray<UObject*>& InObjects, TSharedPtr<class IToolkitHost> EditWithinLevelEditor = TSharedPtr<IToolkitHost>()) override;
 	virtual bool ShouldForceWorldCentric() override;
 	virtual bool CanLocalize() const override { return false; }
-	virtual bool HasActions(const TArray<UObject*>& InObjects) const override { return false; }
 
 protected:
 
diff --git a/Engine/Plugins/Runtime/DataRegistry/Source/DataRegistryEditor/Private/AssetTypeActions_DataRegistry.h b/Engine/Plugins/Runtime/DataRegistry/Source/DataRegistryEditor/Private/AssetTypeActions_DataRegistry.h
--- a/Engine/Plugins/Runtime/DataRegistry/Source/DataRegistryEditor/Private/AssetTypeActions_DataRegistry.h
+++ b/Engine/Plugins/Runtime/DataRegistry/Source/DataRegistryEditor/Private/AssetTypeActions_DataRegistry.h
@@ -9,7 +9,6 @@ class FAssetTypeActions_DataRegistry : public FAssetTypeActions_Base
 {
 public:
 	virtual FColor GetTypeColor() const override { return FColor(140, 62, 35); }
-	virtual bool HasActions(const TArray<UObject*>& InObjects) const override { return false; }
 	virtual uint32 GetCategories() override { return EAssetTypeCategories::Misc; }
 	virtual bool IsImportedAsset() const override { return false; }
 
diff --git a/Engine/Plugins/Runtime/GeometryCache/Source/GeometryCacheEd/Classes/AssetTypeActions_GeometryCache.h b/Engine/Plugins/Runtime/GeometryCache/Source/GeometryCacheEd/Classes/AssetTypeActions_GeometryCache.h
--- a/Engine/Plugins/Runtime/GeometryCache/Source/GeometryCacheEd/Classes/AssetTypeActions_GeometryCache.h
+++ b/Engine/Plugins/Runtime/GeometryCache/Source/GeometryCacheEd/Classes/AssetTypeActions_GeometryCache.h
@@ -14,8 +14,6 @@ public:
 	virtual FText GetName() const override;
 	virtual FColor GetTypeColor() const override;
 	virtual UClass* GetSupportedClass() const override;
-	virtual bool HasActions(const TArray<UObject*>& InObjects) const override;
-	virtual void GetActions(const TArray<UObject*>& InObjects, struct FToolMenuSection& Section) override;
 	virtual void OpenAssetEditor(const TArray<UObject*>& InObjects, TSharedPtr<class IToolkitHost> EditWithinLevelEditor = TSharedPtr<IToolkitHost>()) override;
 	virtual uint32 GetCategories() override;
 	virtual class UThumbnailInfo* GetThumbnailInfo(UObject* Asset) const override;
diff --git a/Engine/Plugins/Runtime/GeometryCache/Source/GeometryCacheEd/Private/AssetTypeActions_GeometryCache.cpp b/Engine/Plugins/Runtime/GeometryCache/Source/GeometryCacheEd/Private/AssetTypeActions_GeometryCache.cpp
--- a/Engine/Plugins/Runtime/GeometryCache/Source/GeometryCacheEd/Private/AssetTypeActions_GeometryCache.cpp
+++ b/Engine/Plugins/Runtime/GeometryCache/Source/GeometryCacheEd/Private/AssetTypeActions_GeometryCache.cpp
@@ -21,16 +21,6 @@ UClass* FAssetTypeActions_GeometryCache::GetSupportedClass() const
 	return UGeometryCache::StaticClass();
 }
 
-bool FAssetTypeActions_GeometryCache::HasActions(const TArray<UObject*>& InObjects) const
-{
-	return false;
-}
-
-void FAssetTypeActions_GeometryCache::GetActions(const TArray<UObject*>& InObjects, FToolMenuSection& Section)
-{
-	FAssetTypeActions_Base::GetActions(InObjects, Section);
-}
-
 void FAssetTypeActions_GeometryCache::OpenAssetEditor(const TArray<UObject*>& InObjects, TSharedPtr<class IToolkitHost> EditWithinLevelEditor /*= TSharedPtr<IToolkitHost>()*/)
 {
 	FAssetTypeActions_Base::OpenAssetEditor(InObjects, EditWithinLevelEditor);
diff --git a/Engine/Plugins/Runtime/HairStrands/Source/HairStrandsEditor/HairStrandsEditor.Build.cs b/Engine/Plugins/Runtime/HairStrands/Source/HairStrandsEditor/HairStrandsEditor.Build.cs
--- a/Engine/Plugins/Runtime/HairStrands/Source/HairStrandsEditor/HairStrandsEditor.Build.cs
+++ b/Engine/Plugins/Runtime/HairStrands/Source/HairStrandsEditor/HairStrandsEditor.Build.cs
@@ -21,6 +21,7 @@ namespace UnrealBuildTool.Rules
 					"SlateCore",
 					"Projects",
 					"ToolMenus",
+					"ContentBrowser",
 					"UnrealEd",
 					"AssetTools",
 					"EditorInteractiveToolsFramework",
diff --git a/Engine/Plugins/Runtime/HairStrands/Source/HairStrandsEditor/Private/GroomActions.cpp b/Engine/Plugins/Runtime/HairStrands/Source/HairStrandsEditor/Private/GroomActions.cpp
--- a/Engine/Plugins/Runtime/HairStrands/Source/HairStrandsEditor/Private/GroomActions.cpp
+++ b/Engine/Plugins/Runtime/HairStrands/Source/HairStrandsEditor/Private/GroomActions.cpp
@@ -49,8 +49,6 @@ bool FGroomActions::CanFilter()
 
 void FGroomActions::GetActions(const TArray<UObject*>& InObjects, FToolMenuSection& Section)
 {
-	FAssetTypeActions_Base::GetActions(InObjects, Section);
-
 	TArray<TWeakObjectPtr<UGroomAsset>> GroomAssets = GetTypedWeakObjectPtrs<UGroomAsset>(InObjects);
 
 	Section.AddMenuEntry(
diff --git a/Engine/Plugins/Runtime/HairStrands/Source/HairStrandsEditor/Private/GroomBindingActions.cpp b/Engine/Plugins/Runtime/HairStrands/Source/HairStrandsEditor/Private/GroomBindingActions.cpp
--- a/Engine/Plugins/Runtime/HairStrands/Source/HairStrandsEditor/Private/GroomBindingActions.cpp
+++ b/Engine/Plugins/Runtime/HairStrands/Source/HairStrandsEditor/Private/GroomBindingActions.cpp
@@ -1,43 +1,18 @@
 // Copyright Epic Games, Inc. All Rights Reserved.
 
 #include "GroomBindingActions.h"
+
+#include "ContentBrowserMenuContexts.h"
 #include "GroomAsset.h"
 
-#include "EditorFramework/AssetImportData.h"
 #include "GeometryCache.h"
-#include "HairStrandsRendering.h"
 #include "Toolkits/SimpleAssetEditor.h"
 #include "ToolMenuSection.h"
 #include "GroomBindingBuilder.h"
+#include "ToolMenus.h"
 
 #define LOCTEXT_NAMESPACE "AssetTypeActions"
 
-FGroomBindingActions::FGroomBindingActions()
-{}
-
-bool FGroomBindingActions::CanFilter()
-{
-	return true;
-}
-
-void FGroomBindingActions::GetActions(const TArray<UObject*>& InObjects, FToolMenuSection& Section)
-{
-	FAssetTypeActions_Base::GetActions(InObjects, Section);
-
-	TArray<TWeakObjectPtr<UGroomBindingAsset>> GroomBindingAssets = GetTypedWeakObjectPtrs<UGroomBindingAsset>(InObjects);
-
-	Section.AddMenuEntry(
-		"RebuildGroomBinding",
-		LOCTEXT("RebuildGroomBinding", "Rebuild"),
-		LOCTEXT("RebuildGroomBindingTooltip", "Rebuild the groom binding"),
-		FSlateIcon(FAppStyle::GetAppStyleSetName(), "ContentBrowser.AssetActions"),
-		FUIAction(
-			FExecuteAction::CreateSP(this, &FGroomBindingActions::ExecuteRebuildBindingAsset, GroomBindingAssets),
-			FCanExecuteAction::CreateSP(this, &FGroomBindingActions::CanRebuildBindingAsset, GroomBindingAssets)
-		)
-	);
-}
-
 uint32 FGroomBindingActions::GetCategories()
 {
 	return EAssetTypeCategories::Misc;
@@ -58,54 +33,57 @@ FColor FGroomBindingActions::GetTypeColor() const
 	return FColor::White;
 }
 
-bool FGroomBindingActions::HasActions(const TArray<UObject*>& InObjects) const
-{
-	return true;
-}
-
 void FGroomBindingActions::OpenAssetEditor(const TArray<UObject*>& InObjects, TSharedPtr<IToolkitHost> EditWithinLevelEditor)
 {
 	// #ueent_todo: Will need a custom editor at some point, for now just use the Properties editor
 	FSimpleAssetEditor::CreateEditor(EToolkitMode::Standalone, EditWithinLevelEditor, InObjects);
 }
 
-bool FGroomBindingActions::CanRebuildBindingAsset(TArray<TWeakObjectPtr<UGroomBindingAsset>> Objects) const
+void FGroomBindingActions::RegisterMenus()
 {
-	for (TWeakObjectPtr<UGroomBindingAsset> BindingAsset : Objects)
+	FToolMenuOwnerScoped MenuOwner(UE_MODULE_NAME);
+	UToolMenu* Menu = UToolMenus::Get()->ExtendMenu("ContentBrowser.AssetContextMenu.GroomBinding");
+	FToolMenuSection& Section = Menu->FindOrAddSection("GetAssetActions");
+
+	Section.AddDynamicEntry("GroomBinding", FNewToolMenuSectionDelegate::CreateLambda([](FToolMenuSection& InSection)
 	{
-		if (BindingAsset.IsValid())
-		{
-			return true;
-		}
-	}
-	return false;
+		const TAttribute<FText> Label = LOCTEXT("RebuildGroomBinding", "Rebuild");
+		const TAttribute<FText> ToolTip = LOCTEXT("RebuildGroomBindingTooltip", "Rebuild the groom binding");
+		const FSlateIcon Icon = FSlateIcon(FAppStyle::GetAppStyleSetName(), "ContentBrowser.AssetActions");
+		const FToolMenuExecuteAction UIAction = FToolMenuExecuteAction::CreateStatic(&FGroomBindingActions::ExecuteRebuildBindingAsset);
+
+		InSection.AddMenuEntry("GroomBinding_Rebuild", Label, ToolTip, Icon, UIAction);
+	}));
 }
 
-void FGroomBindingActions::ExecuteRebuildBindingAsset(TArray<TWeakObjectPtr<UGroomBindingAsset>> Objects) const
+void FGroomBindingActions::ExecuteRebuildBindingAsset(const FToolMenuContext& MenuContext)
 {
-	for (TWeakObjectPtr<UGroomBindingAsset> BindingAsset : Objects)
+	if (const UContentBrowserAssetContextMenuContext* Context = UContentBrowserAssetContextMenuContext::FindContextWithAssets(MenuContext))
 	{
-		if (BindingAsset.IsValid() && BindingAsset->Groom && BindingAsset->HasValidTarget())
+		for (UGroomBindingAsset* BindingAsset : Context->LoadSelectedObjects<UGroomBindingAsset>())
 		{
-			BindingAsset->Groom->ConditionalPostLoad();
-			if (BindingAsset->GroomBindingType == EGroomBindingMeshType::SkeletalMesh)
+			if (BindingAsset->Groom && BindingAsset->HasValidTarget())
 			{
-				BindingAsset->TargetSkeletalMesh->ConditionalPostLoad();
-				if (BindingAsset->SourceSkeletalMesh)
+				BindingAsset->Groom->ConditionalPostLoad();
+				if (BindingAsset->GroomBindingType == EGroomBindingMeshType::SkeletalMesh)
 				{
-					BindingAsset->SourceSkeletalMesh->ConditionalPostLoad();
+					BindingAsset->TargetSkeletalMesh->ConditionalPostLoad();
+					if (BindingAsset->SourceSkeletalMesh)
+					{
+						BindingAsset->SourceSkeletalMesh->ConditionalPostLoad();
+					}
 				}
-			}
-			else
-			{
-				BindingAsset->TargetGeometryCache->ConditionalPostLoad();
-				if (BindingAsset->SourceGeometryCache)
+				else
 				{
-					BindingAsset->SourceGeometryCache->ConditionalPostLoad();
+					BindingAsset->TargetGeometryCache->ConditionalPostLoad();
+					if (BindingAsset->SourceGeometryCache)
+					{
+						BindingAsset->SourceGeometryCache->ConditionalPostLoad();
+					}
 				}
+				FGroomBindingBuilder::BuildBinding(BindingAsset, true);
+				BindingAsset->GetOutermost()->MarkPackageDirty();
 			}
-			FGroomBindingBuilder::BuildBinding(BindingAsset.Get(), true);
-			BindingAsset->GetOutermost()->MarkPackageDirty();
 		}
 	}
 }
diff --git a/Engine/Plugins/Runtime/HairStrands/Source/HairStrandsEditor/Private/HairStrandsEditor.cpp b/Engine/Plugins/Runtime/HairStrands/Source/HairStrandsEditor/Private/HairStrandsEditor.cpp
--- a/Engine/Plugins/Runtime/HairStrands/Source/HairStrandsEditor/Private/HairStrandsEditor.cpp
+++ b/Engine/Plugins/Runtime/HairStrands/Source/HairStrandsEditor/Private/HairStrandsEditor.cpp
@@ -8,10 +8,8 @@
 #include "GroomCacheActions.h"
 
 #include "Styling/SlateStyleRegistry.h"
-#include "Styling/SlateTypes.h"
-#include "Styling/AppStyle.h"
 #include "Interfaces/IPluginManager.h"
-#include "SlateOptMacros.h"
+#include "ToolMenus.h"
 
 #include "GroomAsset.h"
 #include "GroomBindingAsset.h"
@@ -74,6 +72,8 @@ void FGroomEditor::StartupModule()
 {
 	LLM_SCOPE_BYTAG(GroomEditor)
 
+	UToolMenus::RegisterStartupCallback(FSimpleMulticastDelegate::FDelegate::CreateRaw(this, &FGroomEditor::RegisterMenus));
+
 	IAssetTools& AssetTools = FModuleManager::LoadModuleChecked<FAssetToolsModule>("AssetTools").Get();
 	TSharedRef<IAssetTypeActions> GroomAssetActions = MakeShared<FGroomActions>();
 	TSharedRef<IAssetTypeActions> BindingAssetActions = MakeShared<FGroomBindingActions>();
@@ -152,6 +152,9 @@ void FGroomEditor::StartupModule()
 
 void FGroomEditor::ShutdownModule()
 {
+	UToolMenus::UnRegisterStartupCallback(this);
+	UToolMenus::UnregisterOwner(UE_MODULE_NAME);
+
 	ISettingsModule* SettingsModule = FModuleManager::GetModulePtr<ISettingsModule>("Settings");
 	if (SettingsModule != nullptr)
 	{
@@ -197,6 +200,11 @@ void FGroomEditor::ShutdownModule()
 	}
 }
 
+void FGroomEditor::RegisterMenus()
+{
+	FGroomBindingActions::RegisterMenus();
+}
+
 TArray<TSharedPtr<IGroomTranslator>> FGroomEditor::GetHairTranslators()
 {
 	TArray<TSharedPtr<IGroomTranslator>> Translators;
diff --git a/Engine/Plugins/Runtime/HairStrands/Source/HairStrandsEditor/Public/GroomBindingActions.h b/Engine/Plugins/Runtime/HairStrands/Source/HairStrandsEditor/Public/GroomBindingActions.h
--- a/Engine/Plugins/Runtime/HairStrands/Source/HairStrandsEditor/Public/GroomBindingActions.h
+++ b/Engine/Plugins/Runtime/HairStrands/Source/HairStrandsEditor/Public/GroomBindingActions.h
@@ -7,6 +7,8 @@
 #include "GroomBindingAsset.h"
 
 class ISlateStyle;
+struct FToolMenuContext;
+
 /**
  * Implements an action for groom binding assets.
  */
@@ -19,24 +21,22 @@ public:
 	 *
 	 * @param InStyle The style set to use for asset editor toolkits.
 	 */
-	FGroomBindingActions();
+	FGroomBindingActions() { }
 
 public:
 
 	//~ FAssetTypeActions_Base overrides
 
-	virtual bool CanFilter() override;
-	virtual void GetActions(const TArray<UObject*>& InObjects, FToolMenuSection& Section) override;
+	virtual bool CanFilter() override { return true; }
 	virtual uint32 GetCategories() override;
 	virtual FText GetName() const override;
 	virtual UClass* GetSupportedClass() const override;
 	virtual FColor GetTypeColor() const override;
-	virtual bool HasActions(const TArray<UObject*>& InObjects) const override;
 	virtual void OpenAssetEditor(const TArray<UObject*>& InObjects, TSharedPtr<IToolkitHost> EditWithinLevelEditor = TSharedPtr<IToolkitHost>()) override;
 
-private:
+public:
+	static void RegisterMenus();
 
-	/** Callback for rebuild binding asset action */
-	bool CanRebuildBindingAsset(TArray<TWeakObjectPtr<UGroomBindingAsset>> Objects) const;
-	void ExecuteRebuildBindingAsset(TArray<TWeakObjectPtr<UGroomBindingAsset>> Objects) const;
+protected:
+	static void ExecuteRebuildBindingAsset(const FToolMenuContext& MenuContext);
 };
diff --git a/Engine/Plugins/Runtime/HairStrands/Source/HairStrandsEditor/Public/HairStrandsEditor.h b/Engine/Plugins/Runtime/HairStrands/Source/HairStrandsEditor/Public/HairStrandsEditor.h
--- a/Engine/Plugins/Runtime/HairStrands/Source/HairStrandsEditor/Public/HairStrandsEditor.h
+++ b/Engine/Plugins/Runtime/HairStrands/Source/HairStrandsEditor/Public/HairStrandsEditor.h
@@ -49,6 +49,7 @@ public:
 	TArray<TSharedPtr<IGroomTranslator>> GetHairTranslators();
 
 private:
+	void RegisterMenus();
 
 	/** The collection of registered asset type actions. */
 	TArray<TSharedRef<class IAssetTypeActions>> RegisteredAssetTypeActions;
diff --git a/Engine/Plugins/Runtime/Metasound/Source/MetasoundEditor/MetasoundEditor.Build.cs b/Engine/Plugins/Runtime/Metasound/Source/MetasoundEditor/MetasoundEditor.Build.cs
--- a/Engine/Plugins/Runtime/Metasound/Source/MetasoundEditor/MetasoundEditor.Build.cs
+++ b/Engine/Plugins/Runtime/Metasound/Source/MetasoundEditor/MetasoundEditor.Build.cs
@@ -13,11 +13,13 @@ namespace UnrealBuildTool.Rules
 			(
 				new string[]
 				{
+					"AssetTools",
 					"AppFramework",
 					"AudioExtensions",
 					"AudioMixer",
 					"AudioWidgets",
 					"AudioSynesthesia",
+					"ContentBrowser",
 					"EditorWidgets",
 					"Kismet",
 					"KismetWidgets",
diff --git a/Engine/Plugins/Runtime/Metasound/Source/MetasoundEditor/Private/MetasoundAssetTypeActions.cpp b/Engine/Plugins/Runtime/Metasound/Source/MetasoundEditor/Private/MetasoundAssetTypeActions.cpp
--- a/Engine/Plugins/Runtime/Metasound/Source/MetasoundEditor/Private/MetasoundAssetTypeActions.cpp
+++ b/Engine/Plugins/Runtime/Metasound/Source/MetasoundEditor/Private/MetasoundAssetTypeActions.cpp
@@ -40,41 +40,28 @@ namespace Metasound
 			template <typename TClass, typename TFactory>
 			void ExecuteCreatePreset(const FToolMenuContext& MenuContext)
 			{
-				UContentBrowserAssetContextMenuContext* Context = MenuContext.FindContext<UContentBrowserAssetContextMenuContext>();
-				if (!Context || Context->SelectedObjects.IsEmpty())
+				if (const UContentBrowserAssetContextMenuContext* Context = UContentBrowserAssetContextMenuContext::FindContextWithAssets(MenuContext))
 				{
-					return;
-				}
-
-				FString PackagePath;
-				FString AssetName;
-
-				TWeakObjectPtr<UObject> ReferencedMetaSound = Context->SelectedObjects[0];
-				if (!ReferencedMetaSound.IsValid())
-				{
-					return;
-				}
+					for (TClass* MetaSound : Context->LoadSelectedObjects<TClass>())
+					{
+						FString PackagePath;
+						FString AssetName;
 
-				FAssetToolsModule& AssetToolsModule = FModuleManager::Get().LoadModuleChecked<FAssetToolsModule>("AssetTools");
-				AssetToolsModule.Get().CreateUniqueAssetName(ReferencedMetaSound->GetOutermost()->GetName(), TEXT("_Preset"), PackagePath, AssetName);
+						IAssetTools::Get().CreateUniqueAssetName(MetaSound->GetOutermost()->GetName(), TEXT("_Preset"), PackagePath, AssetName);
 
-				TFactory* Factory = NewObject<TFactory>();
-				check(Factory);
+						TFactory* Factory = NewObject<TFactory>();
+						check(Factory);
 
-				Factory->ReferencedMetaSoundObject = ReferencedMetaSound.Get();
+						Factory->ReferencedMetaSoundObject = MetaSound;
 
-				FContentBrowserModule& ContentBrowserModule = FModuleManager::Get().LoadModuleChecked<FContentBrowserModule>("ContentBrowser");
-				ContentBrowserModule.Get().CreateNewAsset(AssetName, FPackageName::GetLongPackagePath(PackagePath), TClass::StaticClass(), Factory);
+						IContentBrowserSingleton::Get().CreateNewAsset(AssetName, FPackageName::GetLongPackagePath(PackagePath), TClass::StaticClass(), Factory);
+					}
+				}
 			}
 
 			template <typename TPresetClass, typename TFactory, typename TReferenceClass = TPresetClass>
 			void RegisterPresetAction(const TAttribute<FText>& InLabel, const TAttribute<FText>& InToolTip)
 			{
-				if (!UToolMenus::IsToolMenuUIEnabled())
-				{
-					return;
-				}
-
 				FString ClassName = TReferenceClass::StaticClass()->GetName();
 				const FString MenuName = TEXT("ContentBrowser.AssetContextMenu.") + ClassName;
 				UToolMenu* Menu = UToolMenus::Get()->ExtendMenu(*MenuName);
@@ -88,18 +75,6 @@ namespace Metasound
 				FToolMenuSection& Section = Menu->FindOrAddSection("GetAssetActions");
 				Section.AddDynamicEntry(*EntryName, FNewToolMenuSectionDelegate::CreateLambda([InPresetClassName = MoveTemp(PresetClassName), Label = InLabel, ToolTip = InToolTip](FToolMenuSection& InSection)
 				{
-					UContentBrowserAssetContextMenuContext* Context = InSection.FindContext<UContentBrowserAssetContextMenuContext>();
-					if (!Context || Context->SelectedObjects.IsEmpty())
-					{
-						return;
-					}
-
-					TWeakObjectPtr<UObject> MetaSoundObject = Context->SelectedObjects[0];
-					if (!MetaSoundObject.IsValid())
-					{
-						return;
-					}
-
 					// TODO: Make class icons for MetaSound types. For now just use SoundCue.
 					//const FString IconName = TEXT("ClassIcon.") + InPresetClassName;
 					const FString IconName = TEXT("ClassIcon.SoundCue");
diff --git a/Engine/Plugins/Runtime/SoundCueTemplates/Source/SoundCueTemplatesEditor/Private/AssetTypeActions_SoundCueTemplate.cpp b/Engine/Plugins/Runtime/SoundCueTemplates/Source/SoundCueTemplatesEditor/Private/AssetTypeActions_SoundCueTemplate.cpp
--- a/Engine/Plugins/Runtime/SoundCueTemplates/Source/SoundCueTemplatesEditor/Private/AssetTypeActions_SoundCueTemplate.cpp
+++ b/Engine/Plugins/Runtime/SoundCueTemplates/Source/SoundCueTemplatesEditor/Private/AssetTypeActions_SoundCueTemplate.cpp
@@ -22,45 +22,6 @@ UClass* FAssetTypeActions_SoundCueTemplate::GetSupportedClass() const
 	return USoundCueTemplate::StaticClass();
 }
 
-void FAssetTypeActions_SoundCueTemplate::GetActions(const TArray<UObject*>& InObjects, FToolMenuSection& Section)
-{
-	TArray<TWeakObjectPtr<USoundCueTemplate>> Cues = GetTypedWeakObjectPtrs<USoundCueTemplate>(InObjects);
-	FAssetTypeActions_SoundBase::GetActions(InObjects, Section);
-
-	Section.AddMenuEntry(
-		"SoundCueTemplate_CopyToSoundCue",
-		LOCTEXT("SoundCueTemplate_CopyToSoundCue", "Copy To Sound Cue"),
-		LOCTEXT("SoundCueTemplate_CopyToSoundCueTooltip", "Exports a Sound Cue Template to a Sound Cue."),
-		FSlateIcon(FAppStyle::GetAppStyleSetName(), "ClassIcon.SoundCue"),
-		FUIAction(
-			FExecuteAction::CreateSP(this, &FAssetTypeActions_SoundCueTemplate::ExecuteCopyToSoundCue, Cues),
-			FCanExecuteAction()
-		)
-	);
-}
-
-void FAssetTypeActions_SoundCueTemplate::ExecuteCopyToSoundCue(TArray<TWeakObjectPtr<USoundCueTemplate>> Objects)
-{
-	for (const TWeakObjectPtr<USoundCueTemplate>& Object : Objects)
-	{
-		if (!Object.IsValid())
-		{
-			continue;
-		}
-
-		FString Name;
-		FString PackagePath;
-		CreateUniqueAssetName(Object->GetOutermost()->GetName(), TEXT(""), PackagePath, Name);
-
-		if (USoundCueTemplateCopyFactory* Factory = NewObject<USoundCueTemplateCopyFactory>())
-		{
-			Factory->SoundCueTemplate = Object;
-			FContentBrowserModule& ContentBrowserModule = FModuleManager::LoadModuleChecked<FContentBrowserModule>("ContentBrowser");
-			ContentBrowserModule.Get().CreateNewAsset(Name, FPackageName::GetLongPackagePath(PackagePath), USoundCue::StaticClass(), Factory);
-		}
-	}
-}
-
 const TArray<FText>& FAssetTypeActions_SoundCueTemplate::GetSubMenus() const
 {
 	if (GetDefault<UAudioEditorSettings>()->bPinSoundCueTemplateInAssetMenu)
@@ -78,70 +39,89 @@ const TArray<FText>& FAssetTypeActions_SoundCueTemplate::GetSubMenus() const
 
 void FAssetActionExtender_SoundCueTemplate::RegisterMenus()
 {
-	if (!UToolMenus::IsToolMenuUIEnabled())
-	{
-		return;
-	}
-
 	FToolMenuOwnerScoped MenuOwner("SoundCueTemplate");
-	UToolMenu* Menu = UToolMenus::Get()->ExtendMenu("ContentBrowser.AssetContextMenu.SoundWave");
-	FToolMenuSection& Section = Menu->FindOrAddSection("GetAssetActions");
 
-	Section.AddDynamicEntry("SoundCueTemplateSoundWaveAsset", FNewToolMenuSectionDelegate::CreateLambda([](FToolMenuSection& InSection)
 	{
-		UContentBrowserAssetContextMenuContext* Context = InSection.FindContext<UContentBrowserAssetContextMenuContext>();
-		if (!Context || Context->SelectedObjects.IsEmpty())
+		UToolMenu* Menu = UToolMenus::Get()->ExtendMenu("ContentBrowser.AssetContextMenu.SoundCue");
+		FToolMenuSection& Section = Menu->FindOrAddSection("GetAssetActions");
+
+		Section.AddDynamicEntry("SoundCueAsset", FNewToolMenuSectionDelegate::CreateLambda([](FToolMenuSection& InSection)
 		{
-			return;
-		}
+			const TAttribute<FText> Label = LOCTEXT("SoundWave_CreateSoundCueTemplate", "Create SoundCueTemplate");
+			const TAttribute<FText> ToolTip = LOCTEXT("SoundWave_CreateSoundCueTemplateToolTip", "Creates a SoundCueTemplate from the selected sound waves.");
+			const FSlateIcon Icon = FSlateIcon(FAppStyle::GetAppStyleSetName(), "ClassIcon.SoundCue");
+			const FToolMenuExecuteAction UIExecuteAction = FToolMenuExecuteAction::CreateStatic(&FAssetActionExtender_SoundCueTemplate::ExecuteCreateSoundCueTemplate);
+
+			InSection.AddMenuEntry("SoundWave_CreateSoundCueTemplate", Label, ToolTip, Icon, UIExecuteAction);
+		}));
+	}
 
-		for (const TWeakObjectPtr<UObject>& Object : Context->SelectedObjects)
+	{
+		UToolMenu* Menu = UToolMenus::Get()->ExtendMenu("ContentBrowser.AssetContextMenu.SoundCueTemplate");
+		FToolMenuSection& Section = Menu->FindOrAddSection("GetAssetActions");
+
+		Section.AddDynamicEntry("SoundCueTemplate", FNewToolMenuSectionDelegate::CreateLambda([](FToolMenuSection& InSection)
 		{
-			if (Object->IsA<USoundWaveProcedural>())
+			if (const UContentBrowserAssetContextMenuContext* Context = UContentBrowserAssetContextMenuContext::FindContextWithAssets(InSection))
 			{
-				return;
+				{
+					const TAttribute<FText> Label = LOCTEXT("SoundCueTemplate_CopyToSoundCue", "Copy To Sound Cue");
+					const TAttribute<FText> ToolTip = LOCTEXT("SoundCueTemplate_CopyToSoundCueTooltip", "Exports a Sound Cue Template to a Sound Cue.");
+					const FSlateIcon Icon = FSlateIcon(FAppStyle::GetAppStyleSetName(), "ClassIcon.SoundCue");
+					const FToolMenuExecuteAction UIExecuteAction = FToolMenuExecuteAction::CreateStatic(&FAssetActionExtender_SoundCueTemplate::ExecuteCopyToSoundCue);
+
+					InSection.AddMenuEntry("SoundCueTemplate_CopyToSoundCue", Label, ToolTip, Icon, UIExecuteAction);
+				}
 			}
-		}
-
-		const TAttribute<FText> Label = LOCTEXT("SoundWave_CreateSoundCueTemplate", "Create SoundCueTemplate");
-		const TAttribute<FText> ToolTip = LOCTEXT("SoundWave_CreateSoundCueTemplateToolTip", "Creates a SoundCueTemplate from the selected sound waves.");
-		const FSlateIcon Icon = FSlateIcon(FAppStyle::GetAppStyleSetName(), "ClassIcon.SoundCue");
-		const FToolMenuExecuteAction UIExecuteAction = FToolMenuExecuteAction::CreateStatic(&FAssetActionExtender_SoundCueTemplate::ExecuteCreateSoundCueTemplate);
-
-		InSection.AddMenuEntry("SoundWave_CreateSoundCueTemplate", Label, ToolTip, Icon, UIExecuteAction);
-	}));
+		}));
+	}
 }
 
-void FAssetActionExtender_SoundCueTemplate::ExecuteCreateSoundCueTemplate(const struct FToolMenuContext& MenuContext)
+void FAssetActionExtender_SoundCueTemplate::ExecuteCopyToSoundCue(const FToolMenuContext& MenuContext)
 {
-	UContentBrowserAssetContextMenuContext* Context = MenuContext.FindContext<UContentBrowserAssetContextMenuContext>();
-	if (!Context || Context->SelectedObjects.Num() == 0)
+	FAssetToolsModule& AssetToolsModule = FModuleManager::Get().LoadModuleChecked<FAssetToolsModule>("AssetTools");
+	
+	if (const UContentBrowserAssetContextMenuContext* Context = UContentBrowserAssetContextMenuContext::FindContextWithAssets(MenuContext))
 	{
-		return;
-	}
+		for (USoundCueTemplate* Object : Context->LoadSelectedObjects<USoundCueTemplate>())
+		{
+			FString Name;
+			FString PackagePath;
+			AssetToolsModule.Get().CreateUniqueAssetName(Object->GetOutermost()->GetName(), TEXT(""), PackagePath, Name);
 
-	USoundWave* SoundWave = Cast<USoundWave>(Context->SelectedObjects[0]);
-	if (SoundWave == nullptr)
-	{
-		return;
+			if (USoundCueTemplateCopyFactory* Factory = NewObject<USoundCueTemplateCopyFactory>())
+			{
+				Factory->SoundCueTemplate = TWeakObjectPtr<USoundCueTemplate>(Object);
+				FContentBrowserModule& ContentBrowserModule = FModuleManager::LoadModuleChecked<FContentBrowserModule>("ContentBrowser");
+				ContentBrowserModule.Get().CreateNewAsset(Name, FPackageName::GetLongPackagePath(PackagePath), USoundCue::StaticClass(), Factory);
+			}
+		}
 	}
+}
 
-	FString PackagePath;
-	FString Name;
-
+void FAssetActionExtender_SoundCueTemplate::ExecuteCreateSoundCueTemplate(const struct FToolMenuContext& MenuContext)
+{
 	FAssetToolsModule& AssetToolsModule = FModuleManager::Get().LoadModuleChecked<FAssetToolsModule>("AssetTools");
-	AssetToolsModule.Get().CreateUniqueAssetName(SoundWave->GetOutermost()->GetName(), TEXT(""), PackagePath, Name);
-
-	USoundCueTemplateFactory* Factory = NewObject<USoundCueTemplateFactory>();
-	Factory->SoundWaves = FObjectEditorUtils::GetTypedWeakObjectPtrs<USoundWave>(Context->GetSelectedObjects());
-	Factory->ConfigureProperties();
-	Name = Factory->GetDefaultNewAssetName();
-
-	FContentBrowserModule& ContentBrowserModule = FModuleManager::Get().LoadModuleChecked<FContentBrowserModule>("ContentBrowser");
-	ContentBrowserModule.Get().CreateNewAsset(Name, FPackageName::GetLongPackagePath(PackagePath), USoundCueTemplate::StaticClass(), Factory);
+	
+	if (const UContentBrowserAssetContextMenuContext* Context = UContentBrowserAssetContextMenuContext::FindContextWithAssets(MenuContext))
+	{
+		TArray<USoundWave*> SoundWaves = Context->LoadSelectedObjectsIf<USoundWave>([](const FAssetData& AssetData) { return !AssetData.IsInstanceOf<USoundWaveProcedural>(); });
+		if (!SoundWaves.IsEmpty())
+		{
+			FString PackagePath;
+			FString Name;
+			
+			AssetToolsModule.Get().CreateUniqueAssetName(SoundWaves[0]->GetOutermost()->GetName(), TEXT(""), PackagePath, Name);
+
+			USoundCueTemplateFactory* Factory = NewObject<USoundCueTemplateFactory>();
+			Factory->SoundWaves = TArray<TWeakObjectPtr<USoundWave>>(SoundWaves);
+			Factory->ConfigureProperties();
+			Name = Factory->GetDefaultNewAssetName();
+
+			FContentBrowserModule& ContentBrowserModule = FModuleManager::Get().LoadModuleChecked<FContentBrowserModule>("ContentBrowser");
+			ContentBrowserModule.Get().CreateNewAsset(Name, FPackageName::GetLongPackagePath(PackagePath), USoundCueTemplate::StaticClass(), Factory);
+		}
+	}
 }
 
-
-
-
 #undef LOCTEXT_NAMESPACE
diff --git a/Engine/Plugins/Runtime/SoundCueTemplates/Source/SoundCueTemplatesEditor/Private/AssetTypeActions_SoundCueTemplate.h b/Engine/Plugins/Runtime/SoundCueTemplates/Source/SoundCueTemplatesEditor/Private/AssetTypeActions_SoundCueTemplate.h
--- a/Engine/Plugins/Runtime/SoundCueTemplates/Source/SoundCueTemplatesEditor/Private/AssetTypeActions_SoundCueTemplate.h
+++ b/Engine/Plugins/Runtime/SoundCueTemplates/Source/SoundCueTemplatesEditor/Private/AssetTypeActions_SoundCueTemplate.h
@@ -3,10 +3,10 @@
 #pragma once
 
 #include "CoreMinimal.h"
-#include "Toolkits/IToolkitHost.h"
 #include "AssetTypeActions/AssetTypeActions_SoundBase.h"
 
 class USoundCueTemplate;
+struct FToolMenuContext;
 
 class FAssetTypeActions_SoundCueTemplate : public FAssetTypeActions_SoundBase
 {
@@ -15,18 +15,16 @@ public:
 	virtual FText GetName() const override { return NSLOCTEXT("AssetTypeActions", "AssetTypeActions_SoundCueTemplate", "Sound Cue Template"); }
 	virtual FColor GetTypeColor() const override { return FColor::Red; }
 	virtual UClass* GetSupportedClass() const override;
-	virtual void GetActions(const TArray<UObject*>& InObjects, struct FToolMenuSection& Section) override;
 	virtual bool CanFilter() override { return true; }
 	virtual const TArray<FText>& GetSubMenus() const override;
-
-protected:
-	/** Converts the provided SoundCue Template to a fully-modifiable SoundCue */
-	void ExecuteCopyToSoundCue(TArray<TWeakObjectPtr<USoundCueTemplate>> Objects);
 };
 
 class FAssetActionExtender_SoundCueTemplate
 {
 public:
 	static void RegisterMenus();
-	static void ExecuteCreateSoundCueTemplate(const struct FToolMenuContext& MenuContext);
+private:
+	static void ExecuteCreateSoundCueTemplate(const FToolMenuContext& MenuContext);
+	/** Converts the provided SoundCue Template to a fully-modifiable SoundCue */
+	static void ExecuteCopyToSoundCue(const FToolMenuContext& MenuContext);
 };
\ No newline at end of file
diff --git a/Engine/Plugins/Runtime/SoundCueTemplates/Source/SoundCueTemplatesEditor/Private/SoundCueTemplatesEditorModule.cpp b/Engine/Plugins/Runtime/SoundCueTemplates/Source/SoundCueTemplatesEditor/Private/SoundCueTemplatesEditorModule.cpp
--- a/Engine/Plugins/Runtime/SoundCueTemplates/Source/SoundCueTemplatesEditor/Private/SoundCueTemplatesEditorModule.cpp
+++ b/Engine/Plugins/Runtime/SoundCueTemplates/Source/SoundCueTemplatesEditor/Private/SoundCueTemplatesEditorModule.cpp
@@ -20,7 +20,7 @@ void FSoundCueTemplatesEditorModule::ShutdownModule()
 {
 	LLM_SCOPE(ELLMTag::AudioMisc);
 	UToolMenus::UnRegisterStartupCallback(this);
-	UToolMenus::UnregisterOwner("SoundCueTemplates");
+	UToolMenus::UnregisterOwner(UE_MODULE_NAME);
 }
 
 void FSoundCueTemplatesEditorModule::RegisterAssetActions()
diff --git a/Engine/Plugins/Runtime/Synthesis/Source/SynthesisEditor/Classes/AudioImpulseResponseAsset.h b/Engine/Plugins/Runtime/Synthesis/Source/SynthesisEditor/Classes/AudioImpulseResponseAsset.h
--- a/Engine/Plugins/Runtime/Synthesis/Source/SynthesisEditor/Classes/AudioImpulseResponseAsset.h
+++ b/Engine/Plugins/Runtime/Synthesis/Source/SynthesisEditor/Classes/AudioImpulseResponseAsset.h
@@ -24,7 +24,6 @@ class FAudioImpulseResponseExtension
 {
 public:
 	static void RegisterMenus();
-	static void GetExtendedActions(const struct FToolMenuContext& MenuContext);
 	static void ExecuteCreateImpulseResponse(const struct FToolMenuContext& MenuContext);
 };
 
diff --git a/Engine/Plugins/Runtime/Synthesis/Source/SynthesisEditor/Private/AudioImpulseResponseAsset.cpp b/Engine/Plugins/Runtime/Synthesis/Source/SynthesisEditor/Private/AudioImpulseResponseAsset.cpp
--- a/Engine/Plugins/Runtime/Synthesis/Source/SynthesisEditor/Private/AudioImpulseResponseAsset.cpp
+++ b/Engine/Plugins/Runtime/Synthesis/Source/SynthesisEditor/Private/AudioImpulseResponseAsset.cpp
@@ -6,8 +6,6 @@
 #include "ContentBrowserMenuContexts.h"
 #include "Sound/SoundWave.h"
 #include "Sound/SoundWaveProcedural.h"
-#include "Sound/SampleBufferIO.h"
-#include "SampleBuffer.h"
 #include "SubmixEffects/SubmixEffectConvolutionReverb.h"
 
 #include UE_INLINE_GENERATED_CPP_BY_NAME(AudioImpulseResponseAsset)
@@ -21,31 +19,11 @@ UClass* FAssetTypeActions_AudioImpulseResponse::GetSupportedClass() const
 
 void FAudioImpulseResponseExtension::RegisterMenus()
 {
-	if (!UToolMenus::IsToolMenuUIEnabled())
-	{
-		return;
-	}
-
-	FToolMenuOwnerScoped MenuOwner("ConvolutionReverb");
 	UToolMenu* Menu = UToolMenus::Get()->ExtendMenu("ContentBrowser.AssetContextMenu.SoundWave");
 	FToolMenuSection& Section = Menu->FindOrAddSection("GetAssetActions");
 
 	Section.AddDynamicEntry("SoundWaveAssetConversion_CreateImpulseResponse", FNewToolMenuSectionDelegate::CreateLambda([](FToolMenuSection& InSection)
 	{
-		UContentBrowserAssetContextMenuContext* Context = InSection.FindContext<UContentBrowserAssetContextMenuContext>();
-		if (!Context || Context->SelectedObjects.Num() == 0)
-		{
-			return;
-		}
-
-		for (const TWeakObjectPtr<UObject>& Object : Context->SelectedObjects)
-		{
-			if (Object->IsA<USoundWaveProcedural>())
-			{
-				return;
-			}
-		}
-
 		const TAttribute<FText> Label = LOCTEXT("SoundWave_CreateImpulseResponse", "Create Impulse Response");
 		const TAttribute<FText> ToolTip = LOCTEXT("SoundWave_CreateImpulseResponseTooltip", "Creates an impulse response asset using the selected sound wave.");
 		const FSlateIcon Icon = FSlateIcon(FAppStyle::GetAppStyleSetName(), "ClassIcon.ImpulseResponse");
@@ -57,12 +35,6 @@ void FAudioImpulseResponseExtension::RegisterMenus()
 
 void FAudioImpulseResponseExtension::ExecuteCreateImpulseResponse(const FToolMenuContext& MenuContext)
 {
-	UContentBrowserAssetContextMenuContext* Context = MenuContext.FindContext<UContentBrowserAssetContextMenuContext>();
-	if (!Context || Context->SelectedObjects.IsEmpty())
-	{
-		return;
-	}
-
 	const FString DefaultSuffix = TEXT("_IR");
 	FAssetToolsModule& AssetToolsModule = FModuleManager::Get().LoadModuleChecked<FAssetToolsModule>("AssetTools");
 
@@ -70,20 +42,20 @@ void FAudioImpulseResponseExtension::ExecuteCreateImpulseResponse(const FToolMen
 	UAudioImpulseResponseFactory* Factory = NewObject<UAudioImpulseResponseFactory>();
 	
 	// only converts 0th selected object for now (see return statement)
-	for (const TWeakObjectPtr<UObject>& Object : Context->SelectedObjects)
+	if (const UContentBrowserAssetContextMenuContext* Context = UContentBrowserAssetContextMenuContext::FindContextWithAssets(MenuContext))
 	{
-		// stage the soundwave on the factory to be used during asset creation
-		USoundWave* Wave = Cast<USoundWave>(Object);
-		check(Wave);
-		Factory->StagedSoundWave = Wave; // WeakPtr gets reset by the Factory after it is consumed
-
-		// Determine an appropriate name
-		FString Name;
-		FString PackagePath;
-		AssetToolsModule.Get().CreateUniqueAssetName(Wave->GetOutermost()->GetName(), DefaultSuffix, PackagePath, Name);
-
-		// create new asset
-		AssetToolsModule.Get().CreateAsset(Name, FPackageName::GetLongPackagePath(PackagePath), UAudioImpulseResponse::StaticClass(), Factory);
+		for (USoundWave* Wave : Context->LoadSelectedObjectsIf<USoundWave>([](const FAssetData& AssetData) { return !AssetData.IsInstanceOf<USoundWaveProcedural>(); }))
+		{
+			Factory->StagedSoundWave = Wave; // WeakPtr gets reset by the Factory after it is consumed
+
+			// Determine an appropriate name
+			FString Name;
+			FString PackagePath;
+			AssetToolsModule.Get().CreateUniqueAssetName(Wave->GetOutermost()->GetName(), DefaultSuffix, PackagePath, Name);
+
+			// create new asset
+			AssetToolsModule.Get().CreateAsset(Name, FPackageName::GetLongPackagePath(PackagePath), UAudioImpulseResponse::StaticClass(), Factory);
+		}
 	}
 }
 
diff --git a/Engine/Plugins/Runtime/Synthesis/Source/SynthesisEditor/Private/SynthesisEditorModule.cpp b/Engine/Plugins/Runtime/Synthesis/Source/SynthesisEditor/Private/SynthesisEditorModule.cpp
--- a/Engine/Plugins/Runtime/Synthesis/Source/SynthesisEditor/Private/SynthesisEditorModule.cpp
+++ b/Engine/Plugins/Runtime/Synthesis/Source/SynthesisEditor/Private/SynthesisEditorModule.cpp
@@ -62,10 +62,14 @@ void FSynthesisEditorModule::StartupModule()
 
 void FSynthesisEditorModule::ShutdownModule()
 {
+	UToolMenus::UnRegisterStartupCallback(this);
+	UToolMenus::UnregisterOwner(this);
+	UToolMenus::UnregisterOwner(UE_MODULE_NAME);
 }
 
 void FSynthesisEditorModule::RegisterMenus()
 {
 	LLM_SCOPE(ELLMTag::AudioSynthesis);
+	FToolMenuOwnerScoped MenuOwner(this);
 	FAudioImpulseResponseExtension::RegisterMenus();
 }
diff --git a/Engine/Plugins/Runtime/nDisplay/Source/DisplayClusterConfigurator/Private/Views/NewAsset/SDisplayClusterConfiguratorAssetPickerList.cpp b/Engine/Plugins/Runtime/nDisplay/Source/DisplayClusterConfigurator/Private/Views/NewAsset/SDisplayClusterConfiguratorAssetPickerList.cpp
--- a/Engine/Plugins/Runtime/nDisplay/Source/DisplayClusterConfigurator/Private/Views/NewAsset/SDisplayClusterConfiguratorAssetPickerList.cpp
+++ b/Engine/Plugins/Runtime/nDisplay/Source/DisplayClusterConfigurator/Private/Views/NewAsset/SDisplayClusterConfiguratorAssetPickerList.cpp
@@ -32,7 +32,6 @@ void SDisplayClusterConfiguratorAssetPickerList::Construct(const FArguments& InA
 		AssetPickerConfig.bAutohideSearchBar = false;
 		AssetPickerConfig.bAllowDragging = false;
 		AssetPickerConfig.bCanShowClasses = false;
-		AssetPickerConfig.bPreloadAssetsForContextMenu = true;
 		AssetPickerConfig.ThumbnailLabel = EThumbnailLabel::AssetName;
 		
 		AssetPickerConfig.Filter.ClassPaths.Add(UDisplayClusterBlueprint::StaticClass()->GetClassPathName());
diff --git a/Engine/Plugins/VirtualProduction/Takes/Source/TakeRecorder/Private/TakeRecorderModule.cpp b/Engine/Plugins/VirtualProduction/Takes/Source/TakeRecorder/Private/TakeRecorderModule.cpp
--- a/Engine/Plugins/VirtualProduction/Takes/Source/TakeRecorder/Private/TakeRecorderModule.cpp
+++ b/Engine/Plugins/VirtualProduction/Takes/Source/TakeRecorder/Private/TakeRecorderModule.cpp
@@ -192,56 +192,58 @@ void FTakeRecorderModule::RegisterMenus()
 			return;
 		}
 
-		ULevelSequence* LevelSequence = Context->SelectedObjects.Num() == 1 ? Cast<ULevelSequence>(Context->SelectedObjects[0]) : nullptr;
-		if (LevelSequence)
+		if (Context->SelectedAssets.Num() == 1 && Context->SelectedAssets[0].IsInstanceOf(ULevelSequence::StaticClass()))
 		{
-			InSection.AddMenuEntry(
-				"OpenInTakeRecorder_Label",
-				LOCTEXT("OpenInTakeRecorder_Label", "Open in Take Recorder"),
-				LOCTEXT("OpenInTakeRecorder_Tooltip", "Opens this level sequence asset in Take Recorder by copying its contents into the pending take"),
-				FSlateIcon(FTakeRecorderStyle::StyleName, "TakeRecorder.TabIcon"),
-				FExecuteAction::CreateLambda(
-					[LevelSequence]
-					{
-						FLevelEditorModule& LevelEditorModule = FModuleManager::LoadModuleChecked<FLevelEditorModule>("LevelEditor");
-						TSharedPtr<SDockTab> DockTab = LevelEditorModule.GetLevelEditorTabManager()->TryInvokeTab(ITakeRecorderModule::TakeRecorderTabName);
-						if (DockTab.IsValid())
+			if (ULevelSequence* LevelSequence = Cast<ULevelSequence>(Context->SelectedAssets[0].GetAsset()))
+			{
+				InSection.AddMenuEntry(
+					"OpenInTakeRecorder_Label",
+					LOCTEXT("OpenInTakeRecorder_Label", "Open in Take Recorder"),
+					LOCTEXT("OpenInTakeRecorder_Tooltip", "Opens this level sequence asset in Take Recorder by copying its contents into the pending take"),
+					FSlateIcon(FTakeRecorderStyle::StyleName, "TakeRecorder.TabIcon"),
+					FExecuteAction::CreateLambda(
+						[LevelSequence]
 						{
-							TSharedRef<STakeRecorderTabContent> TabContent = StaticCastSharedRef<STakeRecorderTabContent>(DockTab->GetContent());
-
-							// If this sequence has already been recorded, set it up for viewing, otherwise start recording from it.
-							UTakeMetaData* TakeMetaData = LevelSequence->FindMetaData<UTakeMetaData>();
-							if (!TakeMetaData || !TakeMetaData->Recorded())
+							FLevelEditorModule& LevelEditorModule = FModuleManager::LoadModuleChecked<FLevelEditorModule>("LevelEditor");
+							TSharedPtr<SDockTab> DockTab = LevelEditorModule.GetLevelEditorTabManager()->TryInvokeTab(ITakeRecorderModule::TakeRecorderTabName);
+							if (DockTab.IsValid())
 							{
-								TabContent->SetupForRecording(LevelSequence);
-							}
-							else
-							{
-								TabContent->SetupForViewing(LevelSequence);
+								TSharedRef<STakeRecorderTabContent> TabContent = StaticCastSharedRef<STakeRecorderTabContent>(DockTab->GetContent());
+
+								// If this sequence has already been recorded, set it up for viewing, otherwise start recording from it.
+								UTakeMetaData* TakeMetaData = LevelSequence->FindMetaData<UTakeMetaData>();
+								if (!TakeMetaData || !TakeMetaData->Recorded())
+								{
+									TabContent->SetupForRecording(LevelSequence);
+								}
+								else
+								{
+									TabContent->SetupForViewing(LevelSequence);
+								}
 							}
 						}
-					}
-				)
-			);
-			InSection.AddMenuEntry(
-				"RecordIntoTakeRecorder_Label",
-				LOCTEXT("RecordWithTakeRecorder_Label", "Record with Take Recorder"),
-				LOCTEXT("RecordWithTakeRecorder_Tooltip", "Opens this level sequence asset for recording into with Take Recorder"),
-				FSlateIcon(FTakeRecorderStyle::StyleName, "TakeRecorder.TabIcon"),
-				FExecuteAction::CreateLambda(
-					[LevelSequence]
-					{
-						FLevelEditorModule& LevelEditorModule = FModuleManager::LoadModuleChecked<FLevelEditorModule>("LevelEditor");
-						TSharedPtr<SDockTab> DockTab = LevelEditorModule.GetLevelEditorTabManager()->TryInvokeTab(ITakeRecorderModule::TakeRecorderTabName);
-						if (DockTab.IsValid())
+					)
+				);
+				InSection.AddMenuEntry(
+					"RecordIntoTakeRecorder_Label",
+					LOCTEXT("RecordWithTakeRecorder_Label", "Record with Take Recorder"),
+					LOCTEXT("RecordWithTakeRecorder_Tooltip", "Opens this level sequence asset for recording into with Take Recorder"),
+					FSlateIcon(FTakeRecorderStyle::StyleName, "TakeRecorder.TabIcon"),
+					FExecuteAction::CreateLambda(
+						[LevelSequence]
 						{
-							TSharedRef<STakeRecorderTabContent> TabContent = StaticCastSharedRef<STakeRecorderTabContent>(DockTab->GetContent());
+							FLevelEditorModule& LevelEditorModule = FModuleManager::LoadModuleChecked<FLevelEditorModule>("LevelEditor");
+							TSharedPtr<SDockTab> DockTab = LevelEditorModule.GetLevelEditorTabManager()->TryInvokeTab(ITakeRecorderModule::TakeRecorderTabName);
+							if (DockTab.IsValid())
+							{
+								TSharedRef<STakeRecorderTabContent> TabContent = StaticCastSharedRef<STakeRecorderTabContent>(DockTab->GetContent());
 
-							TabContent->SetupForRecordingInto(LevelSequence);
+								TabContent->SetupForRecordingInto(LevelSequence);
+							}
 						}
-					}
-				)
-			);
+					)
+				);
+			}
 		}
 	}));
 #endif // WITH_EDITOR
diff --git a/Engine/Source/Developer/AssetTools/Private/AssetTools.cpp b/Engine/Source/Developer/AssetTools/Private/AssetTools.cpp
--- a/Engine/Source/Developer/AssetTools/Private/AssetTools.cpp
+++ b/Engine/Source/Developer/AssetTools/Private/AssetTools.cpp
@@ -171,6 +171,12 @@
 
 #define LOCTEXT_NAMESPACE "AssetTools"
 
+IAssetTools& IAssetTools::Get()
+{
+	FAssetToolsModule& Module = FModuleManager::LoadModuleChecked<FAssetToolsModule>("AssetTools");
+	return Module.Get();
+}
+
 TScriptInterface<IAssetTools> UAssetToolsHelpers::GetAssetTools()
 {
 	return &UAssetToolsImpl::Get();
diff --git a/Engine/Source/Developer/AssetTools/Private/AssetTypeActions/AssetTypeActions_AnimBoneCompressionSettings.h b/Engine/Source/Developer/AssetTools/Private/AssetTypeActions/AssetTypeActions_AnimBoneCompressionSettings.h
--- a/Engine/Source/Developer/AssetTools/Private/AssetTypeActions/AssetTypeActions_AnimBoneCompressionSettings.h
+++ b/Engine/Source/Developer/AssetTools/Private/AssetTypeActions/AssetTypeActions_AnimBoneCompressionSettings.h
@@ -19,7 +19,6 @@ public:
 
 	virtual void OpenAssetEditor(const TArray<UObject*>& InObjects, TSharedPtr<class IToolkitHost> EditWithinLevelEditor = TSharedPtr<IToolkitHost>()) override;
 
-	virtual bool HasActions(const TArray<UObject*>& InObjects) const override { return true; }
 	virtual void GetActions(const TArray<UObject*>& InObjects, FMenuBuilder& MenuBuilder) override;
 
 private:
diff --git a/Engine/Source/Developer/AssetTools/Private/AssetTypeActions/AssetTypeActions_AnimCurveCompressionSettings.h b/Engine/Source/Developer/AssetTools/Private/AssetTypeActions/AssetTypeActions_AnimCurveCompressionSettings.h
--- a/Engine/Source/Developer/AssetTools/Private/AssetTypeActions/AssetTypeActions_AnimCurveCompressionSettings.h
+++ b/Engine/Source/Developer/AssetTools/Private/AssetTypeActions/AssetTypeActions_AnimCurveCompressionSettings.h
@@ -19,7 +19,6 @@ public:
 
 	virtual void OpenAssetEditor(const TArray<UObject*>& InObjects, TSharedPtr<class IToolkitHost> EditWithinLevelEditor = TSharedPtr<IToolkitHost>()) override;
 
-	virtual bool HasActions(const TArray<UObject*>& InObjects) const override { return true; }
 	virtual void GetActions(const TArray<UObject*>& InObjects, struct FToolMenuSection& Section) override;
 
 private:
diff --git a/Engine/Source/Developer/AssetTools/Private/AssetTypeActions/AssetTypeActions_Class.cpp b/Engine/Source/Developer/AssetTools/Private/AssetTypeActions/AssetTypeActions_Class.cpp
--- a/Engine/Source/Developer/AssetTools/Private/AssetTypeActions/AssetTypeActions_Class.cpp
+++ b/Engine/Source/Developer/AssetTools/Private/AssetTypeActions/AssetTypeActions_Class.cpp
@@ -15,11 +15,6 @@
 
 #define LOCTEXT_NAMESPACE "AssetTypeActions"
 
-bool FAssetTypeActions_Class::HasActions(const TArray<UObject*>& InObjects) const
-{
-	return true;
-}
-
 void FAssetTypeActions_Class::GetActions(const TArray<UObject*>& InObjects, FToolMenuSection& Section)
 {
 	UClass *const BaseClass = (InObjects.Num() == 1) ? Cast<UClass>(InObjects[0]) : nullptr;
diff --git a/Engine/Source/Developer/AssetTools/Private/AssetTypeActions/AssetTypeActions_Class.h b/Engine/Source/Developer/AssetTools/Private/AssetTypeActions/AssetTypeActions_Class.h
--- a/Engine/Source/Developer/AssetTools/Private/AssetTypeActions/AssetTypeActions_Class.h
+++ b/Engine/Source/Developer/AssetTools/Private/AssetTypeActions/AssetTypeActions_Class.h
@@ -19,7 +19,6 @@ public:
 	virtual UClass* GetSupportedClass() const override { return UClass::StaticClass(); }
 	virtual uint32 GetCategories() override { return EAssetTypeCategories::Basic; }
 	
-	virtual bool HasActions(const TArray<UObject*>& InObjects) const override;
 	virtual void GetActions(const TArray<UObject*>& InObjects, struct FToolMenuSection& Section) override;
 
 	virtual class UThumbnailInfo* GetThumbnailInfo(UObject* Asset) const override;
diff --git a/Engine/Source/Developer/AssetTools/Private/AssetTypeActions/AssetTypeActions_Curve.h b/Engine/Source/Developer/AssetTools/Private/AssetTypeActions/AssetTypeActions_Curve.h
--- a/Engine/Source/Developer/AssetTools/Private/AssetTypeActions/AssetTypeActions_Curve.h
+++ b/Engine/Source/Developer/AssetTools/Private/AssetTypeActions/AssetTypeActions_Curve.h
@@ -19,7 +19,6 @@ public:
 	virtual uint32 GetCategories() override { return EAssetTypeCategories::Misc; }
 	virtual bool IsImportedAsset() const override { return true; }
 	virtual void GetResolvedSourceFilePaths(const TArray<UObject*>& TypeAssets, TArray<FString>& OutSourceFilePaths) const override;
-	virtual bool HasActions(const TArray<UObject*>& InObjects) const override { return true; }
 	virtual void GetActions(const TArray<UObject*>& InObjects, struct FToolMenuSection& Section) override;
 	virtual void PerformAssetDiff(UObject* OldAsset, UObject* NewAsset, const FRevisionInfo& OldRevision, const FRevisionInfo& NewRevision) const override;
 
diff --git a/Engine/Source/Developer/AssetTools/Private/AssetTypeActions/AssetTypeActions_CurveLinearColorAtlas.h b/Engine/Source/Developer/AssetTools/Private/AssetTypeActions/AssetTypeActions_CurveLinearColorAtlas.h
--- a/Engine/Source/Developer/AssetTools/Private/AssetTypeActions/AssetTypeActions_CurveLinearColorAtlas.h
+++ b/Engine/Source/Developer/AssetTools/Private/AssetTypeActions/AssetTypeActions_CurveLinearColorAtlas.h
@@ -14,7 +14,6 @@ public:
 	virtual FText GetName() const override { return NSLOCTEXT("AssetTypeActions", "AssetTypeActions_CurveLinearColorAtlas", "Curve Atlas"); }
 	virtual FColor GetTypeColor() const override { return FColor(128,64,64); }
 	virtual UClass* GetSupportedClass() const override { return UCurveLinearColorAtlas::StaticClass(); }
-	virtual bool HasActions ( const TArray<UObject*>& InObjects ) const override { return true; }
 	virtual void GetActions(const TArray<UObject*>& InObjects, struct FToolMenuSection& Section) override;
 	virtual bool IsImportedAsset() const override { return false; }
 	virtual bool CanFilter() override { return true; }
diff --git a/Engine/Source/Developer/AssetTools/Private/AssetTypeActions/AssetTypeActions_DataAsset.cpp b/Engine/Source/Developer/AssetTools/Private/AssetTypeActions/AssetTypeActions_DataAsset.cpp
--- a/Engine/Source/Developer/AssetTools/Private/AssetTypeActions/AssetTypeActions_DataAsset.cpp
+++ b/Engine/Source/Developer/AssetTools/Private/AssetTypeActions/AssetTypeActions_DataAsset.cpp
@@ -62,7 +62,7 @@ void FAssetTypeActions_DataAsset::GetActions(const TArray<UObject*>& InObjects,
 
 	Section.AddMenuEntry(
 		"DataAsset_ChangeClass",
-		LOCTEXT("DataAsset_ChangeClass", "Convert to different DataAsset type"),
+		LOCTEXT("DataAsset_ChangeClass", "Convert to Different DataAsset Type"),
 		LOCTEXT("DataAsset_ChangeClassTip", "Change the class these Data Assets are subclassed from."),
 		FSlateIcon(),
 		FUIAction(
diff --git a/Engine/Source/Developer/AssetTools/Private/AssetTypeActions/AssetTypeActions_FbxSceneImportData.h b/Engine/Source/Developer/AssetTools/Private/AssetTypeActions/AssetTypeActions_FbxSceneImportData.h
--- a/Engine/Source/Developer/AssetTools/Private/AssetTypeActions/AssetTypeActions_FbxSceneImportData.h
+++ b/Engine/Source/Developer/AssetTools/Private/AssetTypeActions/AssetTypeActions_FbxSceneImportData.h
@@ -13,9 +13,7 @@ public:
 	virtual FText GetName() const override { return NSLOCTEXT("AssetTypeActions", "AssetTypeActions_SceneImportData", "Fbx Scene Import Data"); }
 	virtual FColor GetTypeColor() const override { return FColor(255,0,255); }
 	virtual UClass* GetSupportedClass() const override { return UFbxSceneImportData::StaticClass(); }
-	virtual bool HasActions ( const TArray<UObject*>& InObjects ) const override { return false; }
 	virtual uint32 GetCategories() override { return EAssetTypeCategories::Misc; }
 	virtual bool IsImportedAsset() const override { return true; }
 	virtual void GetResolvedSourceFilePaths(const TArray<UObject*>& TypeAssets, TArray<FString>& OutSourceFilePaths) const override;
-
 };
diff --git a/Engine/Source/Developer/AssetTools/Private/AssetTypeActions/AssetTypeActions_Font.h b/Engine/Source/Developer/AssetTools/Private/AssetTypeActions/AssetTypeActions_Font.h
--- a/Engine/Source/Developer/AssetTools/Private/AssetTypeActions/AssetTypeActions_Font.h
+++ b/Engine/Source/Developer/AssetTools/Private/AssetTypeActions/AssetTypeActions_Font.h
@@ -15,7 +15,6 @@ public:
 	virtual FText GetName() const override { return NSLOCTEXT("AssetTypeActions", "AssetTypeActions_Font", "Font"); }
 	virtual FColor GetTypeColor() const override { return FColor(128,128,64); }
 	virtual UClass* GetSupportedClass() const override { return UFont::StaticClass(); }
-	virtual bool HasActions(const TArray<UObject*>& InObjects) const override { return true; }
 	virtual void GetActions(const TArray<UObject*>& InObjects, struct FToolMenuSection& Section) override;
 	virtual void OpenAssetEditor( const TArray<UObject*>& InObjects, TSharedPtr<class IToolkitHost> EditWithinLevelEditor = TSharedPtr<IToolkitHost>() ) override;
 	virtual uint32 GetCategories() override { return EAssetTypeCategories::UI; }
diff --git a/Engine/Source/Developer/AssetTools/Private/AssetTypeActions/AssetTypeActions_FontFace.h b/Engine/Source/Developer/AssetTools/Private/AssetTypeActions/AssetTypeActions_FontFace.h
--- a/Engine/Source/Developer/AssetTools/Private/AssetTypeActions/AssetTypeActions_FontFace.h
+++ b/Engine/Source/Developer/AssetTools/Private/AssetTypeActions/AssetTypeActions_FontFace.h
@@ -15,7 +15,6 @@ public:
 	virtual FText GetName() const override { return NSLOCTEXT("AssetTypeActions", "AssetTypeActions_FontFace", "Font Face"); }
 	virtual FColor GetTypeColor() const override { return FColor(184,184,112); }
 	virtual UClass* GetSupportedClass() const override { return UFontFace::StaticClass(); }
-	virtual bool HasActions(const TArray<UObject*>& InObjects) const override { return true; }
 	virtual void GetActions(const TArray<UObject*>& InObjects, struct FToolMenuSection& Section) override;
 	virtual void OpenAssetEditor( const TArray<UObject*>& InObjects, TSharedPtr<class IToolkitHost> EditWithinLevelEditor = TSharedPtr<IToolkitHost>() ) override;
 	virtual uint32 GetCategories() override { return EAssetTypeCategories::UI; }
diff --git a/Engine/Source/Developer/AssetTools/Private/AssetTypeActions/AssetTypeActions_ForceFeedbackEffect.h b/Engine/Source/Developer/AssetTools/Private/AssetTypeActions/AssetTypeActions_ForceFeedbackEffect.h
--- a/Engine/Source/Developer/AssetTools/Private/AssetTypeActions/AssetTypeActions_ForceFeedbackEffect.h
+++ b/Engine/Source/Developer/AssetTools/Private/AssetTypeActions/AssetTypeActions_ForceFeedbackEffect.h
@@ -34,7 +34,6 @@ public:
 	virtual FColor GetTypeColor() const override { return FColor(175, 0, 0); }
 	virtual UClass* GetSupportedClass() const override;
 	virtual uint32 GetCategories() override { return EAssetTypeCategories::Misc; }
-	virtual bool HasActions ( const TArray<UObject*>& InObjects ) const override { return true; }
 	virtual void GetActions(const TArray<UObject*>& InObjects, struct FToolMenuSection& Section) override;
 	virtual bool AssetsActivatedOverride( const TArray<UObject*>& InObjects, EAssetTypeActivationMethod::Type ActivationType ) override;
 	virtual TSharedPtr<SWidget> GetThumbnailOverlay(const FAssetData& AssetData) const override;
diff --git a/Engine/Source/Developer/AssetTools/Private/AssetTypeActions/AssetTypeActions_MaterialFunction.h b/Engine/Source/Developer/AssetTools/Private/AssetTypeActions/AssetTypeActions_MaterialFunction.h
--- a/Engine/Source/Developer/AssetTools/Private/AssetTypeActions/AssetTypeActions_MaterialFunction.h
+++ b/Engine/Source/Developer/AssetTools/Private/AssetTypeActions/AssetTypeActions_MaterialFunction.h
@@ -17,7 +17,6 @@ public:
 	virtual FText GetName() const override { return NSLOCTEXT("AssetTypeActions", "AssetTypeActions_MaterialFunction", "Material Function"); }
 	virtual FColor GetTypeColor() const override { return FColor(0,175,175); }
 	virtual UClass* GetSupportedClass() const override { return UMaterialFunction::StaticClass(); }
-	virtual bool HasActions ( const TArray<UObject*>& InObjects ) const override { return true; }
 	virtual void GetActions(const TArray<UObject*>& InObjects, struct FToolMenuSection& Section) override;
 	virtual void OpenAssetEditor( const TArray<UObject*>& InObjects, TSharedPtr<class IToolkitHost> EditWithinLevelEditor = TSharedPtr<IToolkitHost>() ) override;
 	virtual uint32 GetCategories() override { return EAssetTypeCategories::Materials; }
diff --git a/Engine/Source/Developer/AssetTools/Private/AssetTypeActions/AssetTypeActions_MaterialInterface.h b/Engine/Source/Developer/AssetTools/Private/AssetTypeActions/AssetTypeActions_MaterialInterface.h
--- a/Engine/Source/Developer/AssetTools/Private/AssetTypeActions/AssetTypeActions_MaterialInterface.h
+++ b/Engine/Source/Developer/AssetTools/Private/AssetTypeActions/AssetTypeActions_MaterialInterface.h
@@ -14,7 +14,6 @@ public:
 	virtual FText GetName() const override { return NSLOCTEXT("AssetTypeActions", "AssetTypeActions_MaterialInterface", "Material Interface"); }
 	virtual FColor GetTypeColor() const override { return FColor(64,192,64); }
 	virtual UClass* GetSupportedClass() const override { return UMaterialInterface::StaticClass(); }
-	virtual bool HasActions ( const TArray<UObject*>& InObjects ) const override { return true; }
 	virtual void GetActions(const TArray<UObject*>& InObjects, struct FToolMenuSection& Section) override;
 	virtual bool CanFilter() override { return false; }
 	virtual uint32 GetCategories() override { return EAssetTypeCategories::Materials; }
diff --git a/Engine/Source/Developer/AssetTools/Private/AssetTypeActions/AssetTypeActions_NodeMappingContainer.h b/Engine/Source/Developer/AssetTools/Private/AssetTypeActions/AssetTypeActions_NodeMappingContainer.h
--- a/Engine/Source/Developer/AssetTools/Private/AssetTypeActions/AssetTypeActions_NodeMappingContainer.h
+++ b/Engine/Source/Developer/AssetTools/Private/AssetTypeActions/AssetTypeActions_NodeMappingContainer.h
@@ -14,6 +14,5 @@ public:
 	virtual FColor GetTypeColor() const override { return FColor(112,146,190); }
 	virtual UClass* GetSupportedClass() const override { return UNodeMappingContainer::StaticClass(); }
 	virtual bool CanFilter() override { return true; }
-	virtual bool HasActions(const TArray<UObject*>& InObjects) const override { return false; }
 	virtual uint32 GetCategories() override { return EAssetTypeCategories::Animation; }
 };
diff --git a/Engine/Source/Developer/AssetTools/Private/AssetTypeActions/AssetTypeActions_ParticleSystem.h b/Engine/Source/Developer/AssetTools/Private/AssetTypeActions/AssetTypeActions_ParticleSystem.h
--- a/Engine/Source/Developer/AssetTools/Private/AssetTypeActions/AssetTypeActions_ParticleSystem.h
+++ b/Engine/Source/Developer/AssetTools/Private/AssetTypeActions/AssetTypeActions_ParticleSystem.h
@@ -15,7 +15,6 @@ public:
 	virtual FText GetName() const override { return NSLOCTEXT("AssetTypeActions", "AssetTypeActions_ParticleSystem", "Cascade Particle System (Legacy)"); }
 	virtual FColor GetTypeColor() const override { return FColor(255,255,255); }
 	virtual UClass* GetSupportedClass() const override { return UParticleSystem::StaticClass(); }
-	virtual bool HasActions ( const TArray<UObject*>& InObjects ) const override { return true; }
 	virtual void GetActions(const TArray<UObject*>& InObjects, struct FToolMenuSection& Section) override;
 	virtual void OpenAssetEditor( const TArray<UObject*>& InObjects, TSharedPtr<class IToolkitHost> EditWithinLevelEditor = TSharedPtr<IToolkitHost>() ) override;
 	virtual uint32 GetCategories() override { return EAssetTypeCategories::Misc; }
diff --git a/Engine/Source/Developer/AssetTools/Private/AssetTypeActions/AssetTypeActions_PhysicalMaterialMask.h b/Engine/Source/Developer/AssetTools/Private/AssetTypeActions/AssetTypeActions_PhysicalMaterialMask.h
--- a/Engine/Source/Developer/AssetTools/Private/AssetTypeActions/AssetTypeActions_PhysicalMaterialMask.h
+++ b/Engine/Source/Developer/AssetTools/Private/AssetTypeActions/AssetTypeActions_PhysicalMaterialMask.h
@@ -12,7 +12,6 @@ public:
 	virtual FText GetName() const override { return NSLOCTEXT("AssetTypeActions", "AssetTypeActions_PhysicalMaterialMask", "Physical Material Mask"); }
 	virtual FColor GetTypeColor() const override { return FColor(200,162,108); }
 	virtual UClass* GetSupportedClass() const override;
-	virtual bool HasActions(const TArray<UObject*>& InObjects) const override { return true; }
 	virtual void GetActions(const TArray<UObject*>& InObjects, struct FToolMenuSection& Section) override;
 	virtual uint32 GetCategories() override { return EAssetTypeCategories::Physics; }
 	virtual bool IsImportedAsset() const override { return true; }
diff --git a/Engine/Source/Developer/AssetTools/Private/AssetTypeActions/AssetTypeActions_Redirector.h b/Engine/Source/Developer/AssetTools/Private/AssetTypeActions/AssetTypeActions_Redirector.h
--- a/Engine/Source/Developer/AssetTools/Private/AssetTypeActions/AssetTypeActions_Redirector.h
+++ b/Engine/Source/Developer/AssetTools/Private/AssetTypeActions/AssetTypeActions_Redirector.h
@@ -14,7 +14,6 @@ public:
 	virtual FText GetName() const override { return NSLOCTEXT("AssetTypeActions", "AssetTypeActions_Redirector", "Redirector"); }
 	virtual FColor GetTypeColor() const override { return FColor(128, 128, 128); }
 	virtual UClass* GetSupportedClass() const override { return UObjectRedirector::StaticClass(); }
-	virtual bool HasActions ( const TArray<UObject*>& InObjects ) const override { return true; }
 	virtual void GetActions(const TArray<UObject*>& InObjects, struct FToolMenuSection& Section) override;
 	virtual bool AssetsActivatedOverride( const TArray<UObject*>& InObjects, EAssetTypeActivationMethod::Type ActivationType ) override;
 	virtual uint32 GetCategories() override { return EAssetTypeCategories::Misc; }
diff --git a/Engine/Source/Developer/AssetTools/Private/AssetTypeActions/AssetTypeActions_SkeletalMesh.h b/Engine/Source/Developer/AssetTools/Private/AssetTypeActions/AssetTypeActions_SkeletalMesh.h
--- a/Engine/Source/Developer/AssetTools/Private/AssetTypeActions/AssetTypeActions_SkeletalMesh.h
+++ b/Engine/Source/Developer/AssetTools/Private/AssetTypeActions/AssetTypeActions_SkeletalMesh.h
@@ -17,7 +17,6 @@ public:
 	virtual FText GetName() const override { return NSLOCTEXT("AssetTypeActions", "AssetTypeActions_SkeletalMesh", "Skeletal Mesh"); }
 	virtual FColor GetTypeColor() const override { return FColor(241,163,241); }
 	virtual UClass* GetSupportedClass() const override { return USkeletalMesh::StaticClass(); }
-	virtual bool HasActions ( const TArray<UObject*>& InObjects ) const override { return true; }
 	virtual void GetActions(const TArray<UObject*>& InObjects, struct FToolMenuSection& Section) override;
 	virtual void OpenAssetEditor( const TArray<UObject*>& InObjects, TSharedPtr<class IToolkitHost> EditWithinLevelEditor = TSharedPtr<IToolkitHost>() ) override;
 	virtual uint32 GetCategories() override { return EAssetTypeCategories::Basic | EAssetTypeCategories::Animation; }
diff --git a/Engine/Source/Developer/AssetTools/Private/AssetTypeActions/AssetTypeActions_Skeleton.h b/Engine/Source/Developer/AssetTools/Private/AssetTypeActions/AssetTypeActions_Skeleton.h
--- a/Engine/Source/Developer/AssetTools/Private/AssetTypeActions/AssetTypeActions_Skeleton.h
+++ b/Engine/Source/Developer/AssetTools/Private/AssetTypeActions/AssetTypeActions_Skeleton.h
@@ -39,7 +39,6 @@ public:
 	virtual FText GetName() const override { return NSLOCTEXT("AssetTypeActions", "AssetTypeActions_Skeleton", "Skeleton"); }
 	virtual FColor GetTypeColor() const override { return FColor(105,181,205); }
 	virtual UClass* GetSupportedClass() const override { return USkeleton::StaticClass(); }
-	virtual bool HasActions ( const TArray<UObject*>& InObjects ) const override { return true; }
 	virtual void GetActions(const TArray<UObject*>& InObjects, struct FToolMenuSection& Section) override;
 	virtual void OpenAssetEditor( const TArray<UObject*>& InObjects, TSharedPtr<class IToolkitHost> EditWithinLevelEditor = TSharedPtr<IToolkitHost>() ) override;
 	virtual uint32 GetCategories() override { return EAssetTypeCategories::Animation; }
diff --git a/Engine/Source/Developer/AssetTools/Private/AssetTypeActions/AssetTypeActions_SlateBrush.h b/Engine/Source/Developer/AssetTools/Private/AssetTypeActions/AssetTypeActions_SlateBrush.h
--- a/Engine/Source/Developer/AssetTools/Private/AssetTypeActions/AssetTypeActions_SlateBrush.h
+++ b/Engine/Source/Developer/AssetTools/Private/AssetTypeActions/AssetTypeActions_SlateBrush.h
@@ -14,5 +14,4 @@ public:
 	virtual FColor GetTypeColor() const override { return FColor(105, 165, 60); }
 	virtual UClass* GetSupportedClass() const override { return USlateBrushAsset::StaticClass(); }
 	virtual uint32 GetCategories() override { return EAssetTypeCategories::UI; }
-	virtual bool HasActions ( const TArray<UObject*>& InObjects ) const override { return false; }
 };
diff --git a/Engine/Source/Developer/AssetTools/Private/AssetTypeActions/AssetTypeActions_SlateWidgetStyle.h b/Engine/Source/Developer/AssetTools/Private/AssetTypeActions/AssetTypeActions_SlateWidgetStyle.h
--- a/Engine/Source/Developer/AssetTools/Private/AssetTypeActions/AssetTypeActions_SlateWidgetStyle.h
+++ b/Engine/Source/Developer/AssetTools/Private/AssetTypeActions/AssetTypeActions_SlateWidgetStyle.h
@@ -16,7 +16,6 @@ public:
 	virtual FColor GetTypeColor() const override { return FColor(62, 140, 35); }
 	virtual UClass* GetSupportedClass() const override { return USlateWidgetStyleAsset::StaticClass(); }
 	virtual uint32 GetCategories() override { return EAssetTypeCategories::UI; }
-	virtual bool HasActions ( const TArray<UObject*>& InObjects ) const override { return true; }
 	virtual void GetActions(const TArray<UObject*>& InObjects, struct FToolMenuSection& Section) override;
 
 	virtual void OpenAssetEditor( const TArray<UObject*>& InObjects, TSharedPtr<class IToolkitHost> EditWithinLevelEditor = TSharedPtr<IToolkitHost>() ) override;
diff --git a/Engine/Source/Developer/AssetTools/Private/AssetTypeActions/AssetTypeActions_StaticMesh.h b/Engine/Source/Developer/AssetTools/Private/AssetTypeActions/AssetTypeActions_StaticMesh.h
--- a/Engine/Source/Developer/AssetTools/Private/AssetTypeActions/AssetTypeActions_StaticMesh.h
+++ b/Engine/Source/Developer/AssetTools/Private/AssetTypeActions/AssetTypeActions_StaticMesh.h
@@ -17,7 +17,6 @@ public:
 	virtual FText GetName() const override { return NSLOCTEXT("AssetTypeActions", "AssetTypeActions_StaticMesh", "Static Mesh"); }
 	virtual FColor GetTypeColor() const override { return FColor(0, 255, 255); }
 	virtual UClass* GetSupportedClass() const override { return UStaticMesh::StaticClass(); }
-	virtual bool HasActions( const TArray<UObject*>& InObjects ) const override { return true; }
 	virtual void GetActions(const TArray<UObject*>& InObjects, FToolMenuSection& Section) override;
 	virtual void OpenAssetEditor( const TArray<UObject*>& InObjects, TSharedPtr<class IToolkitHost> EditWithinLevelEditor = TSharedPtr<IToolkitHost>() ) override;
 	virtual uint32 GetCategories() override { return EAssetTypeCategories::Basic; }
diff --git a/Engine/Source/Developer/AssetTools/Private/AssetTypeActions/AssetTypeActions_Texture.h b/Engine/Source/Developer/AssetTools/Private/AssetTypeActions/AssetTypeActions_Texture.h
--- a/Engine/Source/Developer/AssetTools/Private/AssetTypeActions/AssetTypeActions_Texture.h
+++ b/Engine/Source/Developer/AssetTools/Private/AssetTypeActions/AssetTypeActions_Texture.h
@@ -15,7 +15,6 @@ public:
 	virtual FText GetName() const override { return NSLOCTEXT("AssetTypeActions", "AssetTypeActions_Texture", "BaseTexture"); }
 	virtual FColor GetTypeColor() const override { return FColor(192,64,64); }
 	virtual UClass* GetSupportedClass() const override { return UTexture::StaticClass(); }
-	virtual bool HasActions ( const TArray<UObject*>& InObjects ) const override { return true; }
 	virtual void GetActions(const TArray<UObject*>& InObjects, struct FToolMenuSection& Section) override;
 	virtual void OpenAssetEditor( const TArray<UObject*>& InObjects, TSharedPtr<class IToolkitHost> EditWithinLevelEditor = TSharedPtr<IToolkitHost>() ) override;
 	virtual bool CanFilter() override { return false; }
diff --git a/Engine/Source/Developer/AssetTools/Private/AssetTypeActions/AssetTypeActions_Texture2D.h b/Engine/Source/Developer/AssetTools/Private/AssetTypeActions/AssetTypeActions_Texture2D.h
--- a/Engine/Source/Developer/AssetTools/Private/AssetTypeActions/AssetTypeActions_Texture2D.h
+++ b/Engine/Source/Developer/AssetTools/Private/AssetTypeActions/AssetTypeActions_Texture2D.h
@@ -14,7 +14,6 @@ public:
 	virtual FText GetName() const override { return NSLOCTEXT("AssetTypeActions", "AssetTypeActions_Texture2D", "Texture"); }
 	virtual FColor GetTypeColor() const override { return FColor(192,64,64); }
 	virtual UClass* GetSupportedClass() const override { return UTexture2D::StaticClass(); }
-	virtual bool HasActions ( const TArray<UObject*>& InObjects ) const override { return true; }
 	virtual void GetActions(const TArray<UObject*>& InObjects, struct FToolMenuSection& Section) override;
 	virtual bool CanFilter() override { return true; }
 	virtual uint32 GetCategories() override { return FAssetTypeActions_Texture::GetCategories() | EAssetTypeCategories::Basic; }
diff --git a/Engine/Source/Developer/AssetTools/Private/AssetTypeActions/AssetTypeActions_TextureRenderTarget.h b/Engine/Source/Developer/AssetTools/Private/AssetTypeActions/AssetTypeActions_TextureRenderTarget.h
--- a/Engine/Source/Developer/AssetTools/Private/AssetTypeActions/AssetTypeActions_TextureRenderTarget.h
+++ b/Engine/Source/Developer/AssetTools/Private/AssetTypeActions/AssetTypeActions_TextureRenderTarget.h
@@ -6,7 +6,6 @@
 #include "AssetTypeActions/AssetTypeActions_Texture.h"
 #include "Engine/TextureRenderTarget.h"
 
-
 class FAssetTypeActions_TextureRenderTarget : public FAssetTypeActions_Texture
 {
 public:
@@ -14,7 +13,6 @@ public:
 	virtual FText GetName() const override { return NSLOCTEXT("AssetTypeActions", "AssetTypeActions_TextureRenderTarget", "Texture Render Target"); }
 	virtual FColor GetTypeColor() const override { return FColor(128,64,64); }
 	virtual UClass* GetSupportedClass() const override { return UTextureRenderTarget::StaticClass(); }
-	virtual bool HasActions ( const TArray<UObject*>& InObjects ) const override { return true; }
 	virtual void GetActions(const TArray<UObject*>& InObjects, struct FToolMenuSection& Section) override;
 	virtual bool IsImportedAsset() const override { return false; }
 
diff --git a/Engine/Source/Developer/AssetTools/Private/AssetTypeActions/AssetTypeActions_World.h b/Engine/Source/Developer/AssetTools/Private/AssetTypeActions/AssetTypeActions_World.h
--- a/Engine/Source/Developer/AssetTools/Private/AssetTypeActions/AssetTypeActions_World.h
+++ b/Engine/Source/Developer/AssetTools/Private/AssetTypeActions/AssetTypeActions_World.h
@@ -14,7 +14,6 @@ public:
 	virtual FText GetName() const override { return NSLOCTEXT("AssetTypeActions", "AssetTypeActions_World", "Level"); }
 	virtual FColor GetTypeColor() const override { return FAppStyle::Get().GetColor("LevelEditor.AssetColor").ToFColor(true); }
 	virtual UClass* GetSupportedClass() const override { return UWorld::StaticClass(); }
-	virtual bool HasActions ( const TArray<UObject*>& InObjects ) const override { return false; }
 	virtual void OpenAssetEditor( const TArray<UObject*>& InObjects, TSharedPtr<class IToolkitHost> EditWithinLevelEditor = TSharedPtr<IToolkitHost>() ) override;
 	virtual uint32 GetCategories() override { return EAssetTypeCategories::Basic; }
 	virtual bool CanLocalize() const override { return false; }
diff --git a/Engine/Source/Developer/AssetTools/Public/AssetToolsModule.h b/Engine/Source/Developer/AssetTools/Public/AssetToolsModule.h
--- a/Engine/Source/Developer/AssetTools/Public/AssetToolsModule.h
+++ b/Engine/Source/Developer/AssetTools/Public/AssetToolsModule.h
@@ -7,6 +7,7 @@
 #include "Modules/ModuleManager.h"
 #include "IAssetTools.h"
 
+class IAssetTools;
 
 class FAssetToolsModule : public IModuleInterface
 {
@@ -36,5 +37,5 @@ public:
 
 private:
 	TWeakObjectPtr<class UAssetToolsImpl> AssetToolsPtr;
-	class FAssetToolsConsoleCommands* ConsoleCommands;
+	class FAssetToolsConsoleCommands* ConsoleCommands = nullptr;
 };
diff --git a/Engine/Source/Developer/AssetTools/Public/AssetTypeActions/AssetTypeActions_AnimationAsset.h b/Engine/Source/Developer/AssetTools/Public/AssetTypeActions/AssetTypeActions_AnimationAsset.h
--- a/Engine/Source/Developer/AssetTools/Public/AssetTypeActions/AssetTypeActions_AnimationAsset.h
+++ b/Engine/Source/Developer/AssetTools/Public/AssetTypeActions/AssetTypeActions_AnimationAsset.h
@@ -17,7 +17,6 @@ public:
 	virtual FText GetName() const override { return NSLOCTEXT("AssetTypeActions", "AssetTypeActions_AnimationAsset", "AnimationAsset"); }
 	virtual FColor GetTypeColor() const override { return FColor(80,123,72); }
 	virtual UClass* GetSupportedClass() const override { return UAnimationAsset::StaticClass(); }
-	virtual bool HasActions ( const TArray<UObject*>& InObjects ) const override { return true; }
 	virtual void GetActions(const TArray<UObject*>& InObjects, struct FToolMenuSection& Section) override;
 	virtual void OpenAssetEditor( const TArray<UObject*>& InObjects, TSharedPtr<class IToolkitHost> EditWithinLevelEditor = TSharedPtr<IToolkitHost>() ) override;
 	virtual bool CanFilter() override { return false; }
diff --git a/Engine/Source/Developer/AssetTools/Public/AssetTypeActions/AssetTypeActions_Blueprint.h b/Engine/Source/Developer/AssetTools/Public/AssetTypeActions/AssetTypeActions_Blueprint.h
--- a/Engine/Source/Developer/AssetTools/Public/AssetTypeActions/AssetTypeActions_Blueprint.h
+++ b/Engine/Source/Developer/AssetTools/Public/AssetTypeActions/AssetTypeActions_Blueprint.h
@@ -18,7 +18,6 @@ public:
 	virtual FText GetName() const override { return NSLOCTEXT("AssetTypeActions", "AssetTypeActions_Blueprint", "Blueprint Class"); }
 	virtual FColor GetTypeColor() const override { return FColor( 63, 126, 255 ); }
 	virtual UClass* GetSupportedClass() const override { return UBlueprint::StaticClass(); }
-	virtual bool HasActions ( const TArray<UObject*>& InObjects ) const override { return true; }
 	virtual void GetActions(const TArray<UObject*>& InObjects, FToolMenuSection& Section) override;
 	virtual void OpenAssetEditor(const TArray<UObject*>& InObjects, TSharedPtr<class IToolkitHost> EditWithinLevelEditor = TSharedPtr<IToolkitHost>()) override;
 	virtual bool CanMerge() const override;
diff --git a/Engine/Source/Developer/AssetTools/Public/AssetTypeActions/AssetTypeActions_BlueprintGeneratedClass.h b/Engine/Source/Developer/AssetTools/Public/AssetTypeActions/AssetTypeActions_BlueprintGeneratedClass.h
--- a/Engine/Source/Developer/AssetTools/Public/AssetTypeActions/AssetTypeActions_BlueprintGeneratedClass.h
+++ b/Engine/Source/Developer/AssetTools/Public/AssetTypeActions/AssetTypeActions_BlueprintGeneratedClass.h
@@ -18,7 +18,6 @@ public:
 	virtual FText GetName() const override { return NSLOCTEXT("AssetTypeActions", "AssetTypeActions_BlueprintGeneratedClass", "Compiled Blueprint Class"); }
 	virtual FColor GetTypeColor() const override { return FColor(133, 173, 255); }
 	virtual UClass* GetSupportedClass() const override;
-	virtual bool HasActions(const TArray<UObject*>& InObjects) const override { return true; }
 	virtual void GetActions(const TArray<UObject*>& InObjects, FToolMenuSection& Section) override;
 
 	// We don't want BPGCs to show up in CB/Ref viewer filters since the BP filters already include them
diff --git a/Engine/Source/Developer/AssetTools/Public/AssetTypeActions/AssetTypeActions_DataAsset.h b/Engine/Source/Developer/AssetTools/Public/AssetTypeActions/AssetTypeActions_DataAsset.h
--- a/Engine/Source/Developer/AssetTools/Public/AssetTypeActions/AssetTypeActions_DataAsset.h
+++ b/Engine/Source/Developer/AssetTools/Public/AssetTypeActions/AssetTypeActions_DataAsset.h
@@ -16,7 +16,6 @@ public:
 	virtual UClass* GetSupportedClass() const override { return UDataAsset::StaticClass(); }
 	virtual uint32 GetCategories() override { return EAssetTypeCategories::Misc; }
 	virtual FText GetDisplayNameFromAssetData(const FAssetData& AssetData) const override;
-	virtual bool HasActions(const TArray<UObject*>& InObjects) const override { return true; }
 	virtual void GetActions(const TArray<UObject*>& InObjects, struct FToolMenuSection& Section) override;
 	virtual void PerformAssetDiff(UObject* OldAsset, UObject* NewAsset, const FRevisionInfo& OldRevision, const FRevisionInfo& NewRevision) const override;
 
diff --git a/Engine/Source/Developer/AssetTools/Public/AssetTypeActions_Base.h b/Engine/Source/Developer/AssetTools/Public/AssetTypeActions_Base.h
--- a/Engine/Source/Developer/AssetTools/Public/AssetTypeActions_Base.h
+++ b/Engine/Source/Developer/AssetTools/Public/AssetTypeActions_Base.h
@@ -30,21 +30,6 @@ public:
 		return Object->GetName();
 	}
 
-	virtual bool HasActions( const TArray<UObject*>& InObjects ) const override
-	{
-		return false;
-	}
-
-	virtual void GetActions( const TArray<UObject*>& InObjects, FMenuBuilder& MenuBuilder ) override
-	{
-
-	}
-
-	virtual void GetActions(const TArray<UObject*>& InObjects, FToolMenuSection& Section) override
-	{
-
-	}
-
 	virtual void OpenAssetEditor( const TArray<UObject*>& InObjects, TSharedPtr<class IToolkitHost> EditWithinLevelEditor = TSharedPtr<IToolkitHost>() ) override
 	{
 		FSimpleAssetEditor::CreateEditor(EToolkitMode::Standalone, EditWithinLevelEditor, InObjects);
diff --git a/Engine/Source/Developer/AssetTools/Public/AssetTypeActions_CSVAssetBase.h b/Engine/Source/Developer/AssetTools/Public/AssetTypeActions_CSVAssetBase.h
--- a/Engine/Source/Developer/AssetTools/Public/AssetTypeActions_CSVAssetBase.h
+++ b/Engine/Source/Developer/AssetTools/Public/AssetTypeActions_CSVAssetBase.h
@@ -9,7 +9,6 @@ class FAssetTypeActions_CSVAssetBase : public FAssetTypeActions_Base
 {
 public:
 	virtual FColor GetTypeColor() const override { return FColor(62, 140, 35); }
-	virtual bool HasActions(const TArray<UObject*>& InObjects) const override { return true; }
 	virtual uint32 GetCategories() override { return EAssetTypeCategories::Misc; }
 	virtual bool IsImportedAsset() const override { return true; }
 
diff --git a/Engine/Source/Developer/AssetTools/Public/IAssetTools.h b/Engine/Source/Developer/AssetTools/Public/IAssetTools.h
--- a/Engine/Source/Developer/AssetTools/Public/IAssetTools.h
+++ b/Engine/Source/Developer/AssetTools/Public/IAssetTools.h
@@ -226,7 +226,7 @@ class IAssetTools
 	GENERATED_IINTERFACE_BODY()
 
 public:
-
+	ASSETTOOLS_API static IAssetTools& Get();
 
 	/** Registers an asset type actions object so it can provide information about and actions for asset types. */
 	virtual void RegisterAssetTypeActions(const TSharedRef<IAssetTypeActions>& NewActions) = 0;
diff --git a/Engine/Source/Developer/AssetTools/Public/IAssetTypeActions.h b/Engine/Source/Developer/AssetTools/Public/IAssetTypeActions.h
--- a/Engine/Source/Developer/AssetTools/Public/IAssetTypeActions.h
+++ b/Engine/Source/Developer/AssetTools/Public/IAssetTypeActions.h
@@ -45,6 +45,18 @@ struct FRevisionInfo
 	}
 };
 
+#ifndef DEPRECATE_ASSET_TYPE_ACTIONS_NEEDING_LOADED_OBJECTS
+	#define DEPRECATE_ASSET_TYPE_ACTIONS_NEEDING_LOADED_OBJECTS 0
+#endif
+
+#ifndef DEPRECATE_ASSET_TYPE_ACTIONS_CALLING_GETACTIONS
+    #if DEPRECATE_ASSET_TYPE_ACTIONS_NEEDING_LOADED_OBJECTS
+	    #define DEPRECATE_ASSET_TYPE_ACTIONS_CALLING_GETACTIONS DEPRECATE_ASSET_TYPE_ACTIONS_NEEDING_LOADED_OBJECTS
+	#else
+	    #define DEPRECATE_ASSET_TYPE_ACTIONS_CALLING_GETACTIONS 0
+	#endif
+#endif
+
 /** AssetTypeActions provide actions and other information about asset types */
 class IAssetTypeActions : public TSharedFromThis<IAssetTypeActions>
 {
@@ -60,15 +72,26 @@ public:
 
 	/** Returns the color associated with this type */
 	virtual FColor GetTypeColor() const = 0;
-
+	
 	/** Returns true if this class can supply actions for InObjects. */
-	virtual bool HasActions ( const TArray<UObject*>& InObjects ) const = 0;
+   	UE_DEPRECATED(5.2, "HasActions is no longer used all IAssetTypeActions are assumed to need to have them and GetActions is just called instead.  GetActions is always called unless DEPRECATE_ASSET_TYPE_ACTIONS_CALLING_GETACTIONS is true.")
+   	virtual bool HasActions( const TArray<UObject*>& InObjects ) const final { return true; }
+
+#if DEPRECATE_ASSET_TYPE_ACTIONS_NEEDING_LOADED_OBJECTS
+	/** Generates a menubuilder for the specified objects. */
+	UE_DEPRECATED(5.2, "Look at 'HOW TO ADD ASSET EXTENSION MENUS' in IAssetTypeActions.h")
+	virtual void GetActions( const TArray<UObject*>& InObjects, class FMenuBuilder& MenuBuilder ) final { }
 
+	/** Generates a menu section for the specified objects. */
+	UE_DEPRECATED(5.2, "Look at 'HOW TO ADD ASSET EXTENSION MENUS' in IAssetTypeActions.h")
+	virtual void GetActions(const TArray<UObject*>& InObjects, struct FToolMenuSection& Section) final { }
+#else
 	/** Generates a menubuilder for the specified objects. */
-	virtual void GetActions( const TArray<UObject*>& InObjects, class FMenuBuilder& MenuBuilder ) = 0;
+	virtual void GetActions( const TArray<UObject*>& InObjects, class FMenuBuilder& MenuBuilder ) { }
 
 	/** Generates a menu section for the specified objects. */
-	virtual void GetActions(const TArray<UObject*>& InObjects, struct FToolMenuSection& Section) = 0;
+	virtual void GetActions(const TArray<UObject*>& InObjects, struct FToolMenuSection& Section) { }
+#endif
 
 	/** Opens the asset editor for the specified objects. If EditWithinLevelEditor is valid, the world-centric editor will be used. */
 	virtual void OpenAssetEditor( const TArray<UObject*>& InObjects, TSharedPtr<IToolkitHost> EditWithinLevelEditor = TSharedPtr<IToolkitHost>() ) = 0;
@@ -156,3 +179,59 @@ public:
 	 */
 	virtual FTopLevelAssetPath GetClassPathName() const = 0;
 };
+
+/**
+HOW TO ADD ASSET EXTENSION MENUS
+
+void FMyAssetEditorModule::StartupModule()
+{
+	UToolMenus::RegisterStartupCallback(FSimpleMulticastDelegate::FDelegate::CreateRaw(this, &FSynthesisEditorModule::RegisterMenus));
+}
+
+void FSynthesisEditorModule::ShutdownModule()
+{
+	UToolMenus::UnRegisterStartupCallback(this);
+	UToolMenus::UnregisterOwner(this);
+}
+
+void FSynthesisEditorModule::RegisterMenus()
+{
+	FToolMenuOwnerScoped MenuOwner(this);
+	FAudioImpulseResponseExtension::RegisterMenus();
+}
+
+
+class FMyAssetExtension
+{
+public:
+	static void RegisterMenus();
+	static void ExecuteMenuAction(const struct FToolMenuContext& MenuContext);
+};
+
+void FMyAssetExtension::RegisterMenus()
+{
+	UToolMenu* Menu = UToolMenus::Get()->ExtendMenu("ContentBrowser.AssetContextMenu.MyAsset");
+	
+	FToolMenuSection& Section = Menu->FindOrAddSection("GetAssetActions");
+	Section.AddDynamicEntry("GetAssetActions_MyAsset", FNewToolMenuSectionDelegate::CreateLambda([](FToolMenuSection& InSection)
+	{
+		const TAttribute<FText> Label = LOCTEXT("MyAsset_DoTheThing", "Do the thing");
+		const TAttribute<FText> ToolTip = LOCTEXT("MyAsset_DoTheThingTooltip", "Does the thing.");
+		const FSlateIcon Icon = FSlateIcon(FAppStyle::GetAppStyleSetName(), "ClassIcon.MyAsset");
+		const FToolMenuExecuteAction UIAction = FToolMenuExecuteAction::CreateStatic(&FMyAssetExtension::ExecuteDoTheThing);
+
+		InSection.AddMenuEntry("MyAsset_DoTheThing", Label, ToolTip, Icon, UIAction);
+	}));
+}
+
+void FMyAssetExtension::ExecuteDoTheThing(const FToolMenuContext& MenuContext)
+{
+	if (const UContentBrowserAssetContextMenuContext* Context = UContentBrowserAssetContextMenuContext::FindContextWithAssets(MenuContext))
+	{
+		for (UMyAsset* MyAsset : Context->LoadSelectedObjects<UMyAsset>())
+		{
+			// Do the actual thing...
+		}
+	}
+}
+*/
\ No newline at end of file
diff --git a/Engine/Source/Developer/UnsavedAssetsTracker/Source/Private/UnsavedAssetsTracker.cpp b/Engine/Source/Developer/UnsavedAssetsTracker/Source/Private/UnsavedAssetsTracker.cpp
--- a/Engine/Source/Developer/UnsavedAssetsTracker/Source/Private/UnsavedAssetsTracker.cpp
+++ b/Engine/Source/Developer/UnsavedAssetsTracker/Source/Private/UnsavedAssetsTracker.cpp
@@ -19,6 +19,7 @@
 #include "Misc/PackageName.h"
 #include "LevelEditor.h"
 #include "ProfilingDebugging/CpuProfilerTrace.h"
+#include "IAssetTools.h"
 
 #define LOCTEXT_NAMESPACE "UnsavedAssetsTracker"
 
diff --git a/Engine/Source/Editor/AudioEditor/Private/AssetTypeActions/AssetTypeActions_DialogueWave.h b/Engine/Source/Editor/AudioEditor/Private/AssetTypeActions/AssetTypeActions_DialogueWave.h
--- a/Engine/Source/Editor/AudioEditor/Private/AssetTypeActions/AssetTypeActions_DialogueWave.h
+++ b/Engine/Source/Editor/AudioEditor/Private/AssetTypeActions/AssetTypeActions_DialogueWave.h
@@ -28,7 +28,6 @@ public:
 	virtual uint32 GetCategories() override { return EAssetTypeCategories::Sounds; }
 	virtual const TArray<FText>& GetSubMenus() const override;
 	virtual bool CanFilter() override { return true; }
-	virtual bool HasActions(const TArray<UObject*>& InObjects) const override { return true; }
 	virtual void GetActions(const TArray<UObject*>& InObjects, struct FToolMenuSection& Section) override;
 	virtual void OpenAssetEditor( const TArray<UObject*>& InObjects, TSharedPtr<class IToolkitHost> EditWithinLevelEditor = TSharedPtr<IToolkitHost>() ) override;
 
diff --git a/Engine/Source/Editor/AudioEditor/Private/AssetTypeActions/AssetTypeActions_SoundClass.h b/Engine/Source/Editor/AudioEditor/Private/AssetTypeActions/AssetTypeActions_SoundClass.h
--- a/Engine/Source/Editor/AudioEditor/Private/AssetTypeActions/AssetTypeActions_SoundClass.h
+++ b/Engine/Source/Editor/AudioEditor/Private/AssetTypeActions/AssetTypeActions_SoundClass.h
@@ -28,7 +28,6 @@ public:
 	virtual UClass* GetSupportedClass() const override;
 	virtual void OpenAssetEditor( const TArray<UObject*>& InObjects, TSharedPtr<class IToolkitHost> EditWithinLevelEditor = TSharedPtr<IToolkitHost>() ) override;
 	virtual uint32 GetCategories() override { return EAssetTypeCategories::Sounds; }
-	virtual bool HasActions(const TArray<UObject*>& InObjects) const override { return true; }
 	virtual void GetActions(const TArray<UObject*>& InObjects, FMenuBuilder& MenuBuilder) override;
 private:
 	
diff --git a/Engine/Source/Editor/AudioEditor/Public/AssetTypeActions/AssetTypeActions_SoundBase.h b/Engine/Source/Editor/AudioEditor/Public/AssetTypeActions/AssetTypeActions_SoundBase.h
--- a/Engine/Source/Editor/AudioEditor/Public/AssetTypeActions/AssetTypeActions_SoundBase.h
+++ b/Engine/Source/Editor/AudioEditor/Public/AssetTypeActions/AssetTypeActions_SoundBase.h
@@ -28,7 +28,6 @@ public:
 	virtual FText GetName() const override { return NSLOCTEXT("AssetTypeActions", "AssetTypeActions_SoundBase", "Sound Base"); }
 	virtual FColor GetTypeColor() const override { return FColor(97, 85, 212); }
 	virtual UClass* GetSupportedClass() const override;
-	virtual bool HasActions ( const TArray<UObject*>& InObjects ) const override { return true; }
 	virtual void GetActions(const TArray<UObject*>& InObjects, struct FToolMenuSection& Section) override;
 	virtual bool AssetsActivatedOverride(const TArray<UObject*>& InObjects, EAssetTypeActivationMethod::Type ActivationType) override;
 	virtual uint32 GetCategories() override { return EAssetTypeCategories::Sounds; }
diff --git a/Engine/Source/Editor/Blutility/Private/AssetTypeActions_EditorUtilityBlueprint.cpp b/Engine/Source/Editor/Blutility/Private/AssetTypeActions_EditorUtilityBlueprint.cpp
--- a/Engine/Source/Editor/Blutility/Private/AssetTypeActions_EditorUtilityBlueprint.cpp
+++ b/Engine/Source/Editor/Blutility/Private/AssetTypeActions_EditorUtilityBlueprint.cpp
@@ -51,11 +51,6 @@ UClass* FAssetTypeActions_EditorUtilityBlueprint::GetSupportedClass() const
 	return UEditorUtilityBlueprint::StaticClass();
 }
 
-bool FAssetTypeActions_EditorUtilityBlueprint::HasActions(const TArray<UObject*>& InObjects) const
-{
-	return true;
-}
-
 void FAssetTypeActions_EditorUtilityBlueprint::GetActions(const TArray<UObject*>& InObjects, FToolMenuSection& Section)
 {
 	TArray<TWeakObjectPtr<UEditorUtilityBlueprint>> Blueprints = GetTypedWeakObjectPtrs<UEditorUtilityBlueprint>(InObjects);
diff --git a/Engine/Source/Editor/Blutility/Private/AssetTypeActions_EditorUtilityBlueprint.h b/Engine/Source/Editor/Blutility/Private/AssetTypeActions_EditorUtilityBlueprint.h
--- a/Engine/Source/Editor/Blutility/Private/AssetTypeActions_EditorUtilityBlueprint.h
+++ b/Engine/Source/Editor/Blutility/Private/AssetTypeActions_EditorUtilityBlueprint.h
@@ -21,7 +21,6 @@ public:
 	virtual FText GetName() const override;
 	virtual FColor GetTypeColor() const override;
 	virtual UClass* GetSupportedClass() const override;
-	virtual bool HasActions(const TArray<UObject*>& InObjects) const override;
 	virtual void GetActions(const TArray<UObject*>& InObjects, struct FToolMenuSection& Section) override;
 	virtual uint32 GetCategories() override;
 	virtual bool CanLocalize() const override { return false; }
diff --git a/Engine/Source/Editor/Blutility/Private/AssetTypeActions_EditorUtilityWidgetBlueprint.cpp b/Engine/Source/Editor/Blutility/Private/AssetTypeActions_EditorUtilityWidgetBlueprint.cpp
--- a/Engine/Source/Editor/Blutility/Private/AssetTypeActions_EditorUtilityWidgetBlueprint.cpp
+++ b/Engine/Source/Editor/Blutility/Private/AssetTypeActions_EditorUtilityWidgetBlueprint.cpp
@@ -51,11 +51,6 @@ UClass* FAssetTypeActions_EditorUtilityWidgetBlueprint::GetSupportedClass() cons
 	return UEditorUtilityWidgetBlueprint::StaticClass();
 }
 
-bool FAssetTypeActions_EditorUtilityWidgetBlueprint::HasActions(const TArray<UObject*>& InObjects) const
-{
-	return true;
-}
-
 void FAssetTypeActions_EditorUtilityWidgetBlueprint::GetActions(const TArray<UObject*>& InObjects, FToolMenuSection& Section)
 {
 	auto Blueprints = GetTypedWeakObjectPtrs<UWidgetBlueprint>(InObjects);
diff --git a/Engine/Source/Editor/Blutility/Private/AssetTypeActions_EditorUtilityWidgetBlueprint.h b/Engine/Source/Editor/Blutility/Private/AssetTypeActions_EditorUtilityWidgetBlueprint.h
--- a/Engine/Source/Editor/Blutility/Private/AssetTypeActions_EditorUtilityWidgetBlueprint.h
+++ b/Engine/Source/Editor/Blutility/Private/AssetTypeActions_EditorUtilityWidgetBlueprint.h
@@ -23,7 +23,6 @@ public:
 	virtual FText GetName() const override;
 	virtual FColor GetTypeColor() const override;
 	virtual UClass* GetSupportedClass() const override;
-	virtual bool HasActions(const TArray<UObject*>& InObjects) const override;
 	virtual void GetActions(const TArray<UObject*>& InObjects, struct FToolMenuSection& Section) override;
 	virtual void OpenAssetEditor(const TArray<UObject*>& InObjects, TSharedPtr<class IToolkitHost> EditWithinLevelEditor = TSharedPtr<IToolkitHost>()) override;
 	virtual uint32 GetCategories() override;
diff --git a/Engine/Source/Editor/Blutility/Private/BlutilityContentBrowserExtensions.cpp b/Engine/Source/Editor/Blutility/Private/BlutilityContentBrowserExtensions.cpp
--- a/Engine/Source/Editor/Blutility/Private/BlutilityContentBrowserExtensions.cpp
+++ b/Engine/Source/Editor/Blutility/Private/BlutilityContentBrowserExtensions.cpp
@@ -115,13 +115,25 @@ public:
 			// Process asset based utilities
 			for (const FAssetData& UtilAsset : UtilAssets)
 			{
-				if (UEditorUtilityBlueprint* Blueprint = Cast<UEditorUtilityBlueprint>(UtilAsset.GetAsset()))
+				FString ParentClassName;
+				if (!UtilAsset.GetTagValue(FBlueprintTags::NativeParentClassPath, ParentClassName))
 				{
-					if (UClass* BPClass = Blueprint->GeneratedClass.Get())
+					UObject* Outer = nullptr;
+					ResolveName(Outer, ParentClassName, false, false);
+					const UClass* ParentClass = FindObject<UClass>(Outer, *ParentClassName);
+
+					// We only care about UEditorUtilityBlueprint's that are compiling subclasses of UAssetActionUtility
+					if (ParentClass && ParentClass->IsChildOf(UAssetActionUtility::StaticClass()))
 					{
-						if (UAssetActionUtility* DefaultObject = Cast<UAssetActionUtility>(BPClass->GetDefaultObject()))
+						if (const UEditorUtilityBlueprint* Blueprint = Cast<UEditorUtilityBlueprint>(UtilAsset.GetAsset()))
 						{
-							ProcessAssetAction(DefaultObject);
+							if (const UClass* BPClass = Blueprint->GeneratedClass.Get())
+							{
+								if (UAssetActionUtility* DefaultObject = Cast<UAssetActionUtility>(BPClass->GetDefaultObject()))
+								{
+									ProcessAssetAction(DefaultObject);
+								}
+							}
 						}
 					}
 				}
diff --git a/Engine/Source/Editor/ContentBrowser/Private/AssetContextMenu.cpp b/Engine/Source/Editor/ContentBrowser/Private/AssetContextMenu.cpp
--- a/Engine/Source/Editor/ContentBrowser/Private/AssetContextMenu.cpp
+++ b/Engine/Source/Editor/ContentBrowser/Private/AssetContextMenu.cpp
@@ -26,10 +26,12 @@
 #include "Framework/SlateDelegates.h"
 #include "HAL/IConsoleManager.h"
 #include "HAL/PlatformCrt.h"
+#include "HAL/PlatformStackWalk.h"
 #include "IAssetTools.h"
 #include "IAssetTypeActions.h"
 #include "ICollectionManager.h"
 #include "Internationalization/Internationalization.h"
+#include "Logging/MessageLog.h"
 #include "Misc/AssertionMacros.h"
 #include "Misc/Attribute.h"
 #include "Misc/NamePermissionList.h"
@@ -90,8 +92,61 @@ void FAssetContextMenu::BindCommands(TSharedPtr< FUICommandList >& Commands)
 		));
 }
 
+class FDontLoadAssetsDuringThisScope
+{
+public:
+	FDontLoadAssetsDuringThisScope()
+	{
+		FCoreUObjectDelegates::OnAssetLoaded.AddRaw(this, &FDontLoadAssetsDuringThisScope::OnAssetLoaded);
+	}
+
+	~FDontLoadAssetsDuringThisScope()
+	{
+		FCoreUObjectDelegates::OnAssetLoaded.RemoveAll(this);
+
+		FMessageLog EditorErrors("EditorErrors");
+		
+		for (const FLoadEvent& LoadEvent : ObjectsLoaded)
+		{
+			TSharedRef<FTokenizedMessage> ErrorMessage = EditorErrors.Error();
+			ErrorMessage->AddToken(FAssetNameToken::Create(LoadEvent.ObjectPath.ToString(),FText::FromString(LoadEvent.ObjectPath.ToString())));
+			ErrorMessage->AddToken(FTextToken::Create(LOCTEXT("DontLoadAssetsNow", "was loaded while creating the asset context menu.  Don't do this, it causes the operation to hitch.")));
+			
+			for (const FProgramCounterSymbolInfo& CallInfo : LoadEvent.Callstack)
+			{
+				ErrorMessage->AddToken(FTextToken::Create(FText::FromString(FString(CallInfo.FunctionName) + TEXT("\n"))));
+			}
+		}
+
+		EditorErrors.Notify(LOCTEXT("AssetsOpenedDuringContextMenuNotify", "Don't Load Assets Now!"));
+	}
+	
+private:
+	void OnAssetLoaded(UObject* InObject)
+	{
+		FLoadEvent Event;
+		Event.ObjectPath = FSoftObjectPath(InObject);
+		Event.Callstack = FPlatformStackWalk::GetStack(4, 10);
+		ObjectsLoaded.Add(MoveTemp(Event));
+	}
+
+private:
+	struct FLoadEvent
+	{
+		FSoftObjectPath ObjectPath;
+		TArray<FProgramCounterSymbolInfo> Callstack;
+	};
+	
+	TArray<FLoadEvent> ObjectsLoaded;
+};
+
 TSharedRef<SWidget> FAssetContextMenu::MakeContextMenu(TArrayView<const FContentBrowserItem> InSelectedItems, const FSourcesData& InSourcesData, TSharedPtr< FUICommandList > InCommandList)
 {
+#if DEPRECATE_ASSET_TYPE_ACTIONS_CALLING_GETACTIONS
+	// If we're not calling GetActions any more, then fire off errors when people open the context menu and it loads assets
+	FDontLoadAssetsDuringThisScope DontLoadScope;
+#endif
+	
 	SetSelectedItems(InSelectedItems);
 	SourcesData = InSourcesData;
 
@@ -137,9 +192,7 @@ TSharedRef<SWidget> FAssetContextMenu::MakeContextMenu(TArrayView<const FContent
 	static const FName BaseMenuName("ContentBrowser.AssetContextMenu");
 	static const FName ItemContextMenuName("ContentBrowser.ItemContextMenu");
 	RegisterContextMenu(BaseMenuName);
-
-	TArray<UObject*> SelectedObjects;
-
+	
 	// Create menu hierarchy based on class hierarchy
 	FName MenuName = BaseMenuName;
 	{
@@ -147,7 +200,7 @@ TSharedRef<SWidget> FAssetContextMenu::MakeContextMenu(TArrayView<const FContent
 		// need to keep it here for now to build the correct menu name and register the correct extenders
 
 		// Objects must be loaded for this operation... for now
-		TArray<FString> ObjectPaths;
+		TArray<FAssetData> SelectedAssets;
 		UContentBrowserDataSource* CommonDataSource = nullptr;
 		bool bKeepCheckingCommonDataSource = true;
 		for (const FContentBrowserItem& SelectedItem : SelectedItems)
@@ -174,25 +227,34 @@ TSharedRef<SWidget> FAssetContextMenu::MakeContextMenu(TArrayView<const FContent
 			FAssetData ItemAssetData;
 			if (SelectedItem.Legacy_TryGetAssetData(ItemAssetData))
 			{
-				ObjectPaths.Add(ItemAssetData.GetObjectPathString());
+				SelectedAssets.Add(MoveTemp(ItemAssetData));
 			}
 		}
 
 		FAssetToolsModule& AssetToolsModule = FModuleManager::LoadModuleChecked<FAssetToolsModule>("AssetTools");
 		const TSharedRef<FPathPermissionList>& WritableFolderFilter = AssetToolsModule.Get().GetWritableFolderPermissionList();
 
-		ContextObject->bCanBeModified = ObjectPaths.Num() == 0;
+		ContextObject->bCanBeModified = SelectedAssets.Num() == 0;
+
+		ContextObject->SelectedAssets.Reset();
 
-		ContextObject->SelectedObjects.Reset();
-		if (ContentBrowserUtils::LoadAssetsIfNeeded(ObjectPaths, SelectedObjects) && SelectedObjects.Num() > 0)
+		if (SelectedAssets.Num() > 0)
 		{
-			ContextObject->SelectedObjects.Append(SelectedObjects);
+			ContextObject->SelectedAssets.Append(SelectedAssets);
 
 			// Find common class for selected objects
-			UClass* CommonClass = SelectedObjects[0]->GetClass();
-			for (int32 ObjIdx = 1; ObjIdx < SelectedObjects.Num(); ++ObjIdx)
+			UClass* CommonClass = nullptr;
+			for (int32 ObjIdx = 0; ObjIdx < SelectedAssets.Num(); ++ObjIdx)
 			{
-				while (!SelectedObjects[ObjIdx]->IsA(CommonClass))
+				if (CommonClass == nullptr)
+				{
+					CommonClass = SelectedAssets[ObjIdx].GetClass();
+					continue;
+				}
+
+				// Update the CommonClass until we find a common shared class.
+				UClass* Class = SelectedAssets[ObjIdx].GetClass();
+				while (Class && !Class->IsChildOf(CommonClass))
 				{
 					CommonClass = CommonClass->GetSuperClass();
 				}
@@ -200,32 +262,26 @@ TSharedRef<SWidget> FAssetContextMenu::MakeContextMenu(TArrayView<const FContent
 			ContextObject->CommonClass = CommonClass;
 
 			ContextObject->bCanBeModified = true;
-			for (const UObject* SelectedObject : SelectedObjects)
+			for (const FAssetData& SelectedAsset : SelectedAssets)
 			{
-				if (SelectedObject)
+				if (WritableFolderFilter->HasFiltering() && !WritableFolderFilter->PassesStartsWithFilter(SelectedAsset.PackageName))
 				{
-					if (const UPackage* SelectedObjectPackage = SelectedObject->GetOutermost())
-					{
-						if (WritableFolderFilter->HasFiltering() && !WritableFolderFilter->PassesStartsWithFilter(SelectedObjectPackage->GetFName()))
-						{
-							ContextObject->bCanBeModified = false;
-							break;
-						}
+					ContextObject->bCanBeModified = false;
+					break;
+				}
 
-						if (SelectedObjectPackage->HasAnyPackageFlags(PKG_Cooked | PKG_FilterEditorOnly))
-						{
-							ContextObject->bCanBeModified = false;
-							break;
-						}
-					}
+				if (SelectedAsset.HasAnyPackageFlags(PKG_Cooked | PKG_FilterEditorOnly))
+				{
+					ContextObject->bCanBeModified = false;
+					break;
+				}
 
-					if (const UClass* AssetClass = SelectedObject->GetClass())
+				if (const UClass* AssetClass = SelectedAsset.GetClass())
+				{
+					if (AssetClass->IsChildOf<UClass>())
 					{
-						if (AssetClass->IsChildOf<UClass>())
-						{
-							ContextObject->bCanBeModified = false;
-							break;
-						}
+						ContextObject->bCanBeModified = false;
+						break;
 					}
 				}
 			}
@@ -235,13 +291,9 @@ TSharedRef<SWidget> FAssetContextMenu::MakeContextMenu(TArrayView<const FContent
 			RegisterMenuHierarchy(CommonClass);
 
 			// Find asset actions for common class
-			TSharedPtr<IAssetTypeActions> CommonAssetTypeActions = AssetToolsModule.Get().GetAssetTypeActionsForClass(ContextObject->CommonClass).Pin();
-			if (CommonAssetTypeActions.IsValid() && CommonAssetTypeActions->HasActions(SelectedObjects))
-			{
-				ContextObject->CommonAssetTypeActions = CommonAssetTypeActions;
-			}
+			ContextObject->CommonAssetTypeActions = AssetToolsModule.Get().GetAssetTypeActionsForClass(ContextObject->CommonClass).Pin();
 		}
-		else if (SelectedObjects.Num() == 0)
+		else if (SelectedAssets.Num() == 0)
 		{
 			if (CommonDataSource)
 			{
@@ -325,13 +377,17 @@ void FAssetContextMenu::RegisterContextMenu(const FName MenuName)
 		UToolMenu* Menu = ToolMenus->RegisterMenu(MenuName);
 		FToolMenuSection& Section = Menu->FindOrAddSection("GetAssetActions");
 
+#if  !DEPRECATE_ASSET_TYPE_ACTIONS_CALLING_GETACTIONS
+
 		// Note: Do  not use "GetActions" again when copying this code, otherwise "GetActions" menu entry will be overwritten
 		Section.AddDynamicEntry("GetActions", FNewToolMenuSectionDelegate::CreateLambda([](FToolMenuSection& InSection)
 		{
 			UContentBrowserAssetContextMenuContext* Context = InSection.FindContext<UContentBrowserAssetContextMenuContext>();
 			if (Context && Context->CommonAssetTypeActions.IsValid())
 			{
-				Context->CommonAssetTypeActions.Pin()->GetActions(Context->GetSelectedObjects(), InSection);
+				PRAGMA_DISABLE_DEPRECATION_WARNINGS
+				Context->CommonAssetTypeActions.Pin()->GetActions(Context->LoadSelectedObjects(), InSection);
+				PRAGMA_ENABLE_DEPRECATION_WARNINGS
 			}
 		}));
 
@@ -340,9 +396,13 @@ void FAssetContextMenu::RegisterContextMenu(const FName MenuName)
 			UContentBrowserAssetContextMenuContext* Context = InMenu->FindContext<UContentBrowserAssetContextMenuContext>();
 			if (Context && Context->CommonAssetTypeActions.IsValid())
 			{
-				Context->CommonAssetTypeActions.Pin()->GetActions(Context->GetSelectedObjects(), MenuBuilder);
+				PRAGMA_DISABLE_DEPRECATION_WARNINGS
+				Context->CommonAssetTypeActions.Pin()->GetActions(Context->LoadSelectedObjects(), MenuBuilder);
+				PRAGMA_ENABLE_DEPRECATION_WARNINGS
 			}
 		}));
+		
+#endif
 
 		Menu->AddDynamicSection("AddMenuOptions", FNewToolMenuDelegate::CreateLambda([](UToolMenu* InMenu)
 		{
@@ -611,7 +671,7 @@ bool FAssetContextMenu::AddAssetTypeMenuOptions(UToolMenu* Menu)
 	bool bAnyTypeOptions = false;
 
 	UContentBrowserAssetContextMenuContext* Context = Menu->FindContext<UContentBrowserAssetContextMenuContext>();
-	if (Context && Context->SelectedObjects.Num() > 0)
+	if (Context && Context->SelectedAssets.Num() > 0)
 	{
 		// Label "GetAssetActions" section
 		FToolMenuSection& Section = Menu->FindOrAddSection("GetAssetActions");
diff --git a/Engine/Source/Editor/ContentBrowser/Private/ContentBrowserMenuContexts.cpp b/Engine/Source/Editor/ContentBrowser/Private/ContentBrowserMenuContexts.cpp
--- a/Engine/Source/Editor/ContentBrowser/Private/ContentBrowserMenuContexts.cpp
+++ b/Engine/Source/Editor/ContentBrowser/Private/ContentBrowserMenuContexts.cpp
@@ -6,6 +6,7 @@
 #include "ContentBrowserDataSubsystem.h"
 #include "SContentBrowser.h"
 #include "UObject/UnrealNames.h"
+#include "ToolMenuSection.h"
 
 FName UContentBrowserToolbarMenuContext::GetCurrentPath() const
 {
diff --git a/Engine/Source/Editor/ContentBrowser/Private/ContentBrowserSingleton.cpp b/Engine/Source/Editor/ContentBrowser/Private/ContentBrowserSingleton.cpp
--- a/Engine/Source/Editor/ContentBrowser/Private/ContentBrowserSingleton.cpp
+++ b/Engine/Source/Editor/ContentBrowser/Private/ContentBrowserSingleton.cpp
@@ -79,6 +79,12 @@ struct FContentBrowserItem;
 
 static const FName ContentBrowserDrawerInstanceName("ContentBrowserDrawer");
 
+IContentBrowserSingleton& IContentBrowserSingleton::Get()
+{
+	FContentBrowserModule& ContentBrowserModule = FModuleManager::Get().LoadModuleChecked<FContentBrowserModule>("ContentBrowser");
+	return ContentBrowserModule.Get();
+}
+
 FContentBrowserSingleton::FContentBrowserSingleton()
 	: CollectionAssetRegistryBridge(MakeShared<FCollectionAssetRegistryBridge>())
 	, SettingsStringID(0)
diff --git a/Engine/Source/Editor/ContentBrowser/Private/SAssetPicker.cpp b/Engine/Source/Editor/ContentBrowser/Private/SAssetPicker.cpp
--- a/Engine/Source/Editor/ContentBrowser/Private/SAssetPicker.cpp
+++ b/Engine/Source/Editor/ContentBrowser/Private/SAssetPicker.cpp
@@ -355,7 +355,6 @@ void SAssetPicker::Construct( const FArguments& InArgs )
 		.CanShowDevelopersFolder( InArgs._AssetPickerConfig.bCanShowDevelopersFolder )
 		.ForceShowEngineContent( InArgs._AssetPickerConfig.bForceShowEngineContent )
 		.ForceShowPluginContent( InArgs._AssetPickerConfig.bForceShowPluginContent )
-		.PreloadAssetsForContextMenu( InArgs._AssetPickerConfig.bPreloadAssetsForContextMenu )
 		.HighlightedText( HighlightText )
 		.ThumbnailLabel( ThumbnailLabel )
 		.AssetShowWarningText( InArgs._AssetPickerConfig.AssetShowWarningText)
diff --git a/Engine/Source/Editor/ContentBrowser/Private/SAssetView.cpp b/Engine/Source/Editor/ContentBrowser/Private/SAssetView.cpp
--- a/Engine/Source/Editor/ContentBrowser/Private/SAssetView.cpp
+++ b/Engine/Source/Editor/ContentBrowser/Private/SAssetView.cpp
@@ -257,7 +257,6 @@ void SAssetView::Construct( const FArguments& InArgs )
 
 	bCanShowFavorites = InArgs._CanShowFavorites;
 	bCanDockCollections = InArgs._CanDockCollections;
-	bPreloadAssetsForContextMenu = InArgs._PreloadAssetsForContextMenu;
 
 	SelectionMode = InArgs._SelectionMode;
 
@@ -3642,27 +3641,6 @@ bool SAssetView::CanOpenContextMenu() const
 		return false;
 	}
 
-	if (bPreloadAssetsForContextMenu)
-	{
-		// Build a list of selected object paths
-		TArray<FString> ObjectPaths;
-		for (const TSharedPtr<FAssetViewItem>& Item : SelectedItems)
-		{
-			FAssetData ItemAssetData;
-			if (Item->GetItem().Legacy_TryGetAssetData(ItemAssetData))
-			{
-				ObjectPaths.Add(ItemAssetData.GetObjectPathString());
-			}
-		}
-
-		TArray<UObject*> LoadedObjects;
-		if (ObjectPaths.Num() > 0 && !ContentBrowserUtils::LoadAssetsIfNeeded(ObjectPaths, LoadedObjects, /*bAllowedToPrompt*/false))
-		{
-			// Do not show the context menu if the load failed
-			return false;
-		}
-	}
-
 	return true;
 }
 
diff --git a/Engine/Source/Editor/ContentBrowser/Public/ContentBrowserMenuContexts.h b/Engine/Source/Editor/ContentBrowser/Public/ContentBrowserMenuContexts.h
--- a/Engine/Source/Editor/ContentBrowser/Public/ContentBrowserMenuContexts.h
+++ b/Engine/Source/Editor/ContentBrowser/Public/ContentBrowserMenuContexts.h
@@ -7,14 +7,10 @@
 #include "ContentBrowserDelegates.h"
 #include "CoreMinimal.h"
 #include "HAL/Platform.h"
-#include "HAL/PlatformCrt.h"
 #include "Templates/SharedPointer.h"
 #include "UObject/NameTypes.h"
 #include "UObject/Object.h"
-#include "UObject/ObjectMacros.h"
 #include "UObject/ObjectPtr.h"
-#include "UObject/UObjectGlobals.h"
-#include "UObject/WeakObjectPtr.h"
 #include "UObject/WeakObjectPtrTemplates.h"
 
 #include "ContentBrowserMenuContexts.generated.h"
@@ -26,6 +22,8 @@ class SContentBrowser;
 class SFilterList;
 class UClass;
 struct FFrame;
+struct FToolMenuSection;
+struct FToolMenuContext;
 
 UCLASS()
 class CONTENTBROWSER_API UContentBrowserAssetContextMenuContext : public UObject
@@ -37,27 +35,69 @@ public:
 	TWeakPtr<FAssetContextMenu> AssetContextMenu;
 
 	TWeakPtr<IAssetTypeActions> CommonAssetTypeActions;
-	
-	UPROPERTY()
+
+	UE_DEPRECATED(5.2, "Use SelectedAssets now, this field will not contain any objects.")
 	TArray<TWeakObjectPtr<UObject>> SelectedObjects;
 
+	UPROPERTY()
+	TArray<FAssetData> SelectedAssets;
+
 	UPROPERTY()
 	TObjectPtr<UClass> CommonClass;
 
 	UPROPERTY()
 	bool bCanBeModified;
 
-	UFUNCTION(BlueprintCallable, Category="Tool Menus")
+	UFUNCTION(BlueprintCallable, Category="Tool Menus", meta=(DeprecatedFunction, DeprecationMessage = "GetSelectedObjects has been deprecated.  We no longer implictly load assets upon request.  If you can work without loading the assets, please use SelectedAssets.  Otherwise call LoadSelectedObjects"))
 	TArray<UObject*> GetSelectedObjects() const
 	{
-		TArray<UObject*> Result;
-		Result.Reserve(SelectedObjects.Num());
-		for (const TWeakObjectPtr<UObject>& Object : SelectedObjects)
+		return LoadSelectedObjects();
+	}
+
+	UFUNCTION(BlueprintCallable, Category="Tool Menus")
+	TArray<UObject*> LoadSelectedObjects() const
+	{
+		return LoadSelectedObjects<UObject>();
+	}
+
+	template<typename ExpectedAssetType>
+	TArray<ExpectedAssetType*> LoadSelectedObjects() const
+	{
+		return LoadSelectedObjectsIf<ExpectedAssetType>([](const FAssetData& AssetData){ return true; });
+	}
+
+	template<typename ExpectedAssetType>
+	TArray<ExpectedAssetType*> LoadSelectedObjectsIf(TFunctionRef<bool(const FAssetData& AssetData)> PredicateFilter) const
+	{
+		TArray<ExpectedAssetType*> Result;
+		Result.Reserve(SelectedAssets.Num());
+		for (const FAssetData& Asset : SelectedAssets)
 		{
-			Result.Add(Object.Get());
+			if (PredicateFilter(Asset))
+			{
+				if (UObject* AssetObject = Asset.GetAsset())
+				{
+					if (ExpectedAssetType* AssetObjectTyped = Cast<ExpectedAssetType>(AssetObject))
+					{
+						Result.Add(AssetObjectTyped);
+					}
+				}
+			}
 		}
 		return Result;
 	}
+
+	template<typename MenuOrSectionType>
+	static const UContentBrowserAssetContextMenuContext* FindContextWithAssets(const MenuOrSectionType& MenuOrSection)
+	{
+		const UContentBrowserAssetContextMenuContext* Context = MenuOrSection.template FindContext<UContentBrowserAssetContextMenuContext>();
+		if (!Context || Context->SelectedAssets.IsEmpty())
+		{
+			return nullptr;
+		}
+
+		return Context;
+	}
 };
 
 UCLASS()
diff --git a/Engine/Source/Editor/ContentBrowser/Public/IContentBrowserSingleton.h b/Engine/Source/Editor/ContentBrowser/Public/IContentBrowserSingleton.h
--- a/Engine/Source/Editor/ContentBrowser/Public/IContentBrowserSingleton.h
+++ b/Engine/Source/Editor/ContentBrowser/Public/IContentBrowserSingleton.h
@@ -286,9 +286,6 @@ struct FAssetPickerConfig
 	/** Indicates if plugin content should always be shown */
 	bool bForceShowPluginContent;
 
-	/** Indicates if the context menu is going to load the assets, and if so to preload before the context menu is shown, and warn about the pending load. */
-	bool bPreloadAssetsForContextMenu;
-
 	/** Indicates that we would like to build the filter UI with the Asset Picker */
 	bool bAddFilterUI;
 
@@ -319,7 +316,6 @@ struct FAssetPickerConfig
 		, bCanShowDevelopersFolder(true)
 		, bForceShowEngineContent(false)
 		, bForceShowPluginContent(false)
-		, bPreloadAssetsForContextMenu(true)
 		, bAddFilterUI(false)
 		, bShowPathInColumnView(false)
 		, bShowTypeInColumnView(true)
@@ -470,6 +466,8 @@ struct FSaveAssetDialogConfig : public FSharedAssetDialogConfig
 class IContentBrowserSingleton
 {
 public:
+	CONTENTBROWSER_API static IContentBrowserSingleton& Get();
+	
 	/** Virtual destructor */
 	virtual ~IContentBrowserSingleton() {}
 
diff --git a/Engine/Source/Editor/ContentBrowser/Public/SAssetView.h b/Engine/Source/Editor/ContentBrowser/Public/SAssetView.h
--- a/Engine/Source/Editor/ContentBrowser/Public/SAssetView.h
+++ b/Engine/Source/Editor/ContentBrowser/Public/SAssetView.h
@@ -108,7 +108,6 @@ public:
 		, _CanShowDevelopersFolder(false)
 		, _CanShowFavorites(false)
 		, _CanDockCollections(false)
-		, _PreloadAssetsForContextMenu(true)
 		, _SelectionMode( ESelectionMode::Multi )
 		, _AllowDragging(true)
 		, _AllowFocusOnSync(true)
@@ -229,9 +228,6 @@ public:
 		/** Indicates if the 'Dock Collections' option should be enabled or disabled */
 		SLATE_ARGUMENT(bool, CanDockCollections)
 
-		/** Indicates if the context menu is going to load the assets, and if so to preload before the context menu is shown, and warn about the pending load. */
-		SLATE_ARGUMENT( bool, PreloadAssetsForContextMenu )
-
 		/** The selection mode the asset view should use */
 		SLATE_ARGUMENT( ESelectionMode::Type, SelectionMode )
 
@@ -1047,9 +1043,6 @@ private:
 	/** Indicates if the 'Dock Collections' option should be enabled or disabled */
 	bool bCanDockCollections;
 
-	/** Indicates if the context menu is going to load the assets, and if so to preload before the context menu is shown, and warn about the pending load. */
-	bool bPreloadAssetsForContextMenu;
-
 	/** If true, it will show path column in the asset view */
 	bool bShowPathInColumnView;
 
diff --git a/Engine/Source/Editor/LevelInstanceEditor/Private/LevelInstanceEditorModule.cpp b/Engine/Source/Editor/LevelInstanceEditor/Private/LevelInstanceEditorModule.cpp
--- a/Engine/Source/Editor/LevelInstanceEditor/Private/LevelInstanceEditorModule.cpp
+++ b/Engine/Source/Editor/LevelInstanceEditor/Private/LevelInstanceEditorModule.cpp
@@ -558,13 +558,16 @@ namespace LevelInstanceMenuUtils
 		}		
 	}
 
-	void CreateBlueprintFromMenu(UToolMenu* Menu, UWorld* WorldAsset)
+	void CreateBlueprintFromMenu(UToolMenu* Menu, FAssetData WorldAsset)
 	{
 		FToolMenuSection& Section = CreateLevelSection(Menu);
 		FToolUIAction UIAction;
 		UIAction.ExecuteAction.BindLambda([WorldAsset](const FToolMenuContext& MenuContext)
 		{
-			CreateBlueprintFromWorld(WorldAsset);
+			if (UWorld* World = Cast<UWorld>(WorldAsset.GetAsset()))
+			{
+				CreateBlueprintFromWorld(World);
+			}
 		});
 
 		Section.AddMenuEntry(
@@ -634,16 +637,19 @@ namespace LevelInstanceMenuUtils
 		}
 	}
 
-	void AddPartitionedStreamingSupportFromMenu(UToolMenu* Menu, UWorld* WorldAsset)
+	void AddPartitionedStreamingSupportFromMenu(UToolMenu* Menu, FAssetData WorldAsset)
 	{
-		FName WorldAssetName = WorldAsset->GetPackage()->GetFName();
+		FName WorldAssetName = WorldAsset.PackageName;
 		if (!ULevel::GetIsLevelPartitionedFromPackage(WorldAssetName) || !ULevel::GetPartitionedLevelCanBeUsedByLevelInstanceFromPackage(WorldAssetName))
 		{
 			FToolMenuSection& Section = CreateLevelSection(Menu);
 			FToolUIAction UIAction;
 			UIAction.ExecuteAction.BindLambda([WorldAsset](const FToolMenuContext& MenuContext)
 			{
-				AddPartitionedStreamingSupportFromWorld(WorldAsset);
+				if (UWorld* World = Cast<UWorld>(WorldAsset.GetAsset()))
+				{
+					AddPartitionedStreamingSupportFromWorld(World);
+				}
 			});
 
 			Section.AddMenuEntry(
@@ -836,12 +842,13 @@ void FLevelInstanceEditorModule::ExtendContextMenu()
 			{
 				if (UContentBrowserAssetContextMenuContext* AssetMenuContext = ToolMenu->Context.FindContext<UContentBrowserAssetContextMenuContext>())
 				{
-					if (AssetMenuContext->SelectedObjects.Num() != 1)
+					if (AssetMenuContext->SelectedAssets.Num() != 1)
 					{
 						return;
 					}
-					// World is already loaded by the AssetContextMenu code
-					if (UWorld* WorldAsset = Cast<UWorld>(AssetMenuContext->SelectedObjects[0].Get()))
+
+					const FAssetData& WorldAsset = AssetMenuContext->SelectedAssets[0];
+					if (AssetMenuContext->SelectedAssets[0].IsInstanceOf<UWorld>())
 					{
 						LevelInstanceMenuUtils::CreateBlueprintFromMenu(ToolMenu, WorldAsset);
 						LevelInstanceMenuUtils::AddPartitionedStreamingSupportFromMenu(ToolMenu, WorldAsset);
diff --git a/Engine/Source/Editor/Sequencer/Private/SequencerModule.cpp b/Engine/Source/Editor/Sequencer/Private/SequencerModule.cpp
--- a/Engine/Source/Editor/Sequencer/Private/SequencerModule.cpp
+++ b/Engine/Source/Editor/Sequencer/Private/SequencerModule.cpp
@@ -375,40 +375,42 @@ public:
 				return;
 			}
 
-			ULevelSequence* LevelSequence = Context->SelectedObjects.Num() == 1 ? Cast<ULevelSequence>(Context->SelectedObjects[0]) : nullptr;
-			if (LevelSequence)
+			if (Context->SelectedAssets.Num() == 1 && Context->SelectedAssets[0].IsInstanceOf(ULevelSequence::StaticClass()))
 			{
-				// if this LevelSequence has associated maps, offer to load them
-				TArray<FString> AssociatedMaps = FSequencerUtilities::GetAssociatedMapPackages(LevelSequence);
-
-				if(AssociatedMaps.Num()>0)
+				if (ULevelSequence* LevelSequence = Cast<ULevelSequence>(Context->SelectedAssets[0].GetAsset()))
 				{
-					InSection.AddSubMenu(
-						"SequencerOpenMap_Label",
-						LOCTEXT("SequencerOpenMap_Label", "Open Map"),
-						LOCTEXT("SequencerOpenMap_Tooltip", "Open a map associated with this Level Sequence Asset"),
-						FNewMenuDelegate::CreateLambda(
-							[AssociatedMaps](FMenuBuilder& SubMenuBuilder)
-							{
-								for (const FString& AssociatedMap : AssociatedMaps)
+					// if this LevelSequence has associated maps, offer to load them
+					TArray<FString> AssociatedMaps = FSequencerUtilities::GetAssociatedMapPackages(LevelSequence);
+
+					if(AssociatedMaps.Num()>0)
+					{
+						InSection.AddSubMenu(
+							"SequencerOpenMap_Label",
+							LOCTEXT("SequencerOpenMap_Label", "Open Map"),
+							LOCTEXT("SequencerOpenMap_Tooltip", "Open a map associated with this Level Sequence Asset"),
+							FNewMenuDelegate::CreateLambda(
+								[AssociatedMaps](FMenuBuilder& SubMenuBuilder)
 								{
-									SubMenuBuilder.AddMenuEntry(
-										FText::FromString(FPaths::GetBaseFilename(AssociatedMap)),
-										FText(),
-										FSlateIcon(FAppStyle::GetAppStyleSetName(), "LevelEditor.Tabs.Levels"),
-										FExecuteAction::CreateLambda(
-											[AssociatedMap]
-											{
-												FEditorFileUtils::LoadMap(AssociatedMap);
-											}
-										)
-									);
+									for (const FString& AssociatedMap : AssociatedMaps)
+									{
+										SubMenuBuilder.AddMenuEntry(
+											FText::FromString(FPaths::GetBaseFilename(AssociatedMap)),
+											FText(),
+											FSlateIcon(FAppStyle::GetAppStyleSetName(), "LevelEditor.Tabs.Levels"),
+											FExecuteAction::CreateLambda(
+												[AssociatedMap]
+												{
+													FEditorFileUtils::LoadMap(AssociatedMap);
+												}
+											)
+										);
+									}
 								}
-							}
-						),
-						false,
-						FSlateIcon(FAppStyle::GetAppStyleSetName(), "LevelEditor.Tabs.Levels")
-					);
+							),
+							false,
+							FSlateIcon(FAppStyle::GetAppStyleSetName(), "LevelEditor.Tabs.Levels")
+						);
+					}
 				}
 			}
 		}));
diff --git a/Engine/Source/Editor/UMGEditor/Private/AssetTypeActions_SlateVectorArtData.h b/Engine/Source/Editor/UMGEditor/Private/AssetTypeActions_SlateVectorArtData.h
--- a/Engine/Source/Editor/UMGEditor/Private/AssetTypeActions_SlateVectorArtData.h
+++ b/Engine/Source/Editor/UMGEditor/Private/AssetTypeActions_SlateVectorArtData.h
@@ -14,5 +14,4 @@ public:
 	virtual FColor GetTypeColor() const override { return FColor(105, 165, 60); }
 	virtual UClass* GetSupportedClass() const override { return USlateVectorArtData::StaticClass(); }
 	virtual uint32 GetCategories() override { return EAssetTypeCategories::UI; }
-	virtual bool HasActions(const TArray<UObject*>& InObjects) const override { return false; }
 }; 
\ No newline at end of file
diff --git a/Engine/Source/Editor/UnrealEd/Private/Dialogs/SDeleteAssetsDialog.cpp b/Engine/Source/Editor/UnrealEd/Private/Dialogs/SDeleteAssetsDialog.cpp
--- a/Engine/Source/Editor/UnrealEd/Private/Dialogs/SDeleteAssetsDialog.cpp
+++ b/Engine/Source/Editor/UnrealEd/Private/Dialogs/SDeleteAssetsDialog.cpp
@@ -758,7 +758,6 @@ TSharedRef<SWidget> SDeleteAssetsDialog::MakeAssetViewForReferencerAssets()
 	AssetPickerConfig.bAllowNullSelection = false;
 	AssetPickerConfig.bShowBottomToolbar = false;
 	AssetPickerConfig.bAutohideSearchBar = true;
-	AssetPickerConfig.bPreloadAssetsForContextMenu = false;
 
 	AssetPickerConfig.AssetShowWarningText = TAttribute< FText >( this, &SDeleteAssetsDialog::GetReferencingAssetsEmptyText );
 
@@ -783,7 +782,6 @@ TSharedRef<SWidget> SDeleteAssetsDialog::MakeConsolidationAssetPicker()
 	AssetPickerConfig.bShowBottomToolbar = true;
 	AssetPickerConfig.bAllowDragging = false;
 	AssetPickerConfig.bCanShowClasses = false;
-	AssetPickerConfig.bPreloadAssetsForContextMenu = false;
 	AssetPickerConfig.SelectionMode = ESelectionMode::Single;
 
 	FContentBrowserModule& ContentBrowserModule = FModuleManager::Get().LoadModuleChecked<FContentBrowserModule>( TEXT( "ContentBrowser" ) );
diff --git a/Engine/Source/Editor/UnrealEd/Private/Dialogs/SPrivateAssetsDialog.cpp b/Engine/Source/Editor/UnrealEd/Private/Dialogs/SPrivateAssetsDialog.cpp
--- a/Engine/Source/Editor/UnrealEd/Private/Dialogs/SPrivateAssetsDialog.cpp
+++ b/Engine/Source/Editor/UnrealEd/Private/Dialogs/SPrivateAssetsDialog.cpp
@@ -312,7 +312,6 @@ TSharedRef<SWidget> SPrivateAssetsDialog::BuildAssetViewForReferencerAssets()
 	AssetPickerConfig.bAllowNullSelection = false;
 	AssetPickerConfig.bShowBottomToolbar = false;
 	AssetPickerConfig.bAutohideSearchBar = true;
-	AssetPickerConfig.bPreloadAssetsForContextMenu = false;
 
 	AssetPickerConfig.InitialAssetViewType = EAssetViewType::Column;
 	AssetPickerConfig.OnShouldFilterAsset = FOnShouldFilterAsset::CreateSP(this, &SPrivateAssetsDialog::OnShouldFilterAsset);
diff --git a/Engine/Source/Editor/VirtualTexturingEditor/Classes/RuntimeVirtualTextureAssetTypeActions.h b/Engine/Source/Editor/VirtualTexturingEditor/Classes/RuntimeVirtualTextureAssetTypeActions.h
--- a/Engine/Source/Editor/VirtualTexturingEditor/Classes/RuntimeVirtualTextureAssetTypeActions.h
+++ b/Engine/Source/Editor/VirtualTexturingEditor/Classes/RuntimeVirtualTextureAssetTypeActions.h
@@ -19,7 +19,6 @@ protected:
 	virtual FText GetName() const override;
 	virtual FColor GetTypeColor() const override;
 	virtual uint32 GetCategories() override;
-	virtual bool HasActions(TArray<UObject*> const& InObjects) const override { return true; }
 	virtual void GetActions(TArray<UObject*> const& InObjects, FMenuBuilder& MenuBuilder) override;
 	//~ End FAssetTypeActions_Base Interface.
 
