SourceControl: Generalized the 'checked out by other' text in the AssetBrowser tooltip to a 'source control text'.

In essense, this was already true because 'FPerforceSourceControlState::GetDisplayTooltip()' would return the 'IsConflicted' or '!IsCurrent' text before considering the 'IsCheckedOutByOther' text, so the tooltip would potentially already contain warning texts.

diff --git a/Engine/Source/Developer/SourceControl/Public/ISourceControlState.h b/Engine/Source/Developer/SourceControl/Public/ISourceControlState.h
--- a/Engine/Source/Developer/SourceControl/Public/ISourceControlState.h
+++ b/Engine/Source/Developer/SourceControl/Public/ISourceControlState.h
@@ -192,5 +192,16 @@ public:
 		}
 		return WarningText;
 	}
+
+	/** Gets the status message associated with this state, if any. This is a superset of the GetWarningText that also includes IsCheckedOut. */
+	virtual TOptional<FText> GetStatusText() const
+	{
+		TOptional<FText> StatusText = GetWarningText();
+		if (!StatusText.IsSet() && IsCheckedOut())
+		{
+			StatusText.Emplace(GetDisplayTooltip());
+		}
+		return StatusText;
+	}
 };
 
diff --git a/Engine/Source/Editor/ContentBrowser/Private/AssetViewWidgets.cpp b/Engine/Source/Editor/ContentBrowser/Private/AssetViewWidgets.cpp
--- a/Engine/Source/Editor/ContentBrowser/Private/AssetViewWidgets.cpp
+++ b/Engine/Source/Editor/ContentBrowser/Private/AssetViewWidgets.cpp
@@ -50,6 +50,7 @@
 #include "SThumbnailEditModeTools.h"
 #include "SlateOptMacros.h"
 #include "SlotBase.h"
+#include "SourceControlHelpers.h"
 #include "Styling/ISlateStyle.h"
 #include "Styling/SlateBrush.h"
 #include "Styling/SlateTypes.h"
@@ -784,7 +785,7 @@ TSharedRef<SWidget> SAssetViewItem::CreateToolTipWidget() const
 		// TODO: Remove this special caseness so that folders can also have visible attributes
 		if(AssetItem->IsFile())
 		{
-			// The tooltip contains the name, class, path, and asset registry tags
+			// The tooltip contains the name, class, path, asset registry tags and source control status
 			const FText NameText = GetNameText();
 			const FText ClassText = FText::Format(LOCTEXT("ClassName", "({0})"), GetAssetClassText());
 			FText PublicStateText = LOCTEXT("PublicAssetState", "Public");
@@ -933,8 +934,8 @@ TSharedRef<SWidget> SAssetViewItem::CreateToolTipWidget() const
 						.AutoHeight()
 						[
 							SNew(STextBlock)
-							.Visibility(this, &SAssetViewItem::GetCheckedOutByOtherTextVisibility)
-							.Text(this, &SAssetViewItem::GetCheckedOutByOtherText)
+							.Visibility(this, &SAssetViewItem::GetSourceControlTextVisibility)
+							.Text(this, &SAssetViewItem::GetSourceControlText)
 							.ColorAndOpacity(FLinearColor(0.1f, 0.5f, 1.f, 1.f))
 						]
 
@@ -1122,12 +1123,12 @@ PRAGMA_ENABLE_DEPRECATION_WARNINGS
 	}
 }
 
-EVisibility SAssetViewItem::GetCheckedOutByOtherTextVisibility() const
+EVisibility SAssetViewItem::GetSourceControlTextVisibility() const
 {
-	return GetCheckedOutByOtherText().IsEmpty() ? EVisibility::Collapsed : EVisibility::Visible;
+	return GetSourceControlText().IsEmpty() ? EVisibility::Collapsed : EVisibility::Visible;
 }
 
-FText SAssetViewItem::GetCheckedOutByOtherText() const
+FText SAssetViewItem::GetSourceControlText() const
 {
 	if (AssetItem && AssetItem->IsFile() && !AssetItem->IsTemporary() && ISourceControlModule::Get().IsEnabled())
 	{
@@ -1135,9 +1136,9 @@ FText SAssetViewItem::GetCheckedOutByOtherText() const
 		if (AssetItem->GetItem().GetItemPhysicalPath(AssetFilename))
 		{
 			FSourceControlStatePtr SourceControlState = ISourceControlModule::Get().GetProvider().GetState(AssetFilename, EStateCacheUsage::Use);
-			if (SourceControlState && (SourceControlState->IsCheckedOutOther() || SourceControlState->IsCheckedOutOrModifiedInOtherBranch()))
+			if (SourceControlState)
 			{
-				return SourceControlState->GetDisplayTooltip();
+				return SourceControlState->GetStatusText().Get(FText::GetEmpty());
 			}
 		}
 	}
diff --git a/Engine/Source/Editor/ContentBrowser/Public/AssetViewWidgets.h b/Engine/Source/Editor/ContentBrowser/Public/AssetViewWidgets.h
--- a/Engine/Source/Editor/ContentBrowser/Public/AssetViewWidgets.h
+++ b/Engine/Source/Editor/ContentBrowser/Public/AssetViewWidgets.h
@@ -231,11 +231,11 @@ protected:
 	/** Creates a tooltip widget for this item */
 	TSharedRef<SWidget> CreateToolTipWidget() const;
 
-	/** Gets the visibility of the checked out by other text block in the tooltip */
-	EVisibility GetCheckedOutByOtherTextVisibility() const;
+	/** Gets the visibility of the source control text block in the tooltip */
+	EVisibility GetSourceControlTextVisibility() const;
 
-	/** Gets the text for the checked out by other text block in the tooltip */
-	FText GetCheckedOutByOtherText() const;
+	/** Gets the text for the source control text block in the tooltip */
+	FText GetSourceControlText() const;
 
 	/** Helper function for CreateToolTipWidget. Gets the user description for the asset, if it exists. */
 	FText GetAssetUserDescription() const;
