Add SelectedAssets to content browser in a stream. Ignore these changes during merge conflict.
#rnx

diff --git a/Engine/Source/Editor/ContentBrowser/Private/ContentBrowserMenuContexts.cpp b/Engine/Source/Editor/ContentBrowser/Private/ContentBrowserMenuContexts.cpp
--- a/Engine/Source/Editor/ContentBrowser/Private/ContentBrowserMenuContexts.cpp
+++ b/Engine/Source/Editor/ContentBrowser/Private/ContentBrowserMenuContexts.cpp
@@ -4,10 +4,26 @@
 
 #include "Containers/UnrealString.h"
 #include "ContentBrowserDataSubsystem.h"
+#include "SAssetView.h"
 #include "SContentBrowser.h"
 #include "UObject/UnrealNames.h"
 #include "ToolMenuSection.h"
 
+const TArray<FAssetData>& UContentBrowserAssetViewContextMenuContext::GetSelectedAssets() const
+{
+	if (!bHasInitSelectedAssets)
+	{
+		UContentBrowserAssetViewContextMenuContext* MutableThis = const_cast<UContentBrowserAssetViewContextMenuContext*>(this);
+		MutableThis->bHasInitSelectedAssets = true;
+		if (AssetView.IsValid())
+		{
+			MutableThis->SelectedAssets = AssetView.Pin()->GetSelectedAssets();
+		}
+	}
+
+	return SelectedAssets;
+}
+
 FName UContentBrowserToolbarMenuContext::GetCurrentPath() const
 {
 	if (TSharedPtr<SContentBrowser> Browser = ContentBrowser.Pin())
@@ -26,4 +42,4 @@ bool UContentBrowserToolbarMenuContext::CanWriteToCurrentPath() const
 	}
 
 	return false;
-}
+}
\ No newline at end of file
diff --git a/Engine/Source/Editor/ContentBrowser/Public/ContentBrowserMenuContexts.h b/Engine/Source/Editor/ContentBrowser/Public/ContentBrowserMenuContexts.h
--- a/Engine/Source/Editor/ContentBrowser/Public/ContentBrowserMenuContexts.h
+++ b/Engine/Source/Editor/ContentBrowser/Public/ContentBrowserMenuContexts.h
@@ -37,6 +37,7 @@ public:
 	TWeakPtr<IAssetTypeActions> CommonAssetTypeActions;
 
 	UE_DEPRECATED(5.1, "Use SelectedAssets now, this field will not contain any objects.")
+	UPROPERTY()
 	TArray<TWeakObjectPtr<UObject>> SelectedObjects;
 
 	UPROPERTY()
@@ -108,6 +109,13 @@ class CONTENTBROWSER_API UContentBrowserAssetViewContextMenuContext : public UOb
 public:
 	TWeakPtr<SContentBrowser> OwningContentBrowser;
 	TWeakPtr<SAssetView> AssetView;
+
+	const TArray<FAssetData>& GetSelectedAssets() const;
+
+private:
+
+	TArray<FAssetData> SelectedAssets;
+	bool bHasInitSelectedAssets = false;
 };
 
 UCLASS()
