Post merge fixup
#rnx

diff --git a/Engine/Source/Editor/ContentBrowser/Public/ContentBrowserMenuContexts.h b/Engine/Source/Editor/ContentBrowser/Public/ContentBrowserMenuContexts.h
--- a/Engine/Source/Editor/ContentBrowser/Public/ContentBrowserMenuContexts.h
+++ b/Engine/Source/Editor/ContentBrowser/Public/ContentBrowserMenuContexts.h
@@ -37,7 +37,6 @@ public:
 	TWeakPtr<IAssetTypeActions> CommonAssetTypeActions;
 
 	UE_DEPRECATED(5.1, "Use SelectedAssets now, this field will not contain any objects.")
-	UPROPERTY()
 	TArray<TWeakObjectPtr<UObject>> SelectedObjects;
 
 	UPROPERTY()
