Remove a duplicate function
#rnx

diff --git a/Engine/Source/Editor/ContentBrowser/Private/ContentBrowserMenuContexts.cpp b/Engine/Source/Editor/ContentBrowser/Private/ContentBrowserMenuContexts.cpp
--- a/Engine/Source/Editor/ContentBrowser/Private/ContentBrowserMenuContexts.cpp
+++ b/Engine/Source/Editor/ContentBrowser/Private/ContentBrowserMenuContexts.cpp
@@ -4,26 +4,10 @@
 
 #include "Containers/UnrealString.h"
 #include "ContentBrowserDataSubsystem.h"
-#include "SAssetView.h"
 #include "SContentBrowser.h"
 #include "UObject/UnrealNames.h"
 #include "ToolMenuSection.h"
 
-const TArray<FAssetData>& UContentBrowserAssetViewContextMenuContext::GetSelectedAssets() const
-{
-	if (!bHasInitSelectedAssets)
-	{
-		UContentBrowserAssetViewContextMenuContext* MutableThis = const_cast<UContentBrowserAssetViewContextMenuContext*>(this);
-		MutableThis->bHasInitSelectedAssets = true;
-		if (AssetView.IsValid())
-		{
-			MutableThis->SelectedAssets = AssetView.Pin()->GetSelectedAssets();
-		}
-	}
-
-	return SelectedAssets;
-}
-
 FName UContentBrowserToolbarMenuContext::GetCurrentPath() const
 {
 	if (TSharedPtr<SContentBrowser> Browser = ContentBrowser.Pin())
diff --git a/Engine/Source/Editor/ContentBrowser/Public/ContentBrowserMenuContexts.h b/Engine/Source/Editor/ContentBrowser/Public/ContentBrowserMenuContexts.h
--- a/Engine/Source/Editor/ContentBrowser/Public/ContentBrowserMenuContexts.h
+++ b/Engine/Source/Editor/ContentBrowser/Public/ContentBrowserMenuContexts.h
@@ -108,13 +108,6 @@ class CONTENTBROWSER_API UContentBrowserAssetViewContextMenuContext : public UOb
 public:
 	TWeakPtr<SContentBrowser> OwningContentBrowser;
 	TWeakPtr<SAssetView> AssetView;
-
-	const TArray<FAssetData>& GetSelectedAssets() const;
-
-private:
-
-	TArray<FAssetData> SelectedAssets;
-	bool bHasInitSelectedAssets = false;
 };
 
 UCLASS()
