Expose SelectedPackagePaths on ContentBrowserFolderContext
#rnx

diff --git a/Engine/Source/Editor/ContentBrowser/Private/SContentBrowser.cpp b/Engine/Source/Editor/ContentBrowser/Private/SContentBrowser.cpp
--- a/Engine/Source/Editor/ContentBrowser/Private/SContentBrowser.cpp
+++ b/Engine/Source/Editor/ContentBrowser/Private/SContentBrowser.cpp
@@ -4091,6 +4091,7 @@ TSharedPtr<SWidget> SContentBrowser::GetItemContextMenu(TArrayView<const FConten
 			AppendNewMenuContextObjects(EContentBrowserDataMenuContext_AddNewMenuDomain::PathView, SelectedVirtualPaths, MenuContext, nullptr, Context->bCanBeModified);
 		}
 
+		Context->SelectedPackagePaths = MoveTemp(SelectedPackagePaths);
 		return UToolMenus::Get()->GenerateWidget("ContentBrowser.FolderContextMenu", MenuContext);
 	}
 	else if (SelectedItems.Num() > 0)
diff --git a/Engine/Source/Editor/ContentBrowser/Public/ContentBrowserMenuContexts.h b/Engine/Source/Editor/ContentBrowser/Public/ContentBrowserMenuContexts.h
--- a/Engine/Source/Editor/ContentBrowser/Public/ContentBrowserMenuContexts.h
+++ b/Engine/Source/Editor/ContentBrowser/Public/ContentBrowserMenuContexts.h
@@ -127,19 +127,24 @@ class CONTENTBROWSER_API UContentBrowserFolderContext : public UContentBrowserMe
 
 public:
 
-	UPROPERTY()
+	UPROPERTY(BlueprintReadOnly, Category = "ContentBrowser")
 	bool bCanBeModified;
 
-	UPROPERTY()
+	UPROPERTY(BlueprintReadOnly, Category = "ContentBrowser")
 	bool bNoFolderOnDisk;
 
-	UPROPERTY()
+	UPROPERTY(BlueprintReadOnly, Category = "ContentBrowser")
 	int32 NumAssetPaths;
 
-	UPROPERTY()
+	UPROPERTY(BlueprintReadOnly, Category = "ContentBrowser")
 	int32 NumClassPaths;
 
+	UPROPERTY(BlueprintReadOnly, Category = "ContentBrowser")
+	TArray<FString> SelectedPackagePaths;
+
 	FOnCreateNewFolder OnCreateNewFolder;
+
+	const TArray<FString>& GetSelectedPackagePaths() const { return SelectedPackagePaths; }
 };
 
 UCLASS()
