IContentBrowserSingleton delegate to control whether the public/private state can be changed on a given asset

diff --git a/Engine/Source/Editor/ContentBrowser/Private/AssetContextMenu.cpp b/Engine/Source/Editor/ContentBrowser/Private/AssetContextMenu.cpp
--- a/Engine/Source/Editor/ContentBrowser/Private/AssetContextMenu.cpp
+++ b/Engine/Source/Editor/ContentBrowser/Private/AssetContextMenu.cpp
@@ -1260,9 +1260,20 @@ void FAssetContextMenu::CacheCanExecuteVars()
 		{
 			bCanExecuteFindInExplorer = true;
 		}
+		
+		bool bCanChangeAssetPublicState = SelectedItem.CanEdit();
+		if (bCanChangeAssetPublicState)
+		{
+			FNameBuilder ItemVirtualPath(SelectedItem.GetVirtualPath());
+			bCanChangeAssetPublicState = FContentBrowserSingleton::Get().IsShowingPrivateContent(FPathViews::GetPath(ItemVirtualPath));
+		}
+		if (bCanChangeAssetPublicState)
+		{
+			FNameBuilder ItemInternalPath(SelectedItem.GetInternalPath());
+			bCanChangeAssetPublicState = FContentBrowserSingleton::Get().CanChangeAssetPublicState(ItemInternalPath);
+		}
 
-		if(SelectedItem.CanEdit() &&
-			FContentBrowserSingleton::Get().IsShowingPrivateContent(FPathViews::GetPath(FNameBuilder(SelectedItem.GetVirtualPath()))))
+		if (bCanChangeAssetPublicState)
 		{
 			bCanExecutePublicAssetToggle = true;
 		}
diff --git a/Engine/Source/Editor/ContentBrowser/Private/ContentBrowserSingleton.cpp b/Engine/Source/Editor/ContentBrowser/Private/ContentBrowserSingleton.cpp
--- a/Engine/Source/Editor/ContentBrowser/Private/ContentBrowserSingleton.cpp
+++ b/Engine/Source/Editor/ContentBrowser/Private/ContentBrowserSingleton.cpp
@@ -1127,6 +1127,15 @@ void FContentBrowserSingleton::RebuildPrivateContentStateCache()
 	}
 }
 
+bool FContentBrowserSingleton::CanChangeAssetPublicState(FStringView AssetPath)
+{
+	if (CanChangeAssetPublicStateDelegate.IsBound())
+	{
+		return CanChangeAssetPublicStateDelegate.Execute(AssetPath);
+	}
+	return true;
+}
+
 bool FContentBrowserSingleton::IsFolderShowPrivateContentToggleable(const FStringView VirtualFolderPath)
 {
 	if (IsFolderShowPrivateContentToggleableDelegate.IsBound())
@@ -1147,9 +1156,19 @@ void FContentBrowserSingleton::SetPrivateContentPermissionListDirty()
 	ShowPrivateContentState.CachedVirtualPaths.Reset();
 }
 
+void FContentBrowserSingleton::RegisterCanChangeAssetPublicStateDelegate(FCanChangeAssetPublicStateDelegate InCanChangeAssetPublicStateDelegate)
+{
+	CanChangeAssetPublicStateDelegate = MoveTemp(InCanChangeAssetPublicStateDelegate);
+}
+
+void FContentBrowserSingleton::UnregisterCanChangeAssetPublicStateDelegate()
+{
+	CanChangeAssetPublicStateDelegate = FCanChangeAssetPublicStateDelegate();
+}
+
 void FContentBrowserSingleton::RegisterIsFolderShowPrivateContentToggleableDelegate(FIsFolderShowPrivateContentToggleableDelegate InIsFolderShowPrivateContentToggleableDelegate)
 {
-	IsFolderShowPrivateContentToggleableDelegate = InIsFolderShowPrivateContentToggleableDelegate;
+	IsFolderShowPrivateContentToggleableDelegate = MoveTemp(InIsFolderShowPrivateContentToggleableDelegate);
 }
 
 void FContentBrowserSingleton::UnregisterIsFolderShowPrivateContentToggleableDelegate()
diff --git a/Engine/Source/Editor/ContentBrowser/Private/ContentBrowserSingleton.h b/Engine/Source/Editor/ContentBrowser/Private/ContentBrowserSingleton.h
--- a/Engine/Source/Editor/ContentBrowser/Private/ContentBrowserSingleton.h
+++ b/Engine/Source/Editor/ContentBrowser/Private/ContentBrowserSingleton.h
@@ -108,10 +108,13 @@ PRAGMA_ENABLE_DEPRECATION_WARNINGS
 	virtual void ExecuteRename(TSharedPtr<SWidget> PickerWidget) override;
 	virtual void ExecuteAddFolder(TSharedPtr<SWidget> PathPickerWidget) override;
 	virtual void RefreshPathView(TSharedPtr<SWidget> PathPickerWidget) override;
+	virtual bool CanChangeAssetPublicState(FStringView AssetPath) override;
 	virtual bool IsShowingPrivateContent(const FStringView VirtualFolderPath) override;
 	virtual bool IsFolderShowPrivateContentToggleable(const FStringView VirtualFolderPath) override;
 	virtual const TSharedPtr<FPathPermissionList>& GetShowPrivateContentPermissionList() override;
 	virtual void SetPrivateContentPermissionListDirty() override;
+	virtual void RegisterCanChangeAssetPublicStateDelegate(FCanChangeAssetPublicStateDelegate InCanChangeAssetPublicStateDelegate) override;
+	virtual void UnregisterCanChangeAssetPublicStateDelegate() override;
 	virtual void RegisterIsFolderShowPrivateContentToggleableDelegate(FIsFolderShowPrivateContentToggleableDelegate InIsFolderShowPrivateContentToggleableDelegate) override;
 	virtual void UnregisterIsFolderShowPrivateContentToggleableDelegate() override;
 
@@ -184,6 +187,8 @@ public:
 
 private:
 
+	FCanChangeAssetPublicStateDelegate CanChangeAssetPublicStateDelegate;
+
 	FIsFolderShowPrivateContentToggleableDelegate IsFolderShowPrivateContentToggleableDelegate;
 
 	TArray<TWeakPtr<SContentBrowser>> AllContentBrowsers;
diff --git a/Engine/Source/Editor/ContentBrowser/Public/ContentBrowserDelegates.h b/Engine/Source/Editor/ContentBrowser/Public/ContentBrowserDelegates.h
--- a/Engine/Source/Editor/ContentBrowser/Public/ContentBrowserDelegates.h
+++ b/Engine/Source/Editor/ContentBrowser/Public/ContentBrowserDelegates.h
@@ -125,6 +125,9 @@ DECLARE_DELEGATE_OneParam(FOnAssetsChosenForOpen, const TArray<FAssetData>& /*Se
 /** Called from the Asset Dialog when an asset name is chosen in non-modal Save dialogs */
 DECLARE_DELEGATE_OneParam(FOnObjectPathChosenForSave, const FString& /*ObjectPath*/);
 
+/** Called to determine whether the public/private state of the specified asset can be changed */
+DECLARE_DELEGATE_RetVal_OneParam(bool, FCanChangeAssetPublicStateDelegate, const FStringView /*AssetPath*/);
+
 /** Called when custom behavior is needed for allowing folders to toggle their private content edit state */
 DECLARE_DELEGATE_RetVal_OneParam(bool, FIsFolderShowPrivateContentToggleableDelegate, const FStringView /*VirtualFolderPath*/);
 
diff --git a/Engine/Source/Editor/ContentBrowser/Public/IContentBrowserSingleton.h b/Engine/Source/Editor/ContentBrowser/Public/IContentBrowserSingleton.h
--- a/Engine/Source/Editor/ContentBrowser/Public/IContentBrowserSingleton.h
+++ b/Engine/Source/Editor/ContentBrowser/Public/IContentBrowserSingleton.h
@@ -684,6 +684,9 @@ public:
 	/** Returns InPath if can be written to, otherwise picks a default path that can be written to */
 	virtual FContentBrowserItemPath GetInitialPathToSaveAsset(const FContentBrowserItemPath& InPath) = 0;
 
+	/** Returns true if the public/private state of the specified asset can be changed */
+	virtual bool CanChangeAssetPublicState(FStringView AssetPath) = 0;
+
 	/** Returns true if FolderPath is a private content edit folder */
 	virtual bool IsShowingPrivateContent(const FStringView VirtualFolderPath) = 0;
 
@@ -696,6 +699,12 @@ public:
 	/** Declares the Private Content Permission List dirty */
 	virtual void SetPrivateContentPermissionListDirty() = 0;
 
+	/** Registers the delegate called to determine whether the public/private state of the specified asset can be changed */
+	virtual void RegisterCanChangeAssetPublicStateDelegate(FCanChangeAssetPublicStateDelegate InCanChangeAssetPublicStateDelegate) = 0;
+
+	/** Unregisters the delegate called to determine whether the public/private state of the specified asset can be changed */
+	virtual void UnregisterCanChangeAssetPublicStateDelegate() = 0;
+
 	/** Registers the delegate for custom handling of if a Folder allows private content edits */
 	virtual void RegisterIsFolderShowPrivateContentToggleableDelegate(FIsFolderShowPrivateContentToggleableDelegate InIsFolderShowPrivateContentToggleableDelegate) = 0;
 
