Editor - Fixing a Context Menu regression.  The old code did not in fact load everything always, there was a carveout for WP levels, to prevent them from loading by using AssetViewUtils::LoadAssetsIfNeeded, which prevented loading those assets in particular.  So I've corrected the fallback code to no longer load everything before returning the selected assets they now return the results of AssetViewUtils::LoadAssetsIfNeeded as was previously the case.  There are also now overloads for loading assets that take load tags on the context menu context.

diff --git a/Engine/Plugins/Editor/ContentBrowser/ContentBrowserAssetDataSource/Source/ContentBrowserAssetDataSource/Private/AssetFileContextMenu.cpp b/Engine/Plugins/Editor/ContentBrowser/ContentBrowserAssetDataSource/Source/ContentBrowserAssetDataSource/Private/AssetFileContextMenu.cpp
--- a/Engine/Plugins/Editor/ContentBrowser/ContentBrowserAssetDataSource/Source/ContentBrowserAssetDataSource/Private/AssetFileContextMenu.cpp
+++ b/Engine/Plugins/Editor/ContentBrowser/ContentBrowserAssetDataSource/Source/ContentBrowserAssetDataSource/Private/AssetFileContextMenu.cpp
@@ -1528,19 +1528,19 @@ void FAssetFileContextMenu::ExecuteCreateBlueprintUsing()
 
 void FAssetFileContextMenu::GetSelectedAssets(TArray<UObject*>& Assets, bool SkipRedirectors) const
 {
-	for (int32 AssetIdx = 0; AssetIdx < SelectedAssets.Num(); ++AssetIdx)
-	{
-		if (SkipRedirectors && (SelectedAssets[AssetIdx].AssetClassPath == UObjectRedirector::StaticClass()->GetClassPathName()))
-		{
-			// Don't operate on Redirectors
-			continue;
-		}
-
-		if (UObject* Object = SelectedAssets[AssetIdx].GetAsset())
-		{
-			Assets.Add(Object);
-		}
+	TArray<FString> SelectedAssetPaths;
+	for (const FAssetData& SelectedAsset : SelectedAssets)
+	{
+		if (SkipRedirectors && (SelectedAsset.AssetClassPath == UObjectRedirector::StaticClass()->GetClassPathName()))
+    	{
+    		// Don't operate on Redirectors
+    		continue;
+    	}
+	
+		SelectedAssetPaths.Add(SelectedAsset.GetObjectPathString());
 	}
+
+	AssetViewUtils::LoadAssetsIfNeeded(SelectedAssetPaths, Assets);
 }
 
 /** Generates a reference graph of the world and can then find actors referencing specified objects */
diff --git a/Engine/Source/Editor/ContentBrowser/Private/AssetContextMenu.cpp b/Engine/Source/Editor/ContentBrowser/Private/AssetContextMenu.cpp
--- a/Engine/Source/Editor/ContentBrowser/Private/AssetContextMenu.cpp
+++ b/Engine/Source/Editor/ContentBrowser/Private/AssetContextMenu.cpp
@@ -386,7 +386,7 @@ void FAssetContextMenu::RegisterContextMenu(const FName MenuName)
 			if (Context && Context->CommonAssetTypeActions.IsValid())
 			{
 				PRAGMA_DISABLE_DEPRECATION_WARNINGS
-				Context->CommonAssetTypeActions.Pin()->GetActions(Context->LoadSelectedObjects(), InSection);
+				Context->CommonAssetTypeActions.Pin()->GetActions(Context->LoadSelectedObjectsIfNeeded(), InSection);
 				PRAGMA_ENABLE_DEPRECATION_WARNINGS
 			}
 		}));
@@ -397,7 +397,7 @@ void FAssetContextMenu::RegisterContextMenu(const FName MenuName)
 			if (Context && Context->CommonAssetTypeActions.IsValid())
 			{
 				PRAGMA_DISABLE_DEPRECATION_WARNINGS
-				Context->CommonAssetTypeActions.Pin()->GetActions(Context->LoadSelectedObjects(), MenuBuilder);
+				Context->CommonAssetTypeActions.Pin()->GetActions(Context->LoadSelectedObjectsIfNeeded(), MenuBuilder);
 				PRAGMA_ENABLE_DEPRECATION_WARNINGS
 			}
 		}));
diff --git a/Engine/Source/Editor/ContentBrowser/Private/ContentBrowserMenuContexts.cpp b/Engine/Source/Editor/ContentBrowser/Private/ContentBrowserMenuContexts.cpp
--- a/Engine/Source/Editor/ContentBrowser/Private/ContentBrowserMenuContexts.cpp
+++ b/Engine/Source/Editor/ContentBrowser/Private/ContentBrowserMenuContexts.cpp
@@ -6,7 +6,9 @@
 #include "ContentBrowserDataSubsystem.h"
 #include "SContentBrowser.h"
 #include "UObject/UnrealNames.h"
-#include "ToolMenuSection.h"
+#include "ToolMenus.h"
+#include "Algo/Copy.h"
+#include "AssetViewUtils.h"
 
 FName UContentBrowserToolbarMenuContext::GetCurrentPath() const
 {
@@ -26,4 +28,71 @@ bool UContentBrowserToolbarMenuContext::CanWriteToCurrentPath() const
 	}
 
 	return false;
+}
+
+TArray<UObject*> UContentBrowserAssetContextMenuContext::LoadSelectedObjectsIfNeeded() const
+{
+	TArray<FString> SelectedAssetPaths;
+	for (const FAssetData& SelectedAsset : SelectedAssets)
+	{
+		SelectedAssetPaths.Add(SelectedAsset.GetObjectPathString());
+	}
+
+	TArray<UObject*> Objects;
+	AssetViewUtils::LoadAssetsIfNeeded(SelectedAssetPaths, Objects);
+
+	return Objects;
+}
+
+TArray<FAssetData> UContentBrowserAssetContextMenuContext::GetSelectedAssetsOfType(const UClass* AssetClass, EIncludeSubclasses IncludeSubclasses) const
+{
+	TArray<FAssetData> TypedSelectedAssets;
+	if (IncludeSubclasses == EIncludeSubclasses::Yes)
+	{
+		Algo::CopyIf(SelectedAssets, TypedSelectedAssets, [AssetClass](const FAssetData& AssetData) { return AssetData.IsInstanceOf(AssetClass); });
+	}
+	else
+	{
+		Algo::CopyIf(SelectedAssets, TypedSelectedAssets, [AssetClass](const FAssetData& AssetData) { return AssetData.GetClass() == AssetClass; });
+	}
+
+	return TypedSelectedAssets;
+}
+
+const FAssetData* UContentBrowserAssetContextMenuContext::GetSingleSelectedAssetOfType(const UClass* AssetClass, EIncludeSubclasses IncludeSubclasses) const
+{
+	if (SelectedAssets.Num() == 1)
+	{
+		if (IncludeSubclasses == EIncludeSubclasses::Yes)
+		{
+			if (SelectedAssets[0].IsInstanceOf(AssetClass))
+			{
+				return &SelectedAssets[0];
+			}
+		}
+		else
+		{
+			if (SelectedAssets[0].GetClass() == AssetClass)
+			{
+				return &SelectedAssets[0];
+			}
+		}
+	}
+
+	return nullptr;
+}
+
+namespace UE::ContentBrowser
+{
+	UToolMenu* ExtendToolMenu_AssetContextMenu(UClass* AssetClass)
+	{
+		check(AssetClass);
+
+		FNameBuilder Builder;
+		Builder << TEXT("ContentBrowser.AssetContextMenu.");
+		AssetClass->GetFName().AppendString(Builder);
+		const FName MenuName(Builder.ToView());
+
+		return UToolMenus::Get()->ExtendMenu(MenuName);
+	}
 }
\ No newline at end of file
diff --git a/Engine/Source/Editor/ContentBrowser/Public/ContentBrowserMenuContexts.h b/Engine/Source/Editor/ContentBrowser/Public/ContentBrowserMenuContexts.h
--- a/Engine/Source/Editor/ContentBrowser/Public/ContentBrowserMenuContexts.h
+++ b/Engine/Source/Editor/ContentBrowser/Public/ContentBrowserMenuContexts.h
@@ -25,6 +25,12 @@ struct FFrame;
 struct FToolMenuSection;
 struct FToolMenuContext;
 
+enum class EIncludeSubclasses : uint8
+{
+	No,
+	Yes
+};
+
 UCLASS()
 class CONTENTBROWSER_API UContentBrowserAssetContextMenuContext : public UObject
 {
@@ -35,10 +41,13 @@ public:
 	TWeakPtr<FAssetContextMenu> AssetContextMenu;
 
 	TWeakPtr<IAssetTypeActions> CommonAssetTypeActions;
-
+	
 	UE_DEPRECATED(5.1, "Use SelectedAssets now, this field will not contain any objects.")
 	TArray<TWeakObjectPtr<UObject>> SelectedObjects;
 
+	/**
+	 * The currently selected assets in the content browser.
+	 */
 	UPROPERTY()
 	TArray<FAssetData> SelectedAssets;
 
@@ -48,26 +57,53 @@ public:
 	UPROPERTY()
 	bool bCanBeModified;
 
-	UFUNCTION(BlueprintCallable, Category="Tool Menus", meta=(DeprecatedFunction, DeprecationMessage = "GetSelectedObjects has been deprecated.  We no longer implictly load assets upon request.  If you can work without loading the assets, please use SelectedAssets.  Otherwise call LoadSelectedObjects"))
+	//UE_DEPRECATED(5.1, "GetSelectedObjects has been deprecated.  We no longer automatically load assets on right click.  Please use SelectedAssets and determine whatever you need for your context menu options without actually loading the assets.  When you finally need all or a subset of the selected assets use LoadSelectedAssets or LoadSelectedAssetsIf")
+	UFUNCTION(BlueprintCallable, Category="Tool Menus", meta=(DeprecatedFunction, DeprecationMessage = "GetSelectedObjects has been deprecated.  We no longer automatically load assets on right click.  If you can work without loading the assets, please use SelectedAssets.  Otherwise call LoadSelectedObjects"))
 	TArray<UObject*> GetSelectedObjects() const
 	{
-		return LoadSelectedObjects();
+		return LoadSelectedObjectsIfNeeded();
 	}
 
+	/**
+	 * Loads the selected assets (if needed) which is based on AssetViewUtils::LoadAssetsIfNeeded, this exists primarily
+	 * for backwards compatability.  Reliance on a black box to determine 'neededness' is not recommended, this function
+	 * will likely be deprecated a few versions after GetSelectedObjects.
+	 */
+	UFUNCTION(BlueprintCallable, Category="Tool Menus")
+	TArray<UObject*> LoadSelectedObjectsIfNeeded() const;
+
+	/**
+	 * Loads all the selected assets and returns an array of the objects.
+	 */
 	UFUNCTION(BlueprintCallable, Category="Tool Menus")
-	TArray<UObject*> LoadSelectedObjects() const
+	TArray<UObject*> LoadSelectedObjects(TSet<FName> LoadTags) const
 	{
-		return LoadSelectedObjects<UObject>();
+		return LoadSelectedObjects<UObject>(LoadTags);
 	}
 
+	/**
+	 * Loads all the selected assets and returns an array of the ExpectedAssetType.
+	 */
 	template<typename ExpectedAssetType>
-	TArray<ExpectedAssetType*> LoadSelectedObjects() const
+	TArray<ExpectedAssetType*> LoadSelectedObjects(TSet<FName> LoadTags = {}) const
 	{
-		return LoadSelectedObjectsIf<ExpectedAssetType>([](const FAssetData& AssetData){ return true; });
+		return LoadSelectedObjectsIf<ExpectedAssetType>(LoadTags, [](const FAssetData& AssetData){ return true; });
 	}
 
+	/**
+	 * Loads the selected assets if the PredicateFilter returns true, and returns an array of the objects.
+     */
+	template<typename ExpectedAssetType>
+    TArray<ExpectedAssetType*> LoadSelectedObjectsIf(TFunctionRef<bool(const FAssetData& AssetData)> PredicateFilter) const
+    {
+		return LoadSelectedObjectsIf<ExpectedAssetType>({}, PredicateFilter);
+    }
+
+	/**
+	 * Loads the selected assets if the PredicateFilter returns true, and returns an array of the objects.
+	 */
 	template<typename ExpectedAssetType>
-	TArray<ExpectedAssetType*> LoadSelectedObjectsIf(TFunctionRef<bool(const FAssetData& AssetData)> PredicateFilter) const
+	TArray<ExpectedAssetType*> LoadSelectedObjectsIf(TSet<FName> LoadTags, TFunctionRef<bool(const FAssetData& AssetData)> PredicateFilter) const
 	{
 		TArray<ExpectedAssetType*> Result;
 		Result.Reserve(SelectedAssets.Num());
@@ -75,7 +111,7 @@ public:
 		{
 			if (PredicateFilter(Asset))
 			{
-				if (UObject* AssetObject = Asset.GetAsset())
+				if (UObject* AssetObject = Asset.GetAsset(LoadTags))
 				{
 					if (ExpectedAssetType* AssetObjectTyped = Cast<ExpectedAssetType>(AssetObject))
 					{
@@ -87,6 +123,22 @@ public:
 		return Result;
 	}
 
+	/**
+	 * Returns a filtered array of assets that are of the desired class and potentially any subclasses.
+	 */
+	TArray<FAssetData> GetSelectedAssetsOfType(const UClass* AssetClass, EIncludeSubclasses IncludeSubclasses = EIncludeSubclasses::Yes) const;
+	
+	/**
+	 * Sometimes you want to write actions that will only operate on a singular selected asset, in those cases you
+	 * can use the following function which will only return a live ptr if it's an instance of that asset type, and
+	 * only one thing is selected.
+	 */
+	const FAssetData* GetSingleSelectedAssetOfType(const UClass* AssetClass, EIncludeSubclasses IncludeSubclasses = EIncludeSubclasses::Yes) const;
+
+	/**
+	 * Finds the Content Browser MenuContext from a Menu or Section, and returns the context provided there are some
+	 * selected assets.
+	 */
 	template<typename MenuOrSectionType>
 	static const UContentBrowserAssetContextMenuContext* FindContextWithAssets(const MenuOrSectionType& MenuOrSection)
 	{
@@ -181,3 +233,10 @@ public:
 
 	TWeakPtr<SContentBrowser> ContentBrowser;
 };
+
+class UToolMenu;
+
+namespace UE::ContentBrowser
+{
+	CONTENTBROWSER_API UToolMenu* ExtendToolMenu_AssetContextMenu(UClass* AssetClass);
+}
\ No newline at end of file
