Add cvar to hide ability to add searches as collections

diff --git a/Engine/Source/Editor/ContentBrowser/Private/SCollectionView.cpp b/Engine/Source/Editor/ContentBrowser/Private/SCollectionView.cpp
--- a/Engine/Source/Editor/ContentBrowser/Private/SCollectionView.cpp
+++ b/Engine/Source/Editor/ContentBrowser/Private/SCollectionView.cpp
@@ -71,6 +71,10 @@ struct FSlateBrush;
 
 #define LOCTEXT_NAMESPACE "ContentBrowser"
 
+// Workaround to hide Save As Collection buttons until collections support the AliasDataSource
+bool bHideSaveCollectionButton = false;
+FAutoConsoleVariableRef CVarHideSaveSearchButton(TEXT("ContentBrowser.HideSaveCollectionButton"), bHideSaveCollectionButton, TEXT("Hide the Content Browser button to save search as a dynamic collection."));
+
 namespace CollectionViewFilter
 {
 
@@ -844,7 +848,10 @@ void SCollectionView::MakeSaveDynamicCollectionMenu(FText InQueryString, FSimple
 		MakeSaveSearchMenu(MenuBuilder, OnSaveSearchClicked);
 	}
 
-	CollectionContextMenu->MakeSaveDynamicCollectionSubMenu(MenuBuilder, InQueryString);
+	if (!bHideSaveCollectionButton)
+	{
+		CollectionContextMenu->MakeSaveDynamicCollectionSubMenu(MenuBuilder, InQueryString);
+	}
 
 	FWidgetPath WidgetPath;
 	if (FSlateApplication::Get().GeneratePathToWidgetUnchecked(AsShared(), WidgetPath, EVisibility::All)) // since the collection window can be hidden, we need to manually search the path with a EVisibility::All instead of the default EVisibility::Visible
diff --git a/Engine/Source/Editor/ContentBrowser/Private/SContentBrowser.cpp b/Engine/Source/Editor/ContentBrowser/Private/SContentBrowser.cpp
--- a/Engine/Source/Editor/ContentBrowser/Private/SContentBrowser.cpp
+++ b/Engine/Source/Editor/ContentBrowser/Private/SContentBrowser.cpp
@@ -2230,11 +2230,6 @@ void SContentBrowser::OnSaveSearchButtonClicked(const FText& InSearchText)
 	{
 		SourcesViewExpandClicked();
 	}
-	if (!GetDefault<UContentBrowserSettings>()->GetDockCollections() && ActiveSourcesWidgetIndex != ContentBrowserSourcesWidgetSwitcherIndex::CollectionsView)
-	{
-		ActiveSourcesWidgetIndex = ContentBrowserSourcesWidgetSwitcherIndex::CollectionsView;
-		SourcesWidgetSwitcher->SetActiveWidgetIndex(ActiveSourcesWidgetIndex);
-	}
 
 	// We want to add any currently selected paths to the final saved query so that you get back roughly the same list of objects as what you're currently seeing
 	FString SelectedPathsQuery;
