Content Browser: Fix an issue causing the editor to crash when you tried to delete a custom filter that was active on multiple content browsers

robomerge[FN_EngineMerge] -UE5-RES

diff --git a/Engine/Source/Editor/ContentBrowser/Private/SFilterList.cpp b/Engine/Source/Editor/ContentBrowser/Private/SFilterList.cpp
--- a/Engine/Source/Editor/ContentBrowser/Private/SFilterList.cpp
+++ b/Engine/Source/Editor/ContentBrowser/Private/SFilterList.cpp
@@ -441,7 +441,7 @@ void SFilterList::OnDeleteCustomTextFilter(const TSharedPtr<ICustomTextFilter<FA
 	}
 }
 
-void SFilterList::RestoreCustomTextFilterState(const FCustomTextFilterState& InFilterState)
+bool SFilterList::RestoreCustomTextFilterState(const FCustomTextFilterState& InFilterState)
 {
 	// Find the filter associated with the current instance data from our list of custom text filters
 	TSharedRef< ICustomTextFilter<FAssetFilterType> >* Filter =
@@ -450,8 +450,11 @@ void SFilterList::RestoreCustomTextFilterState(const FCustomTextFilterState& InF
 		return Element->CreateCustomTextFilterData().FilterLabel.EqualTo(InFilterState.FilterData.FilterLabel);
 	});
 
-	// InFilterState should always be a filter that already exists!
-	check(Filter);
+	// Return if we couldn't find the filter we are trying to restore
+	if(!Filter)
+	{
+		return false;
+	}
 
 	// Get the actual FFilterBase
 	TSharedRef<FFilterBase<FAssetFilterType>> ActualFilter = Filter->Get().GetFilter().ToSharedRef();
@@ -462,6 +465,8 @@ void SFilterList::RestoreCustomTextFilterState(const FCustomTextFilterState& InF
 	// Set the filter as active if it was previously
 	AddedFilter->SetEnabled(InFilterState.bIsActive, false);
 	this->SetFrontendFilterActive(ActualFilter, InFilterState.bIsActive);
+
+	return true;
 }
 
 void SFilterList::OnExternalCustomTextFilterCreated(TSharedPtr<SWidget> BroadcastingFilterList)
@@ -494,7 +499,7 @@ void SFilterList::OnExternalCustomTextFilterCreated(TSharedPtr<SWidget> Broadcas
 		if(bIsChecked)
 		{
 			/* Remove the filter from the list (calling SBasicFilterBar::RemoveFilter because we get a compiler error
-			*  due to SAssetFilterBar overriding RemoveFilter that takes in an SFilter that hides the parent class function
+			*  due to SAssetFilterBar overriding RemoveFilter that takes in an SFilter that hides the parent class function)
 			*/
 			SBasicFilterBar<FAssetFilterType>::RemoveFilter(CustomFilter, false);
 			
@@ -639,7 +644,10 @@ void SFilterList::LoadSettings(const FName& InInstanceName)
 	// From the instance settings, get each checked filter and set the checked and active state
 	for(const FCustomTextFilterState& FilterState : InstanceSettings->CustomTextFilters)
 	{
-		RestoreCustomTextFilterState(FilterState);
+		if(!RestoreCustomTextFilterState(FilterState))
+		{
+			UE_LOG(LogSlate, Warning, TEXT("SFilterList was unable to load the following custom text filter: %s"), *FilterState.FilterData.FilterLabel.ToString());
+		}
 	}
 
 	if(InstanceSettings->bIsLayoutSaved)
diff --git a/Engine/Source/Editor/ContentBrowser/Private/SFilterList.h b/Engine/Source/Editor/ContentBrowser/Private/SFilterList.h
--- a/Engine/Source/Editor/ContentBrowser/Private/SFilterList.h
+++ b/Engine/Source/Editor/ContentBrowser/Private/SFilterList.h
@@ -167,8 +167,10 @@ private:
 	/** Empty our list of custom text filters, and load from the given config */
 	void LoadCustomTextFilters(const FFilterBarSettings* FilterBarConfig);
 
-	/** Find the custom text filter corresponding to the specified state, and restore it's state to what is specified */
-	void RestoreCustomTextFilterState(const FCustomTextFilterState& InFilterState);
+	/** Find the custom text filter corresponding to the specified state, and restore it's state to what is specified
+	 *  @return True if the filter was restored successfully, false if not
+	 */
+	bool RestoreCustomTextFilterState(const FCustomTextFilterState& InFilterState);
 	
 private:
 
