Pressing Delete in the Favorites view now removes the favorite instead of deleting the folder.

diff --git a/Engine/Source/Editor/ContentBrowser/Private/PathContextMenu.h b/Engine/Source/Editor/ContentBrowser/Private/PathContextMenu.h
--- a/Engine/Source/Editor/ContentBrowser/Private/PathContextMenu.h
+++ b/Engine/Source/Editor/ContentBrowser/Private/PathContextMenu.h
@@ -96,6 +96,9 @@ public:
 	/** Handler for when "Delete" is selected and the delete was confirmed */
 	FReply ExecuteDeleteFolderConfirmed();
 
+	/** Get the parent widget for which this menu was summoned. */
+	TSharedPtr<SWidget> GetParentContent() const { return ParentContent.Pin(); }
+
 private:
 	/** Get tooltip for delete */
 	FText GetDeleteToolTip() const;
diff --git a/Engine/Source/Editor/ContentBrowser/Private/SContentBrowser.cpp b/Engine/Source/Editor/ContentBrowser/Private/SContentBrowser.cpp
--- a/Engine/Source/Editor/ContentBrowser/Private/SContentBrowser.cpp
+++ b/Engine/Source/Editor/ContentBrowser/Private/SContentBrowser.cpp
@@ -3271,6 +3271,10 @@ bool SContentBrowser::HandleDeleteCommandCanExecute() const
 		// ... but the asset view still takes precedence over an unfocused path view unless it has no selection
 		return AssetContextMenu->CanExecuteDelete();
 	}
+	else if (FavoritePathViewPtr->GetSelectedFolderItems().Num() > 0)
+	{
+		return true;
+	}
 	else if (PathViewPtr->GetSelectedFolderItems().Num() > 0)
 	{
 		return PathContextMenu->CanExecuteDelete();
@@ -3314,12 +3318,53 @@ void SContentBrowser::HandleDeleteCommandExecute()
 		// ... but the asset view still takes precedence over an unfocused path view unless it has no selection
 		AssetContextMenu->ExecuteDelete();
 	}
+	else if (FavoritePathViewPtr->GetSelectedFolderItems().Num() > 0)
+	{
+		HandleDeleteFavorite(PathContextMenu->GetParentContent());
+	}
 	else if (PathViewPtr->GetSelectedFolderItems().Num() > 0)
 	{
 		PathContextMenu->ExecuteDelete();
 	}
 }
 
+void SContentBrowser::HandleDeleteFavorite(TSharedPtr<SWidget> ParentWidget)
+{
+	TArray<FContentBrowserItem> SelectedFolders = FavoritePathViewPtr->GetSelectedFolderItems();
+	if (ParentWidget.IsValid() && SelectedFolders.Num() > 0)
+	{
+		FText Prompt;
+		if (SelectedFolders.Num() == 1)
+		{
+			Prompt = FText::Format(LOCTEXT("FavoriteDeleteConfirm_Single", "Remove favorite '{0}'?"), SelectedFolders[0].GetDisplayName());
+		}
+		else
+		{
+			Prompt = FText::Format(LOCTEXT("FavoriteDeleteConfirm_Multiple", "Remove {0} favorites?"), SelectedFolders.Num());
+		}
+
+		// Spawn a confirmation dialog since this is potentially a highly destructive operation
+		ContentBrowserUtils::DisplayConfirmationPopup(
+			Prompt,
+			LOCTEXT("FavoriteRemoveConfirm_Yes", "Remove"),
+			LOCTEXT("FavoriteRemoveConfirm_No", "Cancel"),
+			ParentWidget.ToSharedRef(),
+			FOnClicked::CreateLambda([this, SelectedFolders]() -> FReply
+			{
+				for (const FContentBrowserItem& Folder : SelectedFolders)
+				{
+					ContentBrowserUtils::RemoveFavoriteFolder(Folder.GetVirtualPath().ToString());
+				}
+
+				GConfig->Flush(false, GEditorPerProjectIni);
+				FavoritePathViewPtr->Populate();
+
+				return FReply::Handled();
+			})
+		);
+	}
+}
+
 void SContentBrowser::HandleOpenAssetsOrFoldersCommandExecute()
 {
 	AssetViewPtr->OnOpenAssetsOrFolders();
diff --git a/Engine/Source/Editor/ContentBrowser/Private/SContentBrowser.h b/Engine/Source/Editor/ContentBrowser/Private/SContentBrowser.h
--- a/Engine/Source/Editor/ContentBrowser/Private/SContentBrowser.h
+++ b/Engine/Source/Editor/ContentBrowser/Private/SContentBrowser.h
@@ -429,6 +429,9 @@ private:
 	/** Handler for Delete */
 	void HandleDeleteCommandExecute();
 
+	/** Handler for deleting a favorite using a keybind. */
+	void HandleDeleteFavorite(TSharedPtr<SWidget> ParentWidget);
+
 	/** Handler for opening assets or folders */
 	void HandleOpenAssetsOrFoldersCommandExecute();
 
