Expose UContentBrowserAssetContextMenuContext::SelectedAssets to scripting

#rnx

diff --git a/Engine/Source/Editor/ContentBrowser/Public/ContentBrowserMenuContexts.h b/Engine/Source/Editor/ContentBrowser/Public/ContentBrowserMenuContexts.h
--- a/Engine/Source/Editor/ContentBrowser/Public/ContentBrowserMenuContexts.h
+++ b/Engine/Source/Editor/ContentBrowser/Public/ContentBrowserMenuContexts.h
@@ -54,7 +54,7 @@ public:
 	/**
 	 * The currently selected assets in the content browser.
 	 */
-	UPROPERTY()
+	UPROPERTY(BlueprintReadOnly, Category="Tool Menus")
 	TArray<FAssetData> SelectedAssets;
 
 	UPROPERTY()
