Added IContentBrowserSingleton::Register/UnregisterOnFavoritesChangedHandler() which is fired when the favorites changed.

diff --git a/Engine/Source/Editor/ContentBrowser/Private/ContentBrowserSingleton.cpp b/Engine/Source/Editor/ContentBrowser/Private/ContentBrowserSingleton.cpp
--- a/Engine/Source/Editor/ContentBrowser/Private/ContentBrowserSingleton.cpp
+++ b/Engine/Source/Editor/ContentBrowser/Private/ContentBrowserSingleton.cpp
@@ -1176,6 +1176,21 @@ void FContentBrowserSingleton::UnregisterIsFolderShowPrivateContentToggleableDel
 	IsFolderShowPrivateContentToggleableDelegate = FIsFolderShowPrivateContentToggleableDelegate();
 }
 
+FDelegateHandle FContentBrowserSingleton::RegisterOnFavoritesChangedHandler(FSimpleDelegate InOnFavoritesChanged)
+{
+	return OnFavoritesChanged.Add(InOnFavoritesChanged);
+}
+
+void FContentBrowserSingleton::UnregisterOnFavoritesChangedDelegate(FDelegateHandle Handle)
+{
+	OnFavoritesChanged.Remove(Handle);
+}
+
+void FContentBrowserSingleton::BroadcastFavoritesChanged() const
+{
+	OnFavoritesChanged.Broadcast();
+}
+
 void FContentBrowserSingleton::PopulateConfigValues()
 {
 	const FString ContentBrowserSection = TEXT("ContentBrowser");
diff --git a/Engine/Source/Editor/ContentBrowser/Private/ContentBrowserSingleton.h b/Engine/Source/Editor/ContentBrowser/Private/ContentBrowserSingleton.h
--- a/Engine/Source/Editor/ContentBrowser/Private/ContentBrowserSingleton.h
+++ b/Engine/Source/Editor/ContentBrowser/Private/ContentBrowserSingleton.h
@@ -117,6 +117,11 @@ PRAGMA_ENABLE_DEPRECATION_WARNINGS
 	virtual void UnregisterCanChangeAssetPublicStateDelegate() override;
 	virtual void RegisterIsFolderShowPrivateContentToggleableDelegate(FIsFolderShowPrivateContentToggleableDelegate InIsFolderShowPrivateContentToggleableDelegate) override;
 	virtual void UnregisterIsFolderShowPrivateContentToggleableDelegate() override;
+	virtual FDelegateHandle RegisterOnFavoritesChangedHandler(FSimpleDelegate OnFavoritesChanged) override;
+	virtual void UnregisterOnFavoritesChangedDelegate(FDelegateHandle Handle) override;
+
+	/** Broadcast that the favorites have changed. */
+	void BroadcastFavoritesChanged() const;
 
 	/** Gets the content browser singleton as a FContentBrowserSingleton */
 	static FContentBrowserSingleton& Get();
@@ -211,4 +216,7 @@ private:
 	int32 SettingsStringID;
 
 	FShowPrivateContentState ShowPrivateContentState;
+
+	DECLARE_MULTICAST_DELEGATE(FOnFavoritesChangedDelegate);
+	FOnFavoritesChangedDelegate OnFavoritesChanged;
 };
diff --git a/Engine/Source/Editor/ContentBrowser/Private/ContentBrowserUtils.cpp b/Engine/Source/Editor/ContentBrowser/Private/ContentBrowserUtils.cpp
--- a/Engine/Source/Editor/ContentBrowser/Private/ContentBrowserUtils.cpp
+++ b/Engine/Source/Editor/ContentBrowser/Private/ContentBrowserUtils.cpp
@@ -706,6 +706,8 @@ bool ContentBrowserUtils::IsFavoriteFolder(const FString& FolderPath)
 void ContentBrowserUtils::AddFavoriteFolder(const FString& FolderPath, bool bFlushConfig /*= true*/)
 {
 	FContentBrowserSingleton::Get().FavoriteFolderPaths.AddUnique(ConvertVirtualPathToInvariantPathString(FolderPath));
+
+	FContentBrowserSingleton::Get().BroadcastFavoritesChanged();
 }
 
 void ContentBrowserUtils::RemoveFavoriteFolder(const FString& FolderPath, bool bFlushConfig /*= true*/)
@@ -722,6 +724,8 @@ void ContentBrowserUtils::RemoveFavoriteFolder(const FString& FolderPath, bool b
 	{
 		GConfig->Flush(false, GEditorPerProjectIni);
 	}
+
+	FContentBrowserSingleton::Get().BroadcastFavoritesChanged();
 }
 
 const TArray<FString>& ContentBrowserUtils::GetFavoriteFolders()
diff --git a/Engine/Source/Editor/ContentBrowser/Public/IContentBrowserSingleton.h b/Engine/Source/Editor/ContentBrowser/Public/IContentBrowserSingleton.h
--- a/Engine/Source/Editor/ContentBrowser/Public/IContentBrowserSingleton.h
+++ b/Engine/Source/Editor/ContentBrowser/Public/IContentBrowserSingleton.h
@@ -363,6 +363,9 @@ struct FPathPickerConfig
 	/** If true, passes virtual paths to OnPathSelected instead of internal asset paths */
 	bool bOnPathSelectedPassesVirtualPaths;
 
+	/** If true, will show the favorites selector in the picker */
+	bool bShowFavorites;
+
 	FPathPickerConfig()
 		: bFocusSearchBoxWhenOpened(true)
 		, bAllowContextMenu(true)
@@ -370,6 +373,7 @@ struct FPathPickerConfig
 		, bAllowReadOnlyFolders(true)
 		, bAddDefaultPath(false)
 		, bOnPathSelectedPassesVirtualPaths(false)
+		, bShowFavorites(true)
 	{}
 };
 
@@ -710,4 +714,10 @@ public:
 
 	/** Unregisters the delegate for custom handling of if a Folder allows private content edits */
 	virtual void UnregisterIsFolderShowPrivateContentToggleableDelegate() = 0;
+
+	/** Register a delegate to be called when the Favorites changes. */
+	virtual FDelegateHandle RegisterOnFavoritesChangedHandler(FSimpleDelegate OnFavoritesChanged) = 0;
+
+	/** Unregister a previously-registered handler for when Favorites changes. */
+	virtual void UnregisterOnFavoritesChangedDelegate(FDelegateHandle Handle) = 0;
 };
