Dialogs that display a paths tree now also have a collapsible Favorites area.

Added SPathView::ShowFavorites option that adds a favorites view.

diff --git a/Engine/Source/Editor/ContentBrowser/Private/SPathPicker.cpp b/Engine/Source/Editor/ContentBrowser/Private/SPathPicker.cpp
--- a/Engine/Source/Editor/ContentBrowser/Private/SPathPicker.cpp
+++ b/Engine/Source/Editor/ContentBrowser/Private/SPathPicker.cpp
@@ -18,22 +18,20 @@
 #include "Internationalization/Text.h"
 #include "Layout/Children.h"
 #include "Misc/Paths.h"
+#include "SourcesSearch.h"
 #include "SPathView.h"
+#include "SSearchToggleButton.h"
 #include "Styling/AppStyle.h"
 #include "Textures/SlateIcon.h"
 #include "UObject/NameTypes.h"
 
-class FExtender;
-class SWidget;
-
 #define LOCTEXT_NAMESPACE "ContentBrowser"
 
-
 void SPathPicker::Construct( const FArguments& InArgs )
 {
 	for (auto DelegateIt = InArgs._PathPickerConfig.SetPathsDelegates.CreateConstIterator(); DelegateIt; ++DelegateIt)
 	{
-		if ((*DelegateIt) != NULL)
+		if ((*DelegateIt) != nullptr)
 		{
 			(**DelegateIt) = FSetPathPickerPathsDelegate::CreateSP(this, &SPathPicker::SetPaths);
 		}
@@ -56,6 +54,7 @@ void SPathPicker::Construct( const FArguments& InArgs )
 		.AllowReadOnlyFolders(InArgs._PathPickerConfig.bAllowReadOnlyFolders)
 		.SelectionMode(ESelectionMode::Single)
 		.CustomFolderPermissionList(InArgs._PathPickerConfig.CustomFolderPermissionList)
+		.ShowFavorites(InArgs._PathPickerConfig.bShowFavorites)
 	];
 
 	const FString& DefaultPath = InArgs._PathPickerConfig.DefaultPath;
@@ -116,7 +115,7 @@ TSharedPtr<SWidget> SPathPicker::GetItemContextMenu(TArrayView<const FContentBro
 	return GetFolderContextMenu(SelectedPackagePaths, OnGetPathContextMenuExtender, OnCreateNewFolder);
 }
 
-TSharedPtr<SWidget> SPathPicker::GetFolderContextMenu(const TArray<FString> & SelectedPaths, FContentBrowserMenuExtender_SelectedPaths InMenuExtender, FOnCreateNewFolder InOnCreateNewFolder)
+TSharedPtr<SWidget> SPathPicker::GetFolderContextMenu(const TArray<FString>& SelectedPaths, FContentBrowserMenuExtender_SelectedPaths InMenuExtender, FOnCreateNewFolder InOnCreateNewFolder)
 {
 	UContentBrowserDataSubsystem* ContentBrowserData = IContentBrowserDataModule::Get().GetSubsystem();
 
diff --git a/Engine/Source/Editor/ContentBrowser/Private/SPathView.cpp b/Engine/Source/Editor/ContentBrowser/Private/SPathView.cpp
--- a/Engine/Source/Editor/ContentBrowser/Private/SPathView.cpp
+++ b/Engine/Source/Editor/ContentBrowser/Private/SPathView.cpp
@@ -61,7 +61,9 @@
 #include "Widgets/Input/SSearchBox.h"
 #include "Widgets/Layout/SBorder.h"
 #include "Widgets/Layout/SBox.h"
+#include "Widgets/Layout/SExpandableArea.h"
 #include "Widgets/Layout/SSeparator.h"
+#include "Widgets/Layout/SSplitter.h"
 #include "Widgets/SBoxPanel.h"
 #include "Widgets/Text/STextBlock.h"
 #include "Widgets/Views/ITableRow.h"
@@ -236,17 +238,16 @@ void SPathView::Construct( const FArguments& InArgs )
 		);
 	}
 
-	ChildSlot
-	[
-		SNew(SVerticalBox)
+	TSharedRef<SVerticalBox> ContentBox = SNew(SVerticalBox);
 
-		+SVerticalBox::Slot()
+	if (!InArgs._ExternalSearch || InArgs._ShowTreeTitle)
+	{
+		ContentBox->AddSlot()
 		.AutoHeight()
 		[
 			SNew(SBorder)
 			.BorderImage(FAppStyle::Get().GetBrush("Brushes.Panel"))
 			.Padding(8.f)
-			.Visibility((!InArgs._ExternalSearch || InArgs._ShowTreeTitle) ? EVisibility::Visible : EVisibility::Collapsed)
 			[
 				SNew(SVerticalBox)
 
@@ -267,23 +268,57 @@ void SPathView::Construct( const FArguments& InArgs )
 					.Visibility(InArgs._ShowTreeTitle ? EVisibility::Visible : EVisibility::Collapsed)
 				]
 			]
-		]
+		];
+	}
 
-		// Separator
-		+SVerticalBox::Slot()
+	// Separator
+	if (InArgs._ShowSeparator)
+	{
+		ContentBox->AddSlot()
 		.AutoHeight()
 		.Padding(0, 0, 0, 1)
 		[
 			SNew(SSeparator)
-			.Visibility( ( InArgs._ShowSeparator) ? EVisibility::Visible : EVisibility::Collapsed )
-		]
-			
+		];
+	}
+
+	if (InArgs._ShowFavorites)
+	{
+		ContentBox->AddSlot()
+		.FillHeight(1.f)
+		[
+			SNew(SSplitter)
+			.Orientation(Orient_Vertical)
+			+ SSplitter::Slot()
+			.SizeRule_Lambda([this]()
+				{ 
+					return (FavoritesArea.IsValid() && FavoritesArea->IsExpanded()) ? SSplitter::ESizeRule::FractionOfParent : SSplitter::ESizeRule::SizeToContent;
+				})
+			.MinSize(24)
+			.Value(0.25f)
+			[
+				CreateFavoritesView()
+			]
+			+ SSplitter::Slot()
+			.Value(0.75f)
+			[
+				TreeViewPtr.ToSharedRef()
+			]
+		];
+	}
+	else
+	{
 		// Tree
-		+SVerticalBox::Slot()
+		ContentBox->AddSlot()
 		.FillHeight(1.f)
 		[
 			TreeViewPtr.ToSharedRef()
-		]
+		];
+	}
+
+	ChildSlot
+	[
+		ContentBox
 	];
 
 	CustomFolderPermissionList = InArgs._CustomFolderPermissionList;
@@ -2074,6 +2109,39 @@ void SPathView::UpdateLastExpandedPathsIfDirty()
 	}
 }
 
+TSharedRef<SWidget> SPathView::CreateFavoritesView()
+{
+	return SAssignNew(FavoritesArea, SExpandableArea)
+		.BorderImage(FAppStyle::Get().GetBrush("Brushes.Header"))
+		.BodyBorderImage(FAppStyle::Get().GetBrush("Brushes.Recessed"))
+		.HeaderPadding(FMargin(4.0f, 4.0f))
+		.Padding(0)
+		.AllowAnimatedTransition(false)
+		.InitiallyCollapsed(true)
+		.HeaderContent()
+		[
+			SNew(STextBlock)
+			.Text(LOCTEXT("Favorites", "Favorites"))
+			.TextStyle(FAppStyle::Get(), "ButtonText")
+			.Font(FAppStyle::Get().GetFontStyle("NormalFontBold"))
+		]
+		.BodyContent()
+		[
+			SNew(SFavoritePathView)
+			.OnItemSelectionChanged(OnItemSelectionChanged)
+			.OnGetItemContextMenu(OnGetItemContextMenu)
+			.FocusSearchBoxWhenOpened(false)
+			.ShowTreeTitle(false)
+			.ShowSeparator(false)
+			.AllowClassesFolder(bAllowClassesFolder)
+			.AllowReadOnlyFolders(bAllowReadOnlyFolders)
+			.AllowContextMenu(bAllowContextMenu)
+			.AddMetaData<FTagMetaData>(FTagMetaData(TEXT("ContentBrowserFavorites")))
+			.ExternalSearch(SearchPtr)
+		];
+}
+
+
 void SFavoritePathView::Construct(const FArguments& InArgs)
 {
 	SAssignNew(TreeViewPtr, STreeView< TSharedPtr<FTreeItem> >)
@@ -2088,11 +2156,18 @@ void SFavoritePathView::Construct(const FArguments& InArgs)
 		.ClearSelectionOnClick(false);
 
 	// Bind the favorites menu to update after folder changes
-	AssetViewUtils::OnFolderPathChanged().AddSP(this, &SFavoritePathView::FixupFavoritesFromExternalChange);
+	AssetViewUtils::OnFolderPathChanged().AddSP(this, &SFavoritePathView::FixupFavoritesFromExternalChange); 
+
+	OnFavoritesChangedHandle = FContentBrowserSingleton::Get().RegisterOnFavoritesChangedHandler(FSimpleDelegate::CreateSP(this, &SFavoritePathView::Populate, false));
 
 	SPathView::Construct(InArgs);
 }
 
+SFavoritePathView::~SFavoritePathView()
+{
+	FContentBrowserSingleton::Get().UnregisterOnFavoritesChangedDelegate(OnFavoritesChangedHandle);
+}
+
 void SFavoritePathView::Populate(const bool bIsRefreshingFilter)
 {
 	// Don't allow the selection changed delegate to be fired here
diff --git a/Engine/Source/Editor/ContentBrowser/Private/SPathView.h b/Engine/Source/Editor/ContentBrowser/Private/SPathView.h
--- a/Engine/Source/Editor/ContentBrowser/Private/SPathView.h
+++ b/Engine/Source/Editor/ContentBrowser/Private/SPathView.h
@@ -72,6 +72,7 @@ public:
 		, _AllowContextMenu(true)
 		, _AllowClassesFolder(false)
 		, _AllowReadOnlyFolders(true)
+		, _ShowFavorites(false)
 		, _SelectionMode( ESelectionMode::Multi )
 		{}
 
@@ -108,6 +109,9 @@ public:
 		/** If true, read only folders will be displayed */
 		SLATE_ARGUMENT( bool, AllowReadOnlyFolders )
 
+		/** If true, the favorites expander will be displayed */
+		SLATE_ARGUMENT(bool, ShowFavorites);
+
 		/** The selection mode for the tree view */
 		SLATE_ARGUMENT( ESelectionMode::Type, SelectionMode )
 
@@ -351,6 +355,9 @@ private:
 	/** Update the LastExpandedPath if required */
 	void UpdateLastExpandedPathsIfDirty();
 
+	/** Create a favorites view. */
+	TSharedRef<SWidget> CreateFavoritesView();
+
 protected:
 	/** A helper class to manage PreventTreeItemChangedDelegateCount by incrementing it when constructed (on the stack) and decrementing when destroyed */
 	class FScopedPreventTreeItemChangedDelegate
@@ -463,6 +470,9 @@ private:
 
 	/** Delegate to sort with */
 	FSortTreeItemChildrenDelegate SortOverride;
+
+	/** The favorites path view if one is set. */
+	TSharedPtr<SExpandableArea> FavoritesArea;
 };
 
 
@@ -473,6 +483,9 @@ private:
 class SFavoritePathView : public SPathView
 {
 public:
+
+	virtual ~SFavoritePathView();
+
 	/** Constructs this widget with InArgs */
 	virtual void Construct(const FArguments& InArgs) override;
 
@@ -501,4 +514,5 @@ private:
 
 private:
 	TArray<FString> RemovedByFolderMove;
+	FDelegateHandle OnFavoritesChangedHandle;
 };
