Content Browser: Update Revision Control icon placement, add multiple columns to the list view and add a column to show the revision control status to the list and column view

diff --git a/Engine/Source/Editor/ContentBrowser/Private/AssetViewSortManager.cpp b/Engine/Source/Editor/ContentBrowser/Private/AssetViewSortManager.cpp
--- a/Engine/Source/Editor/ContentBrowser/Private/AssetViewSortManager.cpp
+++ b/Engine/Source/Editor/ContentBrowser/Private/AssetViewSortManager.cpp
@@ -331,6 +331,8 @@ protected:
 const FName FAssetViewSortManager::NameColumnId = "Name";
 const FName FAssetViewSortManager::ClassColumnId = "Class";
 const FName FAssetViewSortManager::PathColumnId = "Path";
+const FName FAssetViewSortManager::RevisionControlColumnId = "RevisionControl";
+
 
 FAssetViewSortManager::FAssetViewSortManager()
 {
diff --git a/Engine/Source/Editor/ContentBrowser/Private/AssetViewWidgets.cpp b/Engine/Source/Editor/ContentBrowser/Private/AssetViewWidgets.cpp
--- a/Engine/Source/Editor/ContentBrowser/Private/AssetViewWidgets.cpp
+++ b/Engine/Source/Editor/ContentBrowser/Private/AssetViewWidgets.cpp
@@ -243,28 +243,12 @@ TSharedRef<SWidget> FAssetViewItemHelper::CreateListTileItemContents(T* const In
 	}
 	else
 	{
-
-
 		// The actual thumbnail
 		ItemContentsOverlay->AddSlot()
 		[
 			InThumbnail
 		];
 
-		// Source control state
-		ItemContentsOverlay->AddSlot()
-		.HAlign(HAlign_Right)
-		.VAlign(VAlign_Top)
-		.Padding(FMargin(0.0f, 3.0f, 3.0f, 0.0f))
-		[
-			SNew(SBox)
-			.WidthOverride(InTileOrListItem, &T::GetStateIconImageSize)
-			.HeightOverride(InTileOrListItem, &T::GetStateIconImageSize)
-			[
-				InTileOrListItem->GenerateSourceControlIconWidget()
-			]
-		];
-
 		// Extra external state hook
 		ItemContentsOverlay->AddSlot()
 		.HAlign(HAlign_Left)
@@ -299,8 +283,6 @@ TSharedRef<SWidget> FAssetViewItemHelper::CreateListTileItemContents(T* const In
 			.SmallView(!InTileOrListItem->CanDisplayPrimitiveTools())
 			.Visibility(InTileOrListItem, &T::GetThumbnailEditModeUIVisibility)
 		];
-
-		InTileOrListItem->HandleSourceControlStateChanged();
 	}
 	
 	return ItemContentsOverlay;
@@ -1581,10 +1563,11 @@ void SAssetListItem::Construct( const FArguments& InArgs )
 
 	AssetThumbnail = InArgs._AssetThumbnail;
 	ItemHeight = InArgs._ItemHeight;
+	IsSelectedExclusively = InArgs._IsSelectedExclusively;
+	HighlightText = InArgs._HighlightText;
 
-	const float ThumbnailPadding = InArgs._ThumbnailPadding;
+	ThumbnailPadding = InArgs._ThumbnailPadding;
 
-	TSharedPtr<SWidget> Thumbnail;
 	if ( AssetItem.IsValid() && AssetThumbnail.IsValid() )
 	{
 		FAssetThumbnailConfig ThumbnailConfig;
@@ -1610,19 +1593,80 @@ void SAssetListItem::Construct( const FArguments& InArgs )
 			}
 		}
 
-		Thumbnail = AssetThumbnail->MakeThumbnailWidget(ThumbnailConfig);
+		ThumbnailWidget = AssetThumbnail->MakeThumbnailWidget(ThumbnailConfig);
 	}
 	else
 	{
-		Thumbnail = SNew(SImage) .Image( FAppStyle::GetDefaultBrush() );
+		ThumbnailWidget = SNew(SImage) .Image( FAppStyle::GetDefaultBrush() );
 	}
+}
+END_SLATE_FUNCTION_BUILD_OPTIMIZATION
 
-	FName ItemShadowBorderName;
-	TSharedRef<SWidget> ItemContents = FAssetViewItemHelper::CreateListItemContents(this, Thumbnail.ToSharedRef(), ItemShadowBorderName);
+void SAssetListItem::OnAssetDataChanged()
+{
+	SAssetViewItem::OnAssetDataChanged();
 
-	ChildSlot
-	[
-		SNew(SBorder)
+	if (ClassTextWidget.IsValid())
+	{
+		ClassTextWidget->SetText(GetAssetClassText());
+	}
+
+	if (AssetThumbnail)
+	{
+		bool bSetThumbnail = false;
+		if (AssetItem)
+		{
+			bSetThumbnail = AssetItem->GetItem().UpdateThumbnail(*AssetThumbnail);
+		}
+		if (!bSetThumbnail)
+		{
+			AssetThumbnail->SetAsset(FAssetData());
+		}
+	}
+}
+
+void SAssetListItem::OnMouseEnter(const FGeometry& MyGeometry, const FPointerEvent& MouseEvent)
+{
+	SAssetViewItem::OnMouseEnter(MyGeometry, MouseEvent);
+
+	if (AssetThumbnail.IsValid())
+	{
+		AssetThumbnail->SetRealTime(true);
+	}
+}
+
+void SAssetListItem::OnMouseLeave(const FPointerEvent& MouseEvent)
+{
+	SAssetViewItem::OnMouseLeave(MouseEvent);
+
+	if (AssetThumbnail.IsValid())
+	{
+		AssetThumbnail->SetRealTime(false);
+	}
+}
+
+FSlateColor SAssetListItem::GetColumnTextColor(FIsSelected InIsSelected) const
+{
+	const bool bIsSelected = InIsSelected.IsBound() ? InIsSelected.Execute() : false;
+	const bool bIsHoveredOrDraggedOver = IsHovered() || bDraggedOver;
+	if (bIsSelected || bIsHoveredOrDraggedOver)
+	{
+		return FStyleColors::White;
+	}
+
+	return FSlateColor::UseForeground();
+}
+
+TSharedRef<SWidget> SAssetListItem::GenerateWidgetForColumn( const FName& ColumnName, FIsSelected InIsSelected )
+{
+	TSharedPtr<SWidget> Content;
+
+	if ( ColumnName == "Name" )
+	{
+		FName ItemShadowBorderName;
+		TSharedRef<SWidget> ItemContents = FAssetViewItemHelper::CreateListItemContents(this, ThumbnailWidget.ToSharedRef(), ItemShadowBorderName);
+
+		Content = SNew(SBorder)
 		.BorderImage(this, &SAssetViewItem::GetBorderImage)
 		.Padding(0)
 		.AddMetaData<FTagMetaData>(FTagMetaData(AssetItem->GetItem().GetVirtualPath()))
@@ -1660,9 +1704,10 @@ void SAssetListItem::Construct( const FArguments& InArgs )
 					.OnBeginTextEdit(this, &SAssetListItem::HandleBeginNameChange)
 					.OnTextCommitted(this, &SAssetListItem::HandleNameCommitted)
 					.OnVerifyTextChanged(this, &SAssetListItem::HandleVerifyNameChanged)
-					.HighlightText(InArgs._HighlightText)
-					.IsSelected(InArgs._IsSelectedExclusively)
+					.HighlightText(HighlightText)
+					.IsSelected(IsSelectedExclusively)
 					.IsReadOnly(this, &SAssetListItem::IsNameReadOnly)
+					.ColorAndOpacity(this, &SAssetListItem::GetColumnTextColor, InIsSelected)
 				]
 
 				+SVerticalBox::Slot()
@@ -1673,63 +1718,37 @@ void SAssetListItem::Construct( const FArguments& InArgs )
 					SAssignNew(ClassTextWidget, STextBlock)
 					.Font(FAppStyle::GetFontStyle("ContentBrowser.AssetListViewClassFont"))
 					.Text(GetAssetClassText())
-					.HighlightText(InArgs._HighlightText)
+					.HighlightText(HighlightText)
+					.ColorAndOpacity(this, &SAssetListItem::GetColumnTextColor, InIsSelected)
 				]
 			]
-		]
-	];
-
-	if(AssetItem.IsValid())
-	{
-		AssetItem->OnRenameRequested().BindSP(InlineRenameWidget.Get(), &SInlineEditableTextBlock::EnterEditingMode);
-		AssetItem->OnRenameCanceled().BindSP(InlineRenameWidget.Get(), &SInlineEditableTextBlock::ExitEditingMode);
-	}
-}
-END_SLATE_FUNCTION_BUILD_OPTIMIZATION
-
-void SAssetListItem::OnAssetDataChanged()
-{
-	SAssetViewItem::OnAssetDataChanged();
-
-	if (ClassTextWidget.IsValid())
-	{
-		ClassTextWidget->SetText(GetAssetClassText());
-	}
+		];
 
-	if (AssetThumbnail)
-	{
-		bool bSetThumbnail = false;
-		if (AssetItem)
-		{
-			bSetThumbnail = AssetItem->GetItem().UpdateThumbnail(*AssetThumbnail);
-		}
-		if (!bSetThumbnail)
+		if(AssetItem.IsValid())
 		{
-			AssetThumbnail->SetAsset(FAssetData());
+			AssetItem->OnRenameRequested().BindSP(InlineRenameWidget.Get(), &SInlineEditableTextBlock::EnterEditingMode);
+			AssetItem->OnRenameCanceled().BindSP(InlineRenameWidget.Get(), &SInlineEditableTextBlock::ExitEditingMode);
 		}
-	}
-}
 
-void SAssetListItem::OnMouseEnter(const FGeometry& MyGeometry, const FPointerEvent& MouseEvent)
-{
-	SAssetViewItem::OnMouseEnter(MyGeometry, MouseEvent);
-
-	if (AssetThumbnail.IsValid())
-	{
-		AssetThumbnail->SetRealTime(true);
 	}
-}
-
-void SAssetListItem::OnMouseLeave(const FPointerEvent& MouseEvent)
-{
-	SAssetViewItem::OnMouseLeave(MouseEvent);
-
-	if (AssetThumbnail.IsValid())
+	else if ( ColumnName == "RevisionControl")
 	{
-		AssetThumbnail->SetRealTime(false);
+		Content = SNew(SBox)
+			.WidthOverride(16.0f)
+			.HeightOverride(16.0f)
+			.VAlign(VAlign_Center)
+			.HAlign(HAlign_Center)
+			[
+				GenerateSourceControlIconWidget()
+			];
 	}
+	
+	HandleSourceControlStateChanged();
+
+	return Content.ToSharedRef();
 }
 
+
 float SAssetListItem::GetExtraStateIconWidth() const
 {
 	return GetStateIconImageSize().Get();
@@ -1891,11 +1910,28 @@ void SAssetTileItem::Construct( const FArguments& InArgs )
 							.AutoHeight()
 							.Padding(2.0f,0.0f, 0.0f, 2.0f)
 							[
-								SAssignNew(ClassTextWidget, STextBlock)
-								.TextStyle(FAppStyle::Get(), "ContentBrowser.ClassFont")
-								.OverflowPolicy(ETextOverflowPolicy::Ellipsis)
-								.Text(this, &SAssetTileItem::GetAssetClassText)
-								.ColorAndOpacity(this, &SAssetTileItem::GetAssetClassLabelTextColor)
+								SNew(SHorizontalBox)
+								+SHorizontalBox::Slot()
+								.AutoWidth()
+								[
+									SAssignNew(ClassTextWidget, STextBlock)
+									.TextStyle(FAppStyle::Get(), "ContentBrowser.ClassFont")
+									.OverflowPolicy(ETextOverflowPolicy::Ellipsis)
+									.Text(this, &SAssetTileItem::GetAssetClassText)
+									.ColorAndOpacity(this, &SAssetTileItem::GetAssetClassLabelTextColor)
+								]
+
+								+SHorizontalBox::Slot()
+								.Padding(FMargin(3.0f, 0.0f, 0.0f, 1.0f))
+								.HAlign(HAlign_Right)
+								[
+									SNew(SBox)
+									.WidthOverride(16.0f)
+									.HeightOverride(16.0f)
+									[
+										GenerateSourceControlIconWidget()
+									]
+								]
 							]
 							
 						]
@@ -1911,6 +1947,8 @@ void SAssetTileItem::Construct( const FArguments& InArgs )
 		]
 	];
 
+	HandleSourceControlStateChanged();
+
 	if(AssetItem.IsValid())
 	{
 		AssetItem->OnRenameRequested().BindSP(InlineRenameWidget.Get(), &SInlineEditableTextBlock::EnterEditingMode);
@@ -2248,6 +2286,18 @@ void SAssetColumnItem::Construct( const FArguments& InArgs )
 	HighlightText = InArgs._HighlightText;
 }
 
+FSlateColor SAssetColumnItem::GetColumnTextColor(FIsSelected InIsSelected) const
+{
+	const bool bIsSelected = InIsSelected.IsBound() ? InIsSelected.Execute() : false;
+	const bool bIsHoveredOrDraggedOver = IsHovered() || bDraggedOver;
+	if (bIsSelected || bIsHoveredOrDraggedOver)
+	{
+		return FStyleColors::White;
+	}
+
+	return FSlateColor::UseForeground();
+}
+
 TSharedRef<SWidget> SAssetColumnItem::GenerateWidgetForColumn( const FName& ColumnName, FIsSelected InIsSelected )
 {
 	TSharedPtr<SWidget> Content;
@@ -2294,20 +2344,6 @@ TSharedRef<SWidget> SAssetColumnItem::GenerateWidgetForColumn( const FName& Colu
 					.ColorAndOpacity(this, &SAssetColumnItem::GetAssetColor)
 				]
 
-				// Source control state
-				+SOverlay::Slot()
-				.HAlign(HAlign_Right)
-				.VAlign(VAlign_Top)
-				[
-					SNew(SBox)
-					.WidthOverride(IconOverlaySize)
-					.HeightOverride(IconOverlaySize)
-					[
-						SAssignNew(SCCStateWidget, SLayeredImage)
-						.Image(FStyleDefaults::GetNoBrush())
-					]
-				]
-
 				// Extra external state hook
 				+ SOverlay::Slot()
 				.HAlign(HAlign_Left)
@@ -2348,6 +2384,7 @@ TSharedRef<SWidget> SAssetColumnItem::GenerateWidgetForColumn( const FName& Colu
 				.HighlightText(HighlightText)
 				.IsSelected(InIsSelected)
 				.IsReadOnly(this, &SAssetColumnItem::IsNameReadOnly)
+				.ColorAndOpacity(this, &SAssetColumnItem::GetColumnTextColor, InIsSelected)
 			];
 
 		if(AssetItem.IsValid())
@@ -2383,6 +2420,15 @@ TSharedRef<SWidget> SAssetColumnItem::GenerateWidgetForColumn( const FName& Colu
 			.Text( GetAssetPathText() )
 			.HighlightText( HighlightText );
 	}
+	else if ( ColumnName == "RevisionControl" )
+	{
+		Content = SNew(SBox)
+			.WidthOverride(16.0f)
+			.HeightOverride(16.0f)
+			[
+				GenerateSourceControlIconWidget()
+			];
+	}
 	else
 	{
 		Content = SNew(STextBlock)
diff --git a/Engine/Source/Editor/ContentBrowser/Private/SAssetView.cpp b/Engine/Source/Editor/ContentBrowser/Private/SAssetView.cpp
--- a/Engine/Source/Editor/ContentBrowser/Private/SAssetView.cpp
+++ b/Engine/Source/Editor/ContentBrowser/Private/SAssetView.cpp
@@ -62,6 +62,9 @@
 #include "SFilterList.h"
 #include "SPrimaryButton.h"
 
+#include "ISourceControlModule.h"
+#include "RevisionControlStyle/RevisionControlStyle.h"
+
 #define LOCTEXT_NAMESPACE "ContentBrowser"
 #define MAX_THUMBNAIL_SIZE 4096
 
@@ -1673,6 +1676,12 @@ TSharedRef<SAssetTileView> SAssetView::CreateTileView()
 
 TSharedRef<SAssetListView> SAssetView::CreateListView()
 {
+	TSharedRef<SLayeredImage> RevisionControlColumnIcon = SNew(SLayeredImage)
+			.ColorAndOpacity(FSlateColor::UseForeground())
+			.Image(FRevisionControlStyleManager::Get().GetBrush("RevisionControl.Icon"));
+
+	RevisionControlColumnIcon->AddLayer(TAttribute<const FSlateBrush*>::CreateSP(this, &SAssetView::GetRevisionControlColumnIconBadge));
+
 	return SNew(SAssetListView)
 		.SelectionMode( SelectionMode )
 		.ListItemsSource(&FilteredAssetItems)
@@ -1681,11 +1690,42 @@ TSharedRef<SAssetListView> SAssetView::CreateListView()
 		.OnContextMenuOpening(this, &SAssetView::OnGetContextMenuContent)
 		.OnMouseButtonDoubleClick(this, &SAssetView::OnListMouseButtonDoubleClick)
 		.OnSelectionChanged(this, &SAssetView::AssetSelectionChanged)
-		.ItemHeight(this, &SAssetView::GetListViewItemHeight);
+		.ItemHeight(this, &SAssetView::GetListViewItemHeight)
+		.HeaderRow
+		(
+			SNew(SHeaderRow)
+			.ResizeMode(ESplitterResizeMode::FixedSize)
+
+			// Revision Control column, currently doesn't support sorting
+			+ SHeaderRow::Column(SortManager.RevisionControlColumnId)
+			.FixedWidth(30.f)
+			.HAlignHeader(HAlign_Center)
+			.VAlignHeader(VAlign_Center)
+			.HAlignCell(HAlign_Center)
+			.VAlignCell(VAlign_Center)
+			.DefaultLabel( LOCTEXT("Column_RC", "Revision Control") )
+			[
+				RevisionControlColumnIcon
+			]
+			
+			+ SHeaderRow::Column(SortManager.NameColumnId)
+			.FillWidth(300)
+			.SortMode( TAttribute< EColumnSortMode::Type >::Create( TAttribute< EColumnSortMode::Type >::FGetter::CreateSP( this, &SAssetView::GetColumnSortMode, SortManager.NameColumnId ) ) )
+			.SortPriority(TAttribute< EColumnSortPriority::Type >::Create(TAttribute< EColumnSortPriority::Type >::FGetter::CreateSP(this, &SAssetView::GetColumnSortPriority, SortManager.NameColumnId)))
+			.OnSort( FOnSortModeChanged::CreateSP( this, &SAssetView::OnSortColumnHeader ) )
+			.DefaultLabel( LOCTEXT("Column_Name", "Name") )
+			.ShouldGenerateWidget(true) // Can't hide name column, so at least one column is visible
+		);
 }
 
 TSharedRef<SAssetColumnView> SAssetView::CreateColumnView()
 {
+	TSharedRef<SLayeredImage> RevisionControlColumnIcon = SNew(SLayeredImage)
+				.ColorAndOpacity(FSlateColor::UseForeground())
+				.Image(FRevisionControlStyleManager::Get().GetBrush("RevisionControl.Icon"));
+
+	RevisionControlColumnIcon->AddLayer(TAttribute<const FSlateBrush*>::CreateSP(this, &SAssetView::GetRevisionControlColumnIconBadge));
+	
 	TSharedPtr<SAssetColumnView> NewColumnView = SNew(SAssetColumnView)
 		.SelectionMode( SelectionMode )
 		.ListItemsSource(&FilteredAssetItems)
@@ -1701,6 +1741,19 @@ TSharedRef<SAssetColumnView> SAssetView::CreateColumnView()
 			.ResizeMode(ESplitterResizeMode::FixedSize)
 			.CanSelectGeneratedColumn(true)
 			.OnHiddenColumnsListChanged(this, &SAssetView::OnHiddenColumnsChanged)
+
+			// Revision Control column, currently doesn't support sorting
+			+ SHeaderRow::Column(SortManager.RevisionControlColumnId)
+			.FixedWidth(24.f)
+			.HAlignHeader(HAlign_Center)
+			.VAlignHeader(VAlign_Center)
+			.HAlignCell(HAlign_Center)
+			.VAlignCell(VAlign_Center)
+			.DefaultLabel( LOCTEXT("Column_RC", "Revision Control") )
+			[
+				RevisionControlColumnIcon
+			]
+			
 			+ SHeaderRow::Column(SortManager.NameColumnId)
 			.FillWidth(300)
 			.SortMode( TAttribute< EColumnSortMode::Type >::Create( TAttribute< EColumnSortMode::Type >::FGetter::CreateSP( this, &SAssetView::GetColumnSortMode, SortManager.NameColumnId ) ) )
@@ -1710,6 +1763,11 @@ TSharedRef<SAssetColumnView> SAssetView::CreateColumnView()
 			.ShouldGenerateWidget(true) // Can't hide name column, so at least one column is visible
 		);
 
+	{
+		const bool bIsColumnVisible = !HiddenColumnNames.Contains(SortManager.RevisionControlColumnId.ToString());
+        NewColumnView->GetHeaderRow()->SetShowGeneratedColumn(SortManager.RevisionControlColumnId, bIsColumnVisible);
+	}
+
 	NewColumnView->GetHeaderRow()->SetOnGetMaxRowSizeForColumn(FOnGetMaxRowSizeForColumn::CreateRaw(NewColumnView.Get(), &SAssetColumnView::GetMaxRowSizeForColumn));
 
 	if(bShowTypeInColumnView)
@@ -1723,7 +1781,7 @@ TSharedRef<SAssetColumnView> SAssetView::CreateColumnView()
 				.DefaultLabel(LOCTEXT("Column_Class", "Type"))
 			);
 
-		bool bIsColumnVisible = !HiddenColumnNames.Contains(SortManager.ClassColumnId.ToString());
+		const bool bIsColumnVisible = !HiddenColumnNames.Contains(SortManager.ClassColumnId.ToString());
 		
 		NewColumnView->GetHeaderRow()->SetShowGeneratedColumn(SortManager.ClassColumnId, bIsColumnVisible);
 	}
@@ -1739,7 +1797,7 @@ TSharedRef<SAssetColumnView> SAssetView::CreateColumnView()
 				.DefaultLabel(LOCTEXT("Column_Path", "Path"))
 			);
 
-		bool bIsColumnVisible = !HiddenColumnNames.Contains(SortManager.PathColumnId.ToString());
+		const bool bIsColumnVisible = !HiddenColumnNames.Contains(SortManager.PathColumnId.ToString());
 		
 		NewColumnView->GetHeaderRow()->SetShowGeneratedColumn(SortManager.PathColumnId, bIsColumnVisible);
 	}
@@ -1747,6 +1805,19 @@ TSharedRef<SAssetColumnView> SAssetView::CreateColumnView()
 	return NewColumnView.ToSharedRef();
 }
 
+const FSlateBrush* SAssetView::GetRevisionControlColumnIconBadge() const
+{
+	if (ISourceControlModule::Get().IsEnabled())
+	{
+		return FRevisionControlStyleManager::Get().GetBrush("RevisionControl.Icon.ConnectedBadge");
+	}
+	else
+	{
+		return nullptr;
+	}
+}
+
+
 bool SAssetView::IsValidSearchToken(const FString& Token) const
 {
 	if ( Token.Len() == 0 )
@@ -2062,9 +2133,10 @@ void SAssetView::SetMajorityAssetType(FName NewMajorityAssetType)
 	auto IsFixedColumn = [this](FName InColumnId)
 	{
 		const bool bIsFixedNameColumn = InColumnId == SortManager.NameColumnId;
+		const bool bIsFixedRevisionControlColumn = InColumnId == SortManager.RevisionControlColumnId;
 		const bool bIsFixedClassColumn = bShowTypeInColumnView && InColumnId == SortManager.ClassColumnId;
 		const bool bIsFixedPathColumn = bShowPathInColumnView && InColumnId == SortManager.PathColumnId;
-		return bIsFixedNameColumn || bIsFixedClassColumn || bIsFixedPathColumn;
+		return bIsFixedNameColumn || bIsFixedRevisionControlColumn || bIsFixedClassColumn || bIsFixedPathColumn;
 	};
 
 
@@ -2134,7 +2206,7 @@ void SAssetView::SetMajorityAssetType(FName NewMajorityAssetType)
 				.DefaultTooltip(Column.TooltipText)
 				.FillWidth(180));
 
-			bool bIsColumnVisible = !HiddenColumnNames.Contains(TagName.ToString());
+			const bool bIsColumnVisible = !HiddenColumnNames.Contains(TagName.ToString());
 		
 			ColumnView->GetHeaderRow()->SetShowGeneratedColumn(TagName, bIsColumnVisible);
 
@@ -2211,7 +2283,7 @@ void SAssetView::SetMajorityAssetType(FName NewMajorityAssetType)
 							.DefaultTooltip(TagPair.Value.GetMetaData().TooltipText)
 							.FillWidth(180));
 
-						bool bIsColumnVisible = !HiddenColumnNames.Contains(TagPair.Key.ToString());
+						const bool bIsColumnVisible = !HiddenColumnNames.Contains(TagPair.Key.ToString());
 		
 						ColumnView->GetHeaderRow()->SetShowGeneratedColumn(TagPair.Key, bIsColumnVisible);
 						
@@ -3221,27 +3293,21 @@ TSharedRef<ITableRow> SAssetView::MakeListViewWidget(TSharedPtr<FAssetViewItem>
 
 	if (AssetItem->IsFolder())
 	{
-		TSharedPtr< STableRow<TSharedPtr<FAssetViewItem>> > TableRowWidget;
-		SAssignNew( TableRowWidget, STableRow<TSharedPtr<FAssetViewItem>>, OwnerTable )
-			.Style(FAppStyle::Get(), "ContentBrowser.AssetListView.ColumnListTableRow")
+		return
+			SNew( SAssetListViewRow, OwnerTable )
+			.OnDragDetected( this, &SAssetView::OnDraggingAssetItem )
 			.Cursor( bAllowDragging ? EMouseCursor::GrabHand : EMouseCursor::Default )
-			.OnDragDetected( this, &SAssetView::OnDraggingAssetItem );
-
-		TSharedRef<SAssetListItem> Item =
-			SNew(SAssetListItem)
-			.AssetItem(AssetItem)
-			.ItemHeight(this, &SAssetView::GetListViewItemHeight)
-			.OnRenameBegin(this, &SAssetView::AssetRenameBegin)
-			.OnRenameCommit(this, &SAssetView::AssetRenameCommit)
-			.OnVerifyRenameCommit(this, &SAssetView::AssetVerifyRenameCommit)
-			.OnItemDestroyed(this, &SAssetView::AssetItemWidgetDestroyed)
-			.ShouldAllowToolTip(this, &SAssetView::ShouldAllowToolTips)
-			.HighlightText(HighlightedText)
-			.IsSelected( FIsSelected::CreateSP(TableRowWidget.Get(), &STableRow<TSharedPtr<FAssetViewItem>>::IsSelectedExclusively) );
-
-		TableRowWidget->SetContent(Item);
-
-		return TableRowWidget.ToSharedRef();
+			.AssetListItem(
+				SNew(SAssetListItem)
+					.AssetItem(AssetItem)
+					.ItemHeight(this, &SAssetView::GetListViewItemHeight)
+					.OnRenameBegin(this, &SAssetView::AssetRenameBegin)
+					.OnRenameCommit(this, &SAssetView::AssetRenameCommit)
+					.OnVerifyRenameCommit(this, &SAssetView::AssetVerifyRenameCommit)
+					.OnItemDestroyed(this, &SAssetView::AssetItemWidgetDestroyed)
+					.ShouldAllowToolTip(this, &SAssetView::ShouldAllowToolTips)
+					.HighlightText(HighlightedText)
+			);
 	}
 	else
 	{
@@ -3253,38 +3319,32 @@ TSharedRef<ITableRow> SAssetView::MakeListViewWidget(TSharedPtr<FAssetViewItem>
 			AssetItem->GetItem().UpdateThumbnail(*AssetThumbnail);
 			AssetThumbnail->GetViewportRenderTargetTexture(); // Access the texture once to trigger it to render
 		}
-
-		TSharedPtr< STableRow<TSharedPtr<FAssetViewItem>> > TableRowWidget;
-		SAssignNew( TableRowWidget, STableRow<TSharedPtr<FAssetViewItem>>, OwnerTable )
-		.Style(FAppStyle::Get(), "ContentBrowser.AssetListView.ColumnListTableRow")
-		.Cursor( bAllowDragging ? EMouseCursor::GrabHand : EMouseCursor::Default )
-		.OnDragDetected( this, &SAssetView::OnDraggingAssetItem );
-
-		TSharedRef<SAssetListItem> Item =
-			SNew(SAssetListItem)
-			.AssetThumbnail(AssetThumbnail)
-			.AssetItem(AssetItem)
-			.ThumbnailPadding(ListViewThumbnailPadding)
-			.ItemHeight(this, &SAssetView::GetListViewItemHeight)
-			.OnRenameBegin(this, &SAssetView::AssetRenameBegin)
-			.OnRenameCommit(this, &SAssetView::AssetRenameCommit)
-			.OnVerifyRenameCommit(this, &SAssetView::AssetVerifyRenameCommit)
-			.OnItemDestroyed(this, &SAssetView::AssetItemWidgetDestroyed)
-			.ShouldAllowToolTip(this, &SAssetView::ShouldAllowToolTips)
-			.HighlightText(HighlightedText)
-			.ThumbnailEditMode(this, &SAssetView::IsThumbnailEditMode)
-			.ThumbnailLabel( ThumbnailLabel )
-			.ThumbnailHintColorAndOpacity( this, &SAssetView::GetThumbnailHintColorAndOpacity )
-			.AllowThumbnailHintLabel( AllowThumbnailHintLabel )
-			.IsSelected( FIsSelected::CreateSP(TableRowWidget.Get(), &STableRow<TSharedPtr<FAssetViewItem>>::IsSelectedExclusively) )
-			.OnIsAssetValidForCustomToolTip(OnIsAssetValidForCustomToolTip)
-			.OnGetCustomAssetToolTip(OnGetCustomAssetToolTip)
-			.OnVisualizeAssetToolTip(OnVisualizeAssetToolTip)
-			.OnAssetToolTipClosing(OnAssetToolTipClosing);
-
-		TableRowWidget->SetContent(Item);
-
-		return TableRowWidget.ToSharedRef();
+		
+		return
+			SNew( SAssetListViewRow, OwnerTable )
+			.OnDragDetected( this, &SAssetView::OnDraggingAssetItem )
+			.Cursor( bAllowDragging ? EMouseCursor::GrabHand : EMouseCursor::Default )
+			.AssetListItem(
+				SNew(SAssetListItem)
+					.AssetThumbnail(AssetThumbnail)
+					.AssetItem(AssetItem)
+					.ThumbnailPadding(ListViewThumbnailPadding)
+					.ItemHeight(this, &SAssetView::GetListViewItemHeight)
+					.OnRenameBegin(this, &SAssetView::AssetRenameBegin)
+					.OnRenameCommit(this, &SAssetView::AssetRenameCommit)
+					.OnVerifyRenameCommit(this, &SAssetView::AssetVerifyRenameCommit)
+					.OnItemDestroyed(this, &SAssetView::AssetItemWidgetDestroyed)
+					.ShouldAllowToolTip(this, &SAssetView::ShouldAllowToolTips)
+					.HighlightText(HighlightedText)
+					.ThumbnailEditMode(this, &SAssetView::IsThumbnailEditMode)
+					.ThumbnailLabel( ThumbnailLabel )
+					.ThumbnailHintColorAndOpacity( this, &SAssetView::GetThumbnailHintColorAndOpacity )
+					.AllowThumbnailHintLabel( AllowThumbnailHintLabel )
+					.OnIsAssetValidForCustomToolTip(OnIsAssetValidForCustomToolTip)
+					.OnGetCustomAssetToolTip(OnGetCustomAssetToolTip)
+					.OnVisualizeAssetToolTip(OnVisualizeAssetToolTip)
+					.OnAssetToolTipClosing(OnAssetToolTipClosing)
+			);
 	}
 }
 
@@ -4312,7 +4372,7 @@ void SAssetView::OnHiddenColumnsChanged()
 	// So instead for each column that currently exists, we update its visibility state in the HiddenColumnNames array
 	for (const SHeaderRow::FColumn& Column : ColumnView->GetHeaderRow()->GetColumns())
 	{
-		bool bIsColumnVisible = NewHiddenColumns.Contains(Column.ColumnId);
+		const bool bIsColumnVisible = NewHiddenColumns.Contains(Column.ColumnId);
 
 		if(bIsColumnVisible)
 		{
diff --git a/Engine/Source/Editor/ContentBrowser/Public/AssetViewSortManager.h b/Engine/Source/Editor/ContentBrowser/Public/AssetViewSortManager.h
--- a/Engine/Source/Editor/ContentBrowser/Public/AssetViewSortManager.h
+++ b/Engine/Source/Editor/ContentBrowser/Public/AssetViewSortManager.h
@@ -55,6 +55,9 @@ public:
 	static const FName ClassColumnId;
 	static const FName PathColumnId;
 
+	// The revision control column. NOTE: This column currently doesn't support sorting, but is wired through the Sort Manager so the feature can be added in the future
+	static const FName RevisionControlColumnId;
+
 private:
 	/** The name of the column that is currently used for sorting. */
 	FName SortColumnId[EColumnSortPriority::Max];
diff --git a/Engine/Source/Editor/ContentBrowser/Public/AssetViewWidgets.h b/Engine/Source/Editor/ContentBrowser/Public/AssetViewWidgets.h
--- a/Engine/Source/Editor/ContentBrowser/Public/AssetViewWidgets.h
+++ b/Engine/Source/Editor/ContentBrowser/Public/AssetViewWidgets.h
@@ -452,6 +452,9 @@ public:
 	/** Whether the widget should allow primitive tools to be displayed */
 	bool CanDisplayPrimitiveTools() const { return false; }
 
+	/** Generates a widget for a particular column */
+	virtual TSharedRef<SWidget> GenerateWidgetForColumn( const FName& ColumnName, FIsSelected InIsSelected );
+
 private:
 	/** Get the expected width of an extra state icon. */
 	float GetExtraStateIconWidth() const;
@@ -465,12 +468,121 @@ private:
 	/** Returns the size of the thumbnail widget */
 	FOptionalSize GetThumbnailBoxSize() const;
 
+	/** Get the text color for the columns */
+	FSlateColor GetColumnTextColor(FIsSelected InIsSelected) const;
 private:
 	/** The handle to the thumbnail that this item is rendering */
 	TSharedPtr<FAssetThumbnail> AssetThumbnail;
 
 	/** The height allowed for this item */
 	TAttribute<float> ItemHeight;
+
+	/** The padding for the thumbnail */
+	float ThumbnailPadding;
+
+	/** The actual widget for the thumbnail */
+	TSharedPtr<SWidget> ThumbnailWidget;
+
+	/** The string in the title to highlight (used when searching by string) */
+	TAttribute<FText> HighlightText;
+
+	/** Whether the item is selected in the view without anything else being selected*/
+	FIsSelected IsSelectedExclusively;
+};
+
+class SAssetListViewRow : public SMultiColumnTableRow< TSharedPtr<FAssetViewItem> >
+{
+public:
+	SLATE_BEGIN_ARGS( SAssetListViewRow )
+		{}
+
+		SLATE_EVENT( FOnDragDetected, OnDragDetected )
+		SLATE_ARGUMENT( TSharedPtr<SAssetListItem>, AssetListItem )
+
+	SLATE_END_ARGS()
+
+	/** Constructs this widget with InArgs */
+	void Construct( const FArguments& InArgs, const TSharedRef<STableViewBase>& InOwnerTableView )
+	{
+		this->AssetListItem = InArgs._AssetListItem;
+		ensure(this->AssetListItem.IsValid());
+
+		SMultiColumnTableRow< TSharedPtr<FAssetViewItem> >::Construct( 
+			FSuperRowType::FArguments()
+				.Style(FAppStyle::Get(), "ContentBrowser.AssetListView.ColumnListTableRow")
+				.OnDragDetected(InArgs._OnDragDetected), 
+			InOwnerTableView);
+		Content = this->AssetListItem;
+	}
+
+	virtual TSharedRef<SWidget> GenerateWidgetForColumn( const FName& ColumnName ) override
+	{
+		if ( this->AssetListItem.IsValid() )
+		{
+			return this->AssetListItem->GenerateWidgetForColumn(ColumnName, FIsSelected::CreateSP( this, &SAssetListViewRow::IsSelectedExclusively ));
+		}
+		else
+		{
+			return SNew(STextBlock) .Text( NSLOCTEXT("AssetView", "ListViewInvalidColumnId", "Invalid Column Item") );
+		}
+		
+	}
+
+	virtual FVector2D GetRowSizeForColumn(const FName& InColumnName) const override
+	{
+		const TSharedRef<SWidget>* ColumnWidget = GetWidgetFromColumnId(InColumnName);
+
+		if (ColumnWidget != nullptr)
+		{
+			return (*ColumnWidget)->GetDesiredSize();
+		}
+
+		return FVector2D::ZeroVector; 
+	}
+
+	virtual void Tick( const FGeometry& AllottedGeometry, const double InCurrentTime, const float InDeltaTime ) override
+	{
+		this->AssetListItem->Tick(AllottedGeometry, InCurrentTime, InDeltaTime);
+	}
+
+	virtual TSharedPtr<IToolTip> GetToolTip() override
+	{
+		return AssetListItem->GetToolTip();
+	}
+
+	virtual void OnDragEnter(const FGeometry& MyGeometry, const FDragDropEvent& DragDropEvent) override
+	{
+		AssetListItem->OnDragEnter(MyGeometry, DragDropEvent);
+	}
+
+	virtual void OnDragLeave(const FDragDropEvent& DragDropEvent) override
+	{
+		AssetListItem->OnDragLeave(DragDropEvent);
+	}
+
+	virtual FReply OnDragOver(const FGeometry& MyGeometry, const FDragDropEvent& DragDropEvent) override
+	{
+		return AssetListItem->OnDragOver(MyGeometry, DragDropEvent);
+	}
+
+	virtual FReply OnDrop(const FGeometry& MyGeometry, const FDragDropEvent& DragDropEvent) override
+	{
+		return AssetListItem->OnDrop(MyGeometry, DragDropEvent, AsShared());
+	}
+
+	virtual bool OnVisualizeTooltip(const TSharedPtr<SWidget>& TooltipContent) override
+	{
+		// We take the content from the asset column item during construction,
+		// so let the item handle the tooltip callback
+		return AssetListItem->OnVisualizeTooltip(TooltipContent);
+	}
+
+	virtual void OnToolTipClosing() override
+	{
+		AssetListItem->OnToolTipClosing();
+	}
+
+	TSharedPtr<SAssetListItem> AssetListItem;
 };
 
 /** An item in the asset tile view */
@@ -703,6 +815,9 @@ private:
 	/** Gets the value for the specified tag in this asset */
 	FText GetAssetTagText(FName Tag) const;
 
+	/** Get the text color for the columns */
+	FSlateColor GetColumnTextColor(FIsSelected InIsSelected) const;
+
 private:
 	TAttribute<FText> HighlightText;
 
diff --git a/Engine/Source/Editor/ContentBrowser/Public/SAssetView.h b/Engine/Source/Editor/ContentBrowser/Public/SAssetView.h
--- a/Engine/Source/Editor/ContentBrowser/Public/SAssetView.h
+++ b/Engine/Source/Editor/ContentBrowser/Public/SAssetView.h
@@ -451,6 +451,8 @@ private:
 	/** Creates a new column view */
 	TSharedRef<class SAssetColumnView> CreateColumnView();
 
+	const FSlateBrush* GetRevisionControlColumnIconBadge() const;
+
 	/** Returns true if the specified search token is allowed */
 	bool IsValidSearchToken(const FString& Token) const;
 
