Content Browser: Fix a bug where the revision control icons would get skewed

diff --git a/Engine/Source/Editor/ContentBrowser/Private/SAssetView.cpp b/Engine/Source/Editor/ContentBrowser/Private/SAssetView.cpp
--- a/Engine/Source/Editor/ContentBrowser/Private/SAssetView.cpp
+++ b/Engine/Source/Editor/ContentBrowser/Private/SAssetView.cpp
@@ -1744,7 +1744,7 @@ TSharedRef<SAssetColumnView> SAssetView::CreateColumnView()
 
 			// Revision Control column, currently doesn't support sorting
 			+ SHeaderRow::Column(SortManager.RevisionControlColumnId)
-			.FixedWidth(24.f)
+			.FixedWidth(30.f)
 			.HAlignHeader(HAlign_Center)
 			.VAlignHeader(VAlign_Center)
 			.HAlignCell(HAlign_Center)
