Content Browser: Fix a regression where the state of filters was not being saved/loaded

diff --git a/Engine/Source/Editor/ContentBrowser/Private/SContentBrowser.cpp b/Engine/Source/Editor/ContentBrowser/Private/SContentBrowser.cpp
--- a/Engine/Source/Editor/ContentBrowser/Private/SContentBrowser.cpp
+++ b/Engine/Source/Editor/ContentBrowser/Private/SContentBrowser.cpp
@@ -1745,7 +1745,7 @@ void SContentBrowser::SaveSettings() const
 
 
 	// Save all our data using the settings string as a key in the user settings ini
-	FilterListPtr->SaveSettings();
+	FilterListPtr->SaveSettings(GEditorPerProjectIni, SettingsIniSection, SettingsString);
 	PathViewPtr->SaveSettings(GEditorPerProjectIni, SettingsIniSection, SettingsString);
 	FavoritePathViewPtr->SaveSettings(GEditorPerProjectIni, SettingsIniSection, SettingsString + TEXT(".Favorites"));
 	CollectionViewPtr->SaveSettings(GEditorPerProjectIni, SettingsIniSection, SettingsString);
@@ -1895,7 +1895,7 @@ void SContentBrowser::LoadSettings(const FName& InInstanceName)
 	}
 
 	// Save all our data using the settings string as a key in the user settings ini
-	FilterListPtr->LoadSettings(InInstanceName);
+	FilterListPtr->LoadSettings(InInstanceName, GEditorPerProjectIni, SettingsIniSection, SettingsString);
 	PathViewPtr->LoadSettings(GEditorPerProjectIni, SettingsIniSection, SettingsString);
 	FavoritePathViewPtr->LoadSettings(GEditorPerProjectIni, SettingsIniSection, SettingsString + TEXT(".Favorites"));
 	CollectionViewPtr->LoadSettings(GEditorPerProjectIni, SettingsIniSection, SettingsString);
diff --git a/Engine/Source/Editor/ContentBrowser/Private/SFilterList.cpp b/Engine/Source/Editor/ContentBrowser/Private/SFilterList.cpp
--- a/Engine/Source/Editor/ContentBrowser/Private/SFilterList.cpp
+++ b/Engine/Source/Editor/ContentBrowser/Private/SFilterList.cpp
@@ -616,6 +616,41 @@ void SFilterList::SaveSettings()
 	SaveConfig();
 }
 
+void SFilterList::SaveSettings(const FString& IniFilename, const FString& IniSection, const FString& SettingsString)
+{
+	// Workaround for backwards compatibility with filters that save settings until they are ported to EditorConfig
+	for ( const TSharedPtr<SFilter> Filter : this->Filters )
+	{
+		const FString FilterName = Filter->GetFilterName();
+
+		// If it is a FrontendFilter
+		if ( Filter->GetFrontendFilter().IsValid() )
+		{
+			const TSharedPtr<FFilterBase<FAssetFilterType>>& FrontendFilter = Filter->GetFrontendFilter();
+			const FString CustomSettingsString = FString::Printf(TEXT("%s.CustomSettings.%s"), *SettingsString, *FilterName);
+			FrontendFilter->SaveSettings(IniFilename, IniSection, CustomSettingsString);
+		}
+	}
+	
+	SaveSettings();
+}
+
+void SFilterList::LoadSettings(const FName& InInstanceName, const FString& IniFilename, const FString& IniSection, const FString& SettingsString)
+{
+	// Workaround for backwards compatibility with filters that save settings until they are ported to EditorConfig
+	for ( auto FrontendFilterIt = this->AllFrontendFilters.CreateIterator(); FrontendFilterIt; ++FrontendFilterIt )
+	{
+		TSharedRef<FFilterBase<FAssetFilterType>>& FrontendFilter = *FrontendFilterIt;
+		const FString& FilterName = FrontendFilter->GetName();
+
+		const FString CustomSettingsString = FString::Printf(TEXT("%s.CustomSettings.%s"), *SettingsString, *FilterName);
+		FrontendFilter->LoadSettings(IniFilename, IniSection, CustomSettingsString);
+	}
+	
+	LoadSettings(InInstanceName);
+}
+
+
 void SFilterList::LoadSettings(const FName& InInstanceName)
 {
 	// If this instance doesn't want to use the shared settings, load the settings normally
diff --git a/Engine/Source/Editor/ContentBrowser/Private/SFilterList.h b/Engine/Source/Editor/ContentBrowser/Private/SFilterList.h
--- a/Engine/Source/Editor/ContentBrowser/Private/SFilterList.h
+++ b/Engine/Source/Editor/ContentBrowser/Private/SFilterList.h
@@ -140,6 +140,11 @@ public:
 	/* Copy the settings from a specific instance name */
 	void LoadSettings(const FName& InInstanceName);
 
+	/** Helper functions for backwards compatibility with filters that save state until they are ported to EditorConfig */
+	void SaveSettings(const FString& IniFilename, const FString& IniSection, const FString& SettingsString);
+	void LoadSettings(const FName& InInstanceName, const FString& IniFilename, const FString& IniSection, const FString& SettingsString);
+
+
 	virtual void SetFilterLayout(EFilterBarLayout InFilterBarLayout) override;
 
 protected:
