Add an option to collapse assetview scrollbars, preventing jitter when using multiple assetviews within an already-scrollable container.

diff --git a/Engine/Source/Editor/ContentBrowser/Private/SAssetView.cpp b/Engine/Source/Editor/ContentBrowser/Private/SAssetView.cpp
--- a/Engine/Source/Editor/ContentBrowser/Private/SAssetView.cpp
+++ b/Engine/Source/Editor/ContentBrowser/Private/SAssetView.cpp
@@ -269,6 +269,7 @@ void SAssetView::Construct( const FArguments& InArgs )
 	bShowTypeInTileView = InArgs._ShowTypeInTileView;
 	bForceShowEngineContent = InArgs._ForceShowEngineContent;
 	bForceShowPluginContent = InArgs._ForceShowPluginContent;
+	bForceHideScrollbar = InArgs._ForceHideScrollbar;
 
 	bPendingUpdateThumbnails = false;
 	bShouldNotifyNextAssetSync = true;
@@ -1671,7 +1672,8 @@ TSharedRef<SAssetTileView> SAssetView::CreateTileView()
 		.OnMouseButtonDoubleClick(this, &SAssetView::OnListMouseButtonDoubleClick)
 		.OnSelectionChanged(this, &SAssetView::AssetSelectionChanged)
 		.ItemHeight(this, &SAssetView::GetTileViewItemHeight)
-		.ItemWidth(this, &SAssetView::GetTileViewItemWidth);
+		.ItemWidth(this, &SAssetView::GetTileViewItemWidth)
+		.ScrollbarVisibility(bForceHideScrollbar ? EVisibility::Collapsed : EVisibility::Visible);
 }
 
 TSharedRef<SAssetListView> SAssetView::CreateListView()
@@ -1691,6 +1693,7 @@ TSharedRef<SAssetListView> SAssetView::CreateListView()
 		.OnMouseButtonDoubleClick(this, &SAssetView::OnListMouseButtonDoubleClick)
 		.OnSelectionChanged(this, &SAssetView::AssetSelectionChanged)
 		.ItemHeight(this, &SAssetView::GetListViewItemHeight)
+		.ScrollbarVisibility(bForceHideScrollbar ? EVisibility::Collapsed : EVisibility::Visible)
 		.HeaderRow
 		(
 			SNew(SHeaderRow)
@@ -1735,6 +1738,7 @@ TSharedRef<SAssetColumnView> SAssetView::CreateColumnView()
 		.OnMouseButtonDoubleClick(this, &SAssetView::OnListMouseButtonDoubleClick)
 		.OnSelectionChanged(this, &SAssetView::AssetSelectionChanged)
 		.Visibility(this, &SAssetView::GetColumnViewVisibility)
+		.ScrollbarVisibility(bForceHideScrollbar ? EVisibility::Collapsed : EVisibility::Visible)
 		.HeaderRow
 		(
 			SNew(SHeaderRow)
@@ -3185,6 +3189,11 @@ void SAssetView::CreateCurrentView()
 
 TSharedRef<SWidget> SAssetView::CreateShadowOverlay( TSharedRef<STableViewBase> Table )
 {
+	if (bForceHideScrollbar)
+	{
+		return Table;
+	}
+
 	return SNew(SScrollBorder, Table)
 		[
 			Table
diff --git a/Engine/Source/Editor/ContentBrowser/Public/SAssetView.h b/Engine/Source/Editor/ContentBrowser/Public/SAssetView.h
--- a/Engine/Source/Editor/ContentBrowser/Public/SAssetView.h
+++ b/Engine/Source/Editor/ContentBrowser/Public/SAssetView.h
@@ -118,6 +118,7 @@ public:
 		, _ShowTypeInTileView(true)
 		, _ForceShowEngineContent(false)
 		, _ForceShowPluginContent(false)
+		, _ForceHideScrollbar(false)
 		{}
 
 		/** Called to check if an asset should be filtered out by external code */
@@ -258,6 +259,9 @@ public:
 		/** Should always show plugin content */
 		SLATE_ARGUMENT(bool, ForceShowPluginContent)
 
+		/** Should always hide scrollbar (Removes scrollbar) */
+		SLATE_ARGUMENT(bool, ForceHideScrollbar)
+
 		/** Called to check if an asset tag should be display in details view. */
 		SLATE_EVENT( FOnShouldDisplayAssetTag, OnAssetTagWantsToBeDisplayed )
 
@@ -1066,6 +1070,9 @@ private:
 	/** If true, plugin content is always shown */
 	bool bForceShowPluginContent;
 
+	/** If true, scrollbar is never shown, removes scroll border entirely */
+	bool bForceHideScrollbar;
+
 	/** The current selection mode used by the asset view */
 	ESelectionMode::Type SelectionMode;
 
diff --git a/Engine/Source/Editor/UMGEditor/Private/Library/SLibraryViewModel.cpp b/Engine/Source/Editor/UMGEditor/Private/Library/SLibraryViewModel.cpp
--- a/Engine/Source/Editor/UMGEditor/Private/Library/SLibraryViewModel.cpp
+++ b/Engine/Source/Editor/UMGEditor/Private/Library/SLibraryViewModel.cpp
@@ -178,6 +178,7 @@ void FWidgetTemplateListViewModel::ConstructListView(TArray<TSharedPtr<FWidgetTe
 			.FrontendFilters(TemplatesFilter)
 			.ForceShowEngineContent(true)
 			.ForceShowPluginContent(true)
+			.ForceHideScrollbar(true)
 			.ShowTypeInTileView(false)
 			.ShowViewOptions(false)
 			.HighlightedText(this, &FWidgetTemplateListViewModel::GetSearchText)
