Content browser settings moved to editor config system.

There are a few settings here that we use in places where it's ambiguous which content browser instance is triggering the change, and those have not been transferred over. A few settings also appear to now be unused, and those have also not been transferred.

Additionally, set the SetDisplayDeveloperFolder(bool, bool) functions which took an override variable have been deprecated, and the overrides moved to SAssetView variables directly - these were used after syncing content to temporarily override filters and weren't persisted.

Tested across 2 Content Browsers, toggling all settings enabled on one, and all setting disabled on the other, and checking that the settings persist across restarts of the editor.

diff --git a/Engine/Plugins/Editor/GuidedTutorials/Source/IntroTutorials/Private/EditorTutorial.cpp b/Engine/Plugins/Editor/GuidedTutorials/Source/IntroTutorials/Private/EditorTutorial.cpp
--- a/Engine/Plugins/Editor/GuidedTutorials/Source/IntroTutorials/Private/EditorTutorial.cpp
+++ b/Engine/Plugins/Editor/GuidedTutorials/Source/IntroTutorials/Private/EditorTutorial.cpp
@@ -131,7 +131,6 @@ void UEditorTutorial::SetEngineFolderVisibilty(bool bNewVisibility)
 	else
 	{
 		GetMutableDefault<UContentBrowserSettings>()->SetDisplayEngineFolder(false);
-		GetMutableDefault<UContentBrowserSettings>()->SetDisplayEngineFolder(false, true);
 	}
 	GetMutableDefault<UContentBrowserSettings>()->PostEditChange();
 }
diff --git a/Engine/Source/Editor/ContentBrowser/ContentBrowser.Build.cs b/Engine/Source/Editor/ContentBrowser/ContentBrowser.Build.cs
--- a/Engine/Source/Editor/ContentBrowser/ContentBrowser.Build.cs
+++ b/Engine/Source/Editor/ContentBrowser/ContentBrowser.Build.cs
@@ -27,6 +27,7 @@ public class ContentBrowser : ModuleRules
 				"CoreUObject",
 				"ApplicationCore",
                 "InputCore",
+				"EditorConfig",
 				"Engine",
 				"Slate",
 				"SlateCore",
diff --git a/Engine/Source/Editor/ContentBrowser/Private/AssetView/AssetViewConfig.cpp b/Engine/Source/Editor/ContentBrowser/Private/AssetView/AssetViewConfig.cpp
new file mode 100644
--- /dev/null
+++ b/Engine/Source/Editor/ContentBrowser/Private/AssetView/AssetViewConfig.cpp
@@ -0,0 +1,19 @@
+// Copyright Epic Games, Inc. All Rights Reserved.
+
+#include "AssetView/AssetViewConfig.h"
+
+TObjectPtr<UAssetViewConfig> UAssetViewConfig::Instance = nullptr;
+
+void UAssetViewConfig::Initialize()
+{
+	if (!Instance)
+	{
+		Instance = NewObject<UAssetViewConfig>();
+		Instance->AddToRoot();
+	}
+}
+
+FAssetViewInstanceConfig& UAssetViewConfig::GetInstanceConfig(FName ViewName)
+{
+	return Instance->Instances.FindOrAdd(ViewName);
+}
\ No newline at end of file
diff --git a/Engine/Source/Editor/ContentBrowser/Private/AssetView/AssetViewConfig.h b/Engine/Source/Editor/ContentBrowser/Private/AssetView/AssetViewConfig.h
new file mode 100644
--- /dev/null
+++ b/Engine/Source/Editor/ContentBrowser/Private/AssetView/AssetViewConfig.h
@@ -0,0 +1,48 @@
+#pragma once
+
+#include "Containers/Map.h"
+#include "EditorConfigBase.h"
+#include "UObject/ObjectMacros.h"
+#include "UObject/ObjectPtr.h"
+#include "UObject/UObjectGlobals.h"
+
+#include "AssetViewConfig.generated.h"
+
+USTRUCT()
+struct FAssetViewInstanceConfig
+{
+	GENERATED_BODY()
+
+	/** 
+	 * The current thumbnail size, as cast from EThumbnailSize, because that enum is not a UENUM.
+	 */
+	UPROPERTY()
+	uint8 ThumbnailSize = 0; 
+
+	/** 
+	 * The current thumbnail size, as cast from EAssetViewType, because that enum is not a UENUM.
+	 */
+	UPROPERTY()
+	uint8 ViewType = 0; 
+
+	UPROPERTY()
+	TArray<FName> HiddenColumns;
+};
+
+UCLASS(EditorConfig="AssetView")
+class CONTENTBROWSER_API UAssetViewConfig : public UEditorConfigBase
+{
+	GENERATED_BODY()
+
+public:
+	static void Initialize();
+	static UAssetViewConfig* Get() { return Instance; }
+
+	FAssetViewInstanceConfig& GetInstanceConfig(FName ViewName);
+	
+	UPROPERTY(meta=(EditorConfig))
+	TMap<FName, FAssetViewInstanceConfig> Instances;
+	
+private:
+	static TObjectPtr<UAssetViewConfig> Instance;
+};
diff --git a/Engine/Source/Editor/ContentBrowser/Private/ContentBrowserConfig.cpp b/Engine/Source/Editor/ContentBrowser/Private/ContentBrowserConfig.cpp
new file mode 100644
--- /dev/null
+++ b/Engine/Source/Editor/ContentBrowser/Private/ContentBrowserConfig.cpp
@@ -0,0 +1,14 @@
+// Copyright Epic Games, Inc. All Rights Reserved.
+
+#include "ContentBrowserConfig.h"
+
+TObjectPtr<UContentBrowserConfig> UContentBrowserConfig::Instance = nullptr;
+
+void UContentBrowserConfig::Initialize()
+{
+	if (!Instance)
+	{
+		Instance = NewObject<UContentBrowserConfig>(); 
+		Instance->AddToRoot();
+	}
+}
\ No newline at end of file
diff --git a/Engine/Source/Editor/ContentBrowser/Private/ContentBrowserConfig.h b/Engine/Source/Editor/ContentBrowser/Private/ContentBrowserConfig.h
new file mode 100644
--- /dev/null
+++ b/Engine/Source/Editor/ContentBrowser/Private/ContentBrowserConfig.h
@@ -0,0 +1,144 @@
+// Copyright Epic Games, Inc. All Rights Reserved.
+
+#pragma once
+
+#include "AssetThumbnail.h"
+#include "EditorConfigBase.h"
+#include "UObject/ObjectPtr.h"
+#include "UObject/UObjectGlobals.h"
+
+#include "ContentBrowserConfig.generated.h"
+
+USTRUCT()
+struct FCollectionsConfig
+{
+	GENERATED_BODY()
+
+	UPROPERTY()
+	bool bExpanded = false;
+
+	UPROPERTY()
+	bool bSearchExpanded = false;
+	
+	UPROPERTY()
+	TArray<FString> SelectedCollections;
+
+	UPROPERTY()
+	TArray<FString> ExpandedCollections;
+};
+
+USTRUCT()
+struct FPathViewConfig
+{
+	GENERATED_BODY()
+
+	UPROPERTY()
+	bool bExpanded = false;
+
+	UPROPERTY()
+	TArray<FName> SelectedPaths;
+
+	UPROPERTY()
+	TArray<FString> PluginFilters;
+};
+
+USTRUCT()
+struct FSourcesConfig
+{
+	GENERATED_BODY()
+
+	UPROPERTY()
+	bool bExpanded = false;
+
+	UPROPERTY()
+	bool bSearchExpanded = false;
+};
+
+USTRUCT()
+struct FContentBrowserInstanceConfig
+{
+	GENERATED_BODY()
+
+	UPROPERTY()
+	FCollectionsConfig Collections;
+
+	UPROPERTY()
+	FPathViewConfig PathView;
+
+	UPROPERTY()
+	FSourcesConfig Sources;
+
+	UPROPERTY()
+	bool bShowFavorites = false;
+
+	UPROPERTY()
+	bool bFavoritesExpanded = false;
+
+	UPROPERTY()
+	bool bSourcesExpanded = false;
+
+	UPROPERTY()
+	bool bFilterRecursively = false;
+	
+	UPROPERTY()
+	bool bShowFolders = false;
+
+	UPROPERTY()
+	bool bShowEmptyFolders = false;
+
+	UPROPERTY()
+	bool bShowAllFolder = false;
+
+	UPROPERTY()
+	bool bShowEngineContent = false;
+
+	UPROPERTY()
+	bool bShowDeveloperContent = false;
+
+	UPROPERTY()
+	bool bShowLocalizedContent = false;
+
+	UPROPERTY()
+	bool bShowPluginContent = false;
+
+	UPROPERTY()
+	bool bShowCppFolders = false;
+
+	UPROPERTY()
+	bool bOrganizeFolders = false;
+
+	UPROPERTY()
+	bool bRealTimeThumbnails = false;
+
+	UPROPERTY()
+	bool bCollectionsDocked = false;
+
+	UPROPERTY()
+	bool bSearchClasses = false;
+
+	UPROPERTY()
+	bool bSearchAssetPaths = false;
+
+	UPROPERTY()
+	bool bSearchCollections = false;
+};
+
+UCLASS(EditorConfig="ContentBrowser")
+class UContentBrowserConfig : public UEditorConfigBase
+{
+	GENERATED_BODY()
+	
+public:
+
+	static void Initialize();
+	static UContentBrowserConfig* Get() { return Instance; }
+
+	UPROPERTY()
+	TSet<FString> Favorites;
+
+	UPROPERTY(meta=(EditorConfig))
+	TMap<FName, FContentBrowserInstanceConfig> Instances;
+
+private:
+	static TObjectPtr<UContentBrowserConfig> Instance;
+};
\ No newline at end of file
diff --git a/Engine/Source/Editor/ContentBrowser/Private/ContentBrowserUtils.cpp b/Engine/Source/Editor/ContentBrowser/Private/ContentBrowserUtils.cpp
--- a/Engine/Source/Editor/ContentBrowser/Private/ContentBrowserUtils.cpp
+++ b/Engine/Source/Editor/ContentBrowser/Private/ContentBrowserUtils.cpp
@@ -7,6 +7,7 @@
 #include "AssetRegistry/AssetData.h"
 #include "Containers/Map.h"
 #include "Containers/StringView.h"
+#include "ContentBrowserConfig.h"
 #include "ContentBrowserDataFilter.h"
 #include "ContentBrowserDataSource.h"
 #include "ContentBrowserDataSubsystem.h"
@@ -703,31 +704,48 @@ bool ContentBrowserUtils::IsFavoriteFolder(const FString& FolderPath)
 	return FContentBrowserSingleton::Get().FavoriteFolderPaths.Contains(ConvertVirtualPathToInvariantPathString(FolderPath));
 }
 
-void ContentBrowserUtils::AddFavoriteFolder(const FString& FolderPath, bool bFlushConfig /*= true*/)
+void ContentBrowserUtils::AddFavoriteFolder(const FString& FolderPath)
 {
-	FContentBrowserSingleton::Get().FavoriteFolderPaths.AddUnique(ConvertVirtualPathToInvariantPathString(FolderPath));
+	const FString InvariantFolder = ConvertVirtualPathToInvariantPathString(FolderPath);
+
+	FContentBrowserSingleton::Get().FavoriteFolderPaths.AddUnique(InvariantFolder);
+
+	if (UContentBrowserConfig* EditorConfig = UContentBrowserConfig::Get())
+	{
+		EditorConfig->Favorites.Add(InvariantFolder);
+
+		UContentBrowserConfig::Get()->SaveEditorConfig();
+	}
 
 	FContentBrowserSingleton::Get().BroadcastFavoritesChanged();
 }
+	
+void ContentBrowserUtils::AddFavoriteFolder(const FString& FolderPath, bool bFlushConfig)
+{
+	AddFavoriteFolder(FolderPath);
+}
 
-void ContentBrowserUtils::RemoveFavoriteFolder(const FString& FolderPath, bool bFlushConfig /*= true*/)
+void ContentBrowserUtils::RemoveFavoriteFolder(const FString& FolderPath)
 {
-	FString InvariantFolder = ConvertVirtualPathToInvariantPathString(FolderPath);
+	const FString InvariantFolder = ConvertVirtualPathToInvariantPathString(FolderPath);
 
-	// Find and remove any subfolders
-	FContentBrowserSingleton::Get().FavoriteFolderPaths.RemoveAll([&InvariantFolder](const FString& FavoritePath)
-	{
-		return FavoritePath.StartsWith(InvariantFolder) && (FavoritePath.Len() <= InvariantFolder.Len() || FavoritePath[InvariantFolder.Len()] == TEXT('/'));
-	});
+	FContentBrowserSingleton::Get().FavoriteFolderPaths.Remove(InvariantFolder);
 
-	if (bFlushConfig)
+	if (UContentBrowserConfig* EditorConfig = UContentBrowserConfig::Get())
 	{
-		GConfig->Flush(false, GEditorPerProjectIni);
+		EditorConfig->Favorites.Remove(InvariantFolder);
+
+		UContentBrowserConfig::Get()->SaveEditorConfig();
 	}
 
 	FContentBrowserSingleton::Get().BroadcastFavoritesChanged();
 }
 
+void ContentBrowserUtils::RemoveFavoriteFolder(const FString& FolderPath, bool bFlushConfig)
+{
+	RemoveFavoriteFolder(FolderPath);
+}
+
 const TArray<FString>& ContentBrowserUtils::GetFavoriteFolders()
 {
 	return FContentBrowserSingleton::Get().FavoriteFolderPaths;
diff --git a/Engine/Source/Editor/ContentBrowser/Private/ContentBrowserUtils.h b/Engine/Source/Editor/ContentBrowser/Private/ContentBrowserUtils.h
--- a/Engine/Source/Editor/ContentBrowser/Private/ContentBrowserUtils.h
+++ b/Engine/Source/Editor/ContentBrowser/Private/ContentBrowserUtils.h
@@ -97,9 +97,17 @@ namespace ContentBrowserUtils
 	/** Returns if this folder has been marked as a favorite folder */
 	bool IsFavoriteFolder(const FString& FolderPath);
 
-	void AddFavoriteFolder(const FString& FolderPath, bool bFlushConfig = true);
+	/** Add a favorite folder. This should be a full virtual path. */
+	void AddFavoriteFolder(const FString& FolderPath);
 
-	void RemoveFavoriteFolder(const FString& FolderPath, bool bFlushConfig = true);
+	UE_DEPRECATED(5.2, "The bFlushConfig parameter is ignored, since this function does not persist the favorites to config.")
+	void AddFavoriteFolder(const FString& FolderPath, bool bFlushConfig);
+
+	/** Remove a favorite folder. This should be a full virtual path. */
+	void RemoveFavoriteFolder(const FString& FolderPath);
+
+	UE_DEPRECATED(5.2, "The bFlushConfig parameter is ignored, since this function does not persist the favorites to config.")
+	void RemoveFavoriteFolder(const FString& FolderPath, bool bFlushConfig);
 
 	const TArray<FString>& GetFavoriteFolders();
 
diff --git a/Engine/Source/Editor/ContentBrowser/Private/SAssetView.cpp b/Engine/Source/Editor/ContentBrowser/Private/SAssetView.cpp
--- a/Engine/Source/Editor/ContentBrowser/Private/SAssetView.cpp
+++ b/Engine/Source/Editor/ContentBrowser/Private/SAssetView.cpp
@@ -1,66 +1,64 @@
 // Copyright Epic Games, Inc. All Rights Reserved.
 
 #include "SAssetView.h"
+
 #include "Algo/Transform.h"
-#include "HAL/FileManager.h"
-#include "UObject/UnrealType.h"
-#include "Widgets/SOverlay.h"
-#include "Engine/GameViewportClient.h"
-#include "Factories/Factory.h"
-#include "Framework/Commands/UIAction.h"
-#include "Textures/SlateIcon.h"
-#include "Misc/CommandLine.h"
-#include "Misc/ConfigCacheIni.h"
-#include "SlateOptMacros.h"
-#include "Framework/Application/SlateApplication.h"
-#include "Widgets/Images/SImage.h"
-#include "Widgets/Notifications/SProgressBar.h"
-#include "Widgets/Text/STextBlock.h"
-#include "Framework/MultiBox/MultiBoxBuilder.h"
-#include "Widgets/Input/SButton.h"
-#include "Widgets/Layout/SScrollBorder.h"
-#include "Widgets/Input/SComboButton.h"
-#include "Widgets/Input/SSlider.h"
-#include "Framework/Docking/TabManager.h"
-#include "Styling/AppStyle.h"
-#include "Settings/ContentBrowserSettings.h"
-#include "Engine/Blueprint.h"
-#include "Engine/Level.h"
-#include "Editor.h"
+#include "AssetRegistry/AssetRegistryState.h"
 #include "AssetSelection.h"
-#include "IAssetTools.h"
 #include "AssetToolsModule.h"
+#include "AssetView/AssetViewConfig.h"
+#include "AssetViewTypes.h"
+#include "AssetViewWidgets.h"
+#include "ContentBrowserConfig.h"
+#include "ContentBrowserDataDragDropOp.h"
+#include "ContentBrowserDataLegacyBridge.h"
+#include "ContentBrowserDataSource.h"
+#include "ContentBrowserDataSubsystem.h"
 #include "ContentBrowserLog.h"
-#include "FrontendFilterBase.h"
+#include "ContentBrowserMenuContexts.h"
+#include "ContentBrowserModule.h"
 #include "ContentBrowserSingleton.h"
-#include "EditorWidgetsModule.h"
-#include "AssetViewTypes.h"
+#include "ContentBrowserUtils.h"
+#include "DesktopPlatformModule.h"
 #include "DragAndDrop/AssetDragDropOp.h"
 #include "DragDropHandler.h"
-#include "AssetViewWidgets.h"
-#include "ContentBrowserModule.h"
-#include "ObjectTools.h"
+#include "Editor.h"
+#include "EditorWidgetsModule.h"
+#include "Engine/Blueprint.h"
+#include "Engine/Level.h"
+#include "Factories/Factory.h"
+#include "Framework/Application/SlateApplication.h"
+#include "Framework/Commands/UIAction.h"
+#include "Framework/MultiBox/MultiBoxBuilder.h"
 #include "Framework/Notifications/NotificationManager.h"
-#include "Widgets/Notifications/SNotificationList.h"
-#include "Widgets/Layout/SSplitter.h"
+#include "FrontendFilterBase.h"
+#include "HAL/FileManager.h"
 #include "HAL/PlatformApplicationMisc.h"
-#include "DesktopPlatformModule.h"
+#include "IAssetTools.h"
+#include "IContentBrowserDataModule.h"
+#include "Materials/Material.h"
+#include "Misc/CommandLine.h"
+#include "Misc/ConfigCacheIni.h"
 #include "Misc/FileHelper.h"
-#include "Misc/TextFilterUtils.h"
 #include "Misc/NamePermissionList.h"
-#include "AssetRegistry/AssetRegistryState.h"
-#include "Materials/Material.h"
-#include "ContentBrowserMenuContexts.h"
-#include "ContentBrowserUtils.h"
-#include "ToolMenus.h"
-
-#include "IContentBrowserDataModule.h"
-#include "ContentBrowserDataSource.h"
-#include "ContentBrowserDataSubsystem.h"
-#include "ContentBrowserDataLegacyBridge.h"
-#include "ContentBrowserDataDragDropOp.h"
+#include "Misc/TextFilterUtils.h"
+#include "ObjectTools.h"
+#include "SContentBrowser.h"
 #include "SFilterList.h"
 #include "SPrimaryButton.h"
+#include "Settings/ContentBrowserSettings.h"
+#include "SlateOptMacros.h"
+#include "Styling/AppStyle.h"
+#include "Textures/SlateIcon.h"
+#include "ToolMenus.h"
+#include "UObject/UnrealType.h"
+#include "Widgets/Images/SImage.h"
+#include "Widgets/Input/SComboButton.h"
+#include "Widgets/Layout/SScrollBorder.h"
+#include "Widgets/Layout/SSplitter.h"
+#include "Widgets/Notifications/SProgressBar.h"
+#include "Widgets/SOverlay.h"
+#include "Widgets/Text/STextBlock.h"
 
 #include "ISourceControlModule.h"
 #include "RevisionControlStyle/RevisionControlStyle.h"
@@ -265,7 +263,7 @@
 
 	bShowPathInColumnView = InArgs._ShowPathInColumnView;
 	bShowTypeInColumnView = InArgs._ShowTypeInColumnView;
-	bSortByPathInColumnView = bShowPathInColumnView & InArgs._SortByPathInColumnView;
+	bSortByPathInColumnView = bShowPathInColumnView && InArgs._SortByPathInColumnView;
 	bShowTypeInTileView = InArgs._ShowTypeInTileView;
 	bForceShowEngineContent = InArgs._ForceShowEngineContent;
 	bForceShowPluginContent = InArgs._ForceShowPluginContent;
@@ -920,7 +918,8 @@
 	{
 		// Clamp value to normal range and update state
 		ThumbnailSizeConfig = FMath::Clamp<int32>(ThumbnailSizeConfig, 0, (int32)EThumbnailSize::MAX-1);
-		OnThumbnailSizeChanged((EThumbnailSize)ThumbnailSizeConfig);
+
+		ThumbnailSize = (EThumbnailSize) ThumbnailSizeConfig;
 	}
 
 	int32 ViewType = EAssetViewType::Tile;
@@ -948,7 +947,6 @@
 				ColumnView->GetHeaderRow()->SetShowGeneratedColumn(Column.ColumnId, !HiddenColumnNames.Contains(Column.ColumnId.ToString()));
 			}
 		}
-		
 	}
 }
 
@@ -2006,9 +2004,19 @@
 
 bool SAssetView::IsFilteringRecursively() const
 {
+	if (!bFilterRecursivelyWithBackendFilter)
+	{
+		return false;
+	}
+
 	// In some cases we want to not filter recursively even if we have a backend filter (e.g. the open level window)
 	// Most of the time, bFilterRecursivelyWithBackendFilter is true
-	return bFilterRecursivelyWithBackendFilter && GetDefault<UContentBrowserSettings>()->FilterRecursively;
+	if (const FContentBrowserInstanceConfig* EditorConfig = GetContentBrowserConfig())
+	{
+		return EditorConfig->bFilterRecursively;
+	}
+
+	return GetDefault<UContentBrowserSettings>()->FilterRecursively;
 }
 
 bool SAssetView::IsToggleFilteringRecursivelyAllowed() const
@@ -2019,7 +2027,18 @@
 void SAssetView::ToggleFilteringRecursively()
 {
 	check(IsToggleFilteringRecursivelyAllowed());
-	GetMutableDefault<UContentBrowserSettings>()->FilterRecursively = !GetDefault<UContentBrowserSettings>()->FilterRecursively;
+
+	bool bNewState = !GetDefault<UContentBrowserSettings>()->FilterRecursively;
+	
+	if (FContentBrowserInstanceConfig* EditorConfig = GetContentBrowserConfig())
+	{
+		bNewState = !EditorConfig->bFilterRecursively;
+
+		EditorConfig->bFilterRecursively = bNewState;
+		UContentBrowserConfig::Get()->SaveEditorConfig();
+	}
+
+	GetMutableDefault<UContentBrowserSettings>()->FilterRecursively = bNewState;
 	GetMutableDefault<UContentBrowserSettings>()->PostEditChange();
 }
 
@@ -2105,25 +2124,78 @@
 	UE_LOG(LogContentBrowser, VeryVerbose, TEXT("AssetView - RefreshFilteredItems completed in %0.4f seconds"), FPlatformTime::Seconds() - RefreshFilteredItemsStartTime);
 }
 
+FContentBrowserInstanceConfig* SAssetView::GetContentBrowserConfig() const
+{
+	if (TSharedPtr<SContentBrowser> ContentBrowser = OwningContentBrowser.Pin())
+	{
+		if (UContentBrowserConfig* EditorConfig = UContentBrowserConfig::Get())
+		{
+			return UContentBrowserConfig::Get()->Instances.Find(ContentBrowser->GetInstanceName());
+		}
+	}
+	return nullptr;
+}
+
+FAssetViewInstanceConfig* SAssetView::GetAssetViewConfig() const
+{
+	if (TSharedPtr<SContentBrowser> ContentBrowser = OwningContentBrowser.Pin())
+	{
+		const FName InstanceName = ContentBrowser->GetInstanceName();
+		if (!InstanceName.IsNone())
+		{
+			if (UAssetViewConfig* Config = UAssetViewConfig::Get())
+			{
+				return &Config->GetInstanceConfig(InstanceName);
+			}
+		}
+	}
+
+	return nullptr;
+}
+
 void SAssetView::ToggleShowAllFolder()
 {
-	GetMutableDefault<UContentBrowserSettings>()->bShowAllFolder = !GetDefault<UContentBrowserSettings>()->bShowAllFolder;
+	const bool bNewValue = !IsShowingAllFolder();
+	GetMutableDefault<UContentBrowserSettings>()->bShowAllFolder = bNewValue;
 	GetMutableDefault<UContentBrowserSettings>()->PostEditChange();
+
+	if (FContentBrowserInstanceConfig* Config = GetContentBrowserConfig())
+	{
+		Config->bShowAllFolder = bNewValue;
+		UContentBrowserConfig::Get()->SaveEditorConfig();
+	}
 }
 
 bool SAssetView::IsShowingAllFolder() const
 {
+	if (FContentBrowserInstanceConfig* Config = GetContentBrowserConfig())
+	{
+		return Config->bShowAllFolder;
+	}
+
 	return GetDefault<UContentBrowserSettings>()->bShowAllFolder;
 }
 
 void SAssetView::ToggleOrganizeFolders()
 {
-	GetMutableDefault<UContentBrowserSettings>()->bOrganizeFolders = !GetDefault<UContentBrowserSettings>()->bOrganizeFolders;
+	const bool bNewValue = !IsOrganizingFolders();
+	GetMutableDefault<UContentBrowserSettings>()->bOrganizeFolders = bNewValue;
 	GetMutableDefault<UContentBrowserSettings>()->PostEditChange();
+
+	if (FContentBrowserInstanceConfig* Config = GetContentBrowserConfig())
+	{
+		Config->bOrganizeFolders = bNewValue;
+		UContentBrowserConfig::Get()->SaveEditorConfig();
+	}
 }
 
 bool SAssetView::IsOrganizingFolders() const
 {
+	if (FContentBrowserInstanceConfig* Config = GetContentBrowserConfig())
+	{
+		return Config->bOrganizeFolders;
+	}
+
 	return GetDefault<UContentBrowserSettings>()->bOrganizeFolders;
 }
 
@@ -2434,22 +2506,22 @@
 
 	// Get all menu extenders for this context menu from the content browser module
 	FContentBrowserModule& ContentBrowserModule = FModuleManager::GetModuleChecked<FContentBrowserModule>( TEXT("ContentBrowser") );
-	TArray<FContentBrowserMenuExtender> MenuExtenderDelegates = ContentBrowserModule.GetAllAssetViewViewMenuExtenders();
+	const TArray<FContentBrowserMenuExtender>& MenuExtenderDelegates = ContentBrowserModule.GetAllAssetViewViewMenuExtenders();
 
 	TArray<TSharedPtr<FExtender>> Extenders;
-	for (int32 i = 0; i < MenuExtenderDelegates.Num(); ++i)
+	for (const FContentBrowserMenuExtender& Extender : MenuExtenderDelegates)
 	{
-		if (MenuExtenderDelegates[i].IsBound())
+		if (Extender.IsBound())
 		{
-			Extenders.Add(MenuExtenderDelegates[i].Execute());
+			Extenders.Add(Extender.Execute());
 		}
 	}
-	TSharedPtr<FExtender> MenuExtender = FExtender::Combine(Extenders);
 
 	UContentBrowserAssetViewContextMenuContext* Context = NewObject<UContentBrowserAssetViewContextMenuContext>();
 	Context->AssetView = SharedThis(this);
 	Context->OwningContentBrowser = OwningContentBrowser;
 
+	TSharedPtr<FExtender> MenuExtender = FExtender::Combine(Extenders);
 	FToolMenuContext MenuContext(nullptr, MenuExtender, Context);
 
 	if (OnExtendAssetViewOptionsMenuContext.IsBound())
@@ -2850,8 +2922,18 @@
 
 void SAssetView::ToggleShowFolders()
 {
-	check( IsToggleShowFoldersAllowed() );
-	GetMutableDefault<UContentBrowserSettings>()->DisplayFolders = !GetDefault<UContentBrowserSettings>()->DisplayFolders;
+	check(IsToggleShowFoldersAllowed());
+
+	bool bNewState = !GetDefault<UContentBrowserSettings>()->DisplayFolders;
+
+	if (FContentBrowserInstanceConfig* Config = GetContentBrowserConfig())
+	{
+		bNewState = !Config->bShowFolders;
+		Config->bShowFolders = bNewState;
+		UContentBrowserConfig::Get()->SaveEditorConfig();
+	}
+
+	GetMutableDefault<UContentBrowserSettings>()->DisplayFolders = bNewState;
 	GetMutableDefault<UContentBrowserSettings>()->PostEditChange();
 }
 
@@ -2862,7 +2944,17 @@
 
 bool SAssetView::IsShowingFolders() const
 {
-	return IsToggleShowFoldersAllowed() && GetDefault<UContentBrowserSettings>()->DisplayFolders;
+	if (!IsToggleShowFoldersAllowed())
+	{
+		return false;
+	}
+	
+	if (const FContentBrowserInstanceConfig* Config = GetContentBrowserConfig())
+	{
+		return Config->bShowFolders;
+	}
+
+	return GetDefault<UContentBrowserSettings>()->DisplayFolders;
 }
 
 bool SAssetView::IsShowingReadOnlyFolders() const
@@ -2872,7 +2964,17 @@
 
 void SAssetView::ToggleShowEmptyFolders()
 {
-	check( IsToggleShowEmptyFoldersAllowed() );
+	check(IsToggleShowEmptyFoldersAllowed());
+
+	bool bNewState = !GetDefault<UContentBrowserSettings>()->DisplayEmptyFolders;
+
+	if (FContentBrowserInstanceConfig* Config = GetContentBrowserConfig())
+	{
+		bNewState = !Config->bShowEmptyFolders;
+		Config->bShowEmptyFolders = bNewState;
+		UContentBrowserConfig::Get()->SaveEditorConfig();
+	}
+	
 	GetMutableDefault<UContentBrowserSettings>()->DisplayEmptyFolders = !GetDefault<UContentBrowserSettings>()->DisplayEmptyFolders;
 	GetMutableDefault<UContentBrowserSettings>()->PostEditChange();
 }
@@ -2884,14 +2986,33 @@
 
 bool SAssetView::IsShowingEmptyFolders() const
 {
-	return IsToggleShowEmptyFoldersAllowed() && GetDefault<UContentBrowserSettings>()->DisplayEmptyFolders;
+	if (!IsToggleShowEmptyFoldersAllowed())
+	{
+		return false;
+	}
+
+	if (const FContentBrowserInstanceConfig* Config = GetContentBrowserConfig())
+	{
+		return Config->bShowEmptyFolders;
+	}
+
+	return GetDefault<UContentBrowserSettings>()->DisplayEmptyFolders;
 }
 
 void SAssetView::ToggleRealTimeThumbnails()
 {
-	check( CanShowRealTimeThumbnails() );
-	GetMutableDefault<UContentBrowserSettings>()->RealTimeThumbnails = !GetDefault<UContentBrowserSettings>()->RealTimeThumbnails;
+	check(CanShowRealTimeThumbnails());
+
+	bool bNewState = !IsShowingRealTimeThumbnails();
+
+	GetMutableDefault<UContentBrowserSettings>()->RealTimeThumbnails = bNewState;
 	GetMutableDefault<UContentBrowserSettings>()->PostEditChange();
+
+	if (FContentBrowserInstanceConfig* Config = GetContentBrowserConfig())
+	{
+		Config->bRealTimeThumbnails = bNewState;
+		UContentBrowserConfig::Get()->SaveEditorConfig();
+	}
 }
 
 bool SAssetView::CanShowRealTimeThumbnails() const
@@ -2901,70 +3022,113 @@
 
 bool SAssetView::IsShowingRealTimeThumbnails() const
 {
-	return CanShowRealTimeThumbnails() && GetDefault<UContentBrowserSettings>()->RealTimeThumbnails;
+	if (!CanShowRealTimeThumbnails())
+	{
+		return false;
+	}
+	
+	if (const FContentBrowserInstanceConfig* Config = GetContentBrowserConfig())
+	{
+		return Config->bRealTimeThumbnails;
+	}
+
+	return GetDefault<UContentBrowserSettings>()->RealTimeThumbnails;
 }
 
 void SAssetView::ToggleShowPluginContent()
 {
-	bool bDisplayPlugins = GetDefault<UContentBrowserSettings>()->GetDisplayPluginFolders();
-	bool bRawDisplayPlugins = GetDefault<UContentBrowserSettings>()->GetDisplayPluginFolders( true );
+	check(IsToggleShowPluginContentAllowed());
 
-	// Only if both these flags are false when toggling we want to enable the flag, otherwise we're toggling off
-	if ( !bDisplayPlugins && !bRawDisplayPlugins )
+	bOverrideShowPluginContent.Reset();
+
+	bool bNewState = !GetDefault<UContentBrowserSettings>()->GetDisplayPluginFolders();
+
+	if (FContentBrowserInstanceConfig* EditorConfig = GetContentBrowserConfig())
 	{
-		GetMutableDefault<UContentBrowserSettings>()->SetDisplayPluginFolders( true );
+		bNewState = !EditorConfig->bShowPluginContent;
+		EditorConfig->bShowPluginContent = bNewState;
+		UContentBrowserConfig::Get()->SaveEditorConfig();
 	}
-	else
-	{
-		GetMutableDefault<UContentBrowserSettings>()->SetDisplayPluginFolders( false );
-		GetMutableDefault<UContentBrowserSettings>()->SetDisplayPluginFolders( false, true );
-	}	
+
+	GetMutableDefault<UContentBrowserSettings>()->SetDisplayPluginFolders(bNewState);
 	GetMutableDefault<UContentBrowserSettings>()->PostEditChange();
 }
 
 bool SAssetView::IsShowingPluginContent() const
 {
-	return bForceShowPluginContent || GetDefault<UContentBrowserSettings>()->GetDisplayPluginFolders();
+	if (bForceShowPluginContent)
+	{
+		return true;
+	}
+
+	if (bOverrideShowPluginContent.IsSet())
+	{
+		return bOverrideShowPluginContent.GetValue();
+	}
+
+	if (const FContentBrowserInstanceConfig* Config = GetContentBrowserConfig())
+	{
+		return Config->bShowPluginContent;
+	}
+
+	return GetDefault<UContentBrowserSettings>()->GetDisplayPluginFolders();
 }
 
 void SAssetView::ToggleShowEngineContent()
 {
-	bool bDisplayEngine = GetDefault<UContentBrowserSettings>()->GetDisplayEngineFolder();
-	bool bRawDisplayEngine = GetDefault<UContentBrowserSettings>()->GetDisplayEngineFolder( true );
+	check(IsToggleShowEngineContentAllowed());
 
-	// Only if both these flags are false when toggling we want to enable the flag, otherwise we're toggling off
-	if ( !bDisplayEngine && !bRawDisplayEngine )
+	bOverrideShowEngineContent.Reset();
+
+	bool bNewState = !GetDefault<UContentBrowserSettings>()->GetDisplayEngineFolder();
+
+	if (FContentBrowserInstanceConfig* EditorConfig = GetContentBrowserConfig())
 	{
-		GetMutableDefault<UContentBrowserSettings>()->SetDisplayEngineFolder( true );
+		bNewState = !EditorConfig->bShowEngineContent;
+		EditorConfig->bShowEngineContent = bNewState;
+		UContentBrowserConfig::Get()->SaveEditorConfig();
 	}
-	else
-	{
-		GetMutableDefault<UContentBrowserSettings>()->SetDisplayEngineFolder( false );
-		GetMutableDefault<UContentBrowserSettings>()->SetDisplayEngineFolder( false, true );
-	}	
+
+	GetMutableDefault<UContentBrowserSettings>()->SetDisplayEngineFolder(bNewState);
 	GetMutableDefault<UContentBrowserSettings>()->PostEditChange();
 }
 
 bool SAssetView::IsShowingEngineContent() const
 {
-	return bForceShowEngineContent || GetDefault<UContentBrowserSettings>()->GetDisplayEngineFolder();
+	if (bForceShowEngineContent)
+	{
+		return true;
+	}
+
+	if (bOverrideShowEngineContent.IsSet())
+	{
+		return bOverrideShowEngineContent.GetValue();
+	}
+
+	if (const FContentBrowserInstanceConfig* Config = GetContentBrowserConfig())
+	{
+		return Config->bShowEngineContent;
+	}
+
+	return GetDefault<UContentBrowserSettings>()->GetDisplayEngineFolder();
 }
 
 void SAssetView::ToggleShowDevelopersContent()
 {
-	bool bDisplayDev = GetDefault<UContentBrowserSettings>()->GetDisplayDevelopersFolder();
-	bool bRawDisplayDev = GetDefault<UContentBrowserSettings>()->GetDisplayDevelopersFolder( true );
+	check(IsToggleShowDevelopersContentAllowed());
 
-	// Only if both these flags are false when toggling we want to enable the flag, otherwise we're toggling off
-	if ( !bDisplayDev && !bRawDisplayDev )
+	bOverrideShowDeveloperContent.Reset();
+
+	bool bNewState = !GetDefault<UContentBrowserSettings>()->GetDisplayDevelopersFolder();
+
+	if (FContentBrowserInstanceConfig* EditorConfig = GetContentBrowserConfig())
 	{
-		GetMutableDefault<UContentBrowserSettings>()->SetDisplayDevelopersFolder( true );
+		bNewState = !EditorConfig->bShowDeveloperContent;
+		EditorConfig->bShowDeveloperContent = bNewState;
+		UContentBrowserConfig::Get()->SaveEditorConfig();
 	}
-	else
-	{
-		GetMutableDefault<UContentBrowserSettings>()->SetDisplayDevelopersFolder( false );
-		GetMutableDefault<UContentBrowserSettings>()->SetDisplayDevelopersFolder( false, true );
-	}	
+
+	GetMutableDefault<UContentBrowserSettings>()->SetDisplayDevelopersFolder(bNewState);
 	GetMutableDefault<UContentBrowserSettings>()->PostEditChange();
 }
 
@@ -2985,12 +3149,40 @@
 
 bool SAssetView::IsShowingDevelopersContent() const
 {
-	return IsToggleShowDevelopersContentAllowed() && GetDefault<UContentBrowserSettings>()->GetDisplayDevelopersFolder();
+	if (!IsToggleShowDevelopersContentAllowed())
+	{
+		return false;
+	}
+
+	if (bOverrideShowDeveloperContent.IsSet())
+	{
+		return bOverrideShowDeveloperContent.GetValue();
+	}
+
+	if (const FContentBrowserInstanceConfig* Config = GetContentBrowserConfig())
+	{
+		return Config->bShowDeveloperContent;
+	}
+
+	return GetDefault<UContentBrowserSettings>()->GetDisplayDevelopersFolder();
 }
 
 void SAssetView::ToggleShowLocalizedContent()
 {
-	GetMutableDefault<UContentBrowserSettings>()->SetDisplayL10NFolder(!GetDefault<UContentBrowserSettings>()->GetDisplayL10NFolder());
+	check(IsToggleShowLocalizedContentAllowed());
+
+	bool bNewState = !GetDefault<UContentBrowserSettings>()->GetDisplayL10NFolder();
+
+	bOverrideShowLocalizedContent.Reset();
+
+	if (FContentBrowserInstanceConfig* Config = GetContentBrowserConfig())
+	{
+		bNewState = !Config->bShowLocalizedContent;
+		Config->bShowLocalizedContent = bNewState;
+		UContentBrowserConfig::Get()->SaveEditorConfig();
+	}
+
+	GetMutableDefault<UContentBrowserSettings>()->SetDisplayL10NFolder(bNewState);
 	GetMutableDefault<UContentBrowserSettings>()->PostEditChange();
 }
 
@@ -3001,13 +3193,38 @@
 
 bool SAssetView::IsShowingLocalizedContent() const
 {
-	return IsToggleShowLocalizedContentAllowed() && GetDefault<UContentBrowserSettings>()->GetDisplayL10NFolder();
+	if (!IsToggleShowLocalizedContentAllowed())
+	{
+		return false;
+	}
+
+	if (bOverrideShowLocalizedContent.IsSet())
+	{
+		return bOverrideShowLocalizedContent.GetValue();
+	}
+
+	if (const FContentBrowserInstanceConfig* Config = GetContentBrowserConfig())
+	{
+		return Config->bShowLocalizedContent;
+	}
+
+	return GetDefault<UContentBrowserSettings>()->GetDisplayL10NFolder();
 }
 
 void SAssetView::ToggleShowFavorites()
 {
-	const bool bShowingFavorites = GetDefault<UContentBrowserSettings>()->GetDisplayFavorites();
-	GetMutableDefault<UContentBrowserSettings>()->SetDisplayFavorites(!bShowingFavorites);
+	check(IsToggleShowFavoritesAllowed());
+
+	bool bNewState = !GetDefault<UContentBrowserSettings>()->GetDisplayFavorites();
+
+	if (FContentBrowserInstanceConfig* Config = GetContentBrowserConfig())
+	{
+		bNewState = !Config->bShowFavorites;
+		Config->bShowFavorites = bNewState;
+		UContentBrowserConfig::Get()->SaveEditorConfig();
+	}
+
+	GetMutableDefault<UContentBrowserSettings>()->SetDisplayFavorites(bNewState);
 	GetMutableDefault<UContentBrowserSettings>()->PostEditChange();
 }
 
@@ -3018,13 +3235,33 @@
 
 bool SAssetView::IsShowingFavorites() const
 {
-	return IsToggleShowFavoritesAllowed() && GetDefault<UContentBrowserSettings>()->GetDisplayFavorites();
+	if (!IsToggleShowFavoritesAllowed())
+	{
+		return false;
+	}
+
+	if (const FContentBrowserInstanceConfig* Config = GetContentBrowserConfig())
+	{
+		return Config->bShowFavorites;
+	}
+
+	return GetDefault<UContentBrowserSettings>()->GetDisplayFavorites();
 }
 
 void SAssetView::ToggleDockCollections()
 {
-	const bool bDockCollections = GetDefault<UContentBrowserSettings>()->GetDockCollections();
-	GetMutableDefault<UContentBrowserSettings>()->SetDockCollections(!bDockCollections);
+	check(IsToggleDockCollectionsAllowed()); 
+
+	bool bNewState = !GetDefault<UContentBrowserSettings>()->GetDockCollections();
+
+	if (FContentBrowserInstanceConfig* Config = GetContentBrowserConfig())
+	{
+		bNewState = !Config->bCollectionsDocked;
+		Config->bCollectionsDocked = bNewState;
+		UContentBrowserConfig::Get()->SaveEditorConfig();
+	}
+
+	GetMutableDefault<UContentBrowserSettings>()->SetDockCollections(bNewState);
 	GetMutableDefault<UContentBrowserSettings>()->PostEditChange();
 }
 
@@ -3035,13 +3272,33 @@
 
 bool SAssetView::HasDockedCollections() const
 {
-	return IsToggleDockCollectionsAllowed() && GetDefault<UContentBrowserSettings>()->GetDockCollections();
+	if (!IsToggleDockCollectionsAllowed())
+	{
+		return false;
+	}
+
+	if (const FContentBrowserInstanceConfig* Config = GetContentBrowserConfig())
+	{
+		return Config->bCollectionsDocked;
+	}
+
+	return GetDefault<UContentBrowserSettings>()->GetDockCollections();
 }
 
 void SAssetView::ToggleShowCppContent()
 {
-	const bool bDisplayCppFolders = GetDefault<UContentBrowserSettings>()->GetDisplayCppFolders();
-	GetMutableDefault<UContentBrowserSettings>()->SetDisplayCppFolders(!bDisplayCppFolders);
+	check(IsToggleShowCppContentAllowed());
+
+	bool bNewState = !GetDefault<UContentBrowserSettings>()->GetDisplayCppFolders();
+
+	if (FContentBrowserInstanceConfig* Config = GetContentBrowserConfig())
+	{
+		bNewState = !Config->bShowCppFolders;
+		Config->bShowCppFolders = bNewState;
+		UContentBrowserConfig::Get()->SaveEditorConfig();
+	}
+
+	GetMutableDefault<UContentBrowserSettings>()->SetDisplayCppFolders(bNewState);
 	GetMutableDefault<UContentBrowserSettings>()->PostEditChange();
 }
 
@@ -3052,13 +3309,33 @@
 
 bool SAssetView::IsShowingCppContent() const
 {
-	return IsToggleShowCppContentAllowed() && GetDefault<UContentBrowserSettings>()->GetDisplayCppFolders();
+	if (!IsToggleShowCppContentAllowed())
+	{
+		return false;
+	}
+
+	if (const FContentBrowserInstanceConfig* Config = GetContentBrowserConfig())
+	{
+		return Config->bShowCppFolders;
+	}
+	
+	return GetDefault<UContentBrowserSettings>()->GetDisplayCppFolders();
 }
 
 void SAssetView::ToggleIncludeClassNames()
 {
-	const bool bIncludeClassNames = GetDefault<UContentBrowserSettings>()->GetIncludeClassNames();
-	GetMutableDefault<UContentBrowserSettings>()->SetIncludeClassNames(!bIncludeClassNames);
+	check(IsToggleIncludeClassNamesAllowed());
+
+	bool bNewState = !GetDefault<UContentBrowserSettings>()->GetIncludeClassNames();
+
+	if (FContentBrowserInstanceConfig* Config = GetContentBrowserConfig())
+	{
+		bNewState = !Config->bSearchClasses;
+		Config->bSearchClasses = bNewState;
+		UContentBrowserConfig::Get()->SaveEditorConfig();
+	}
+
+	GetMutableDefault<UContentBrowserSettings>()->SetIncludeClassNames(bNewState);
 	GetMutableDefault<UContentBrowserSettings>()->PostEditChange();
 
 	OnSearchOptionsChanged.ExecuteIfBound();
@@ -3071,13 +3348,33 @@
 
 bool SAssetView::IsIncludingClassNames() const
 {
-	return IsToggleIncludeClassNamesAllowed() && GetDefault<UContentBrowserSettings>()->GetIncludeClassNames();
+	if (!IsToggleIncludeClassNamesAllowed())
+	{
+		return false;
+	}
+
+	if (const FContentBrowserInstanceConfig* Config = GetContentBrowserConfig())
+	{
+		return Config->bSearchClasses;
+	}
+	
+	return GetDefault<UContentBrowserSettings>()->GetIncludeClassNames();
 }
 
 void SAssetView::ToggleIncludeAssetPaths()
 {
-	const bool bIncludeAssetPaths = GetDefault<UContentBrowserSettings>()->GetIncludeAssetPaths();
-	GetMutableDefault<UContentBrowserSettings>()->SetIncludeAssetPaths(!bIncludeAssetPaths);
+	check(IsToggleIncludeAssetPathsAllowed());
+
+	bool bNewState = !GetDefault<UContentBrowserSettings>()->GetIncludeAssetPaths();
+
+	if (FContentBrowserInstanceConfig* Config = GetContentBrowserConfig())
+	{
+		bNewState = !Config->bSearchAssetPaths;
+		Config->bSearchAssetPaths = bNewState;
+		UContentBrowserConfig::Get()->SaveEditorConfig();
+	}
+
+	GetMutableDefault<UContentBrowserSettings>()->SetIncludeAssetPaths(bNewState);
 	GetMutableDefault<UContentBrowserSettings>()->PostEditChange();
 
 	OnSearchOptionsChanged.ExecuteIfBound();
@@ -3090,13 +3387,33 @@
 
 bool SAssetView::IsIncludingAssetPaths() const
 {
-	return IsToggleIncludeAssetPathsAllowed() && GetDefault<UContentBrowserSettings>()->GetIncludeAssetPaths();
+	if (!IsToggleIncludeAssetPathsAllowed())
+	{
+		return false;
+	}
+
+	if (const FContentBrowserInstanceConfig* Config = GetContentBrowserConfig())
+	{
+		return Config->bSearchAssetPaths;
+	}
+	
+	return GetDefault<UContentBrowserSettings>()->GetIncludeAssetPaths();
 }
 
 void SAssetView::ToggleIncludeCollectionNames()
 {
-	const bool bIncludeCollectionNames = GetDefault<UContentBrowserSettings>()->GetIncludeCollectionNames();
-	GetMutableDefault<UContentBrowserSettings>()->SetIncludeCollectionNames(!bIncludeCollectionNames);
+	check(IsToggleIncludeCollectionNamesAllowed());
+
+	bool bNewState = !GetDefault<UContentBrowserSettings>()->GetIncludeCollectionNames();
+
+	if (FContentBrowserInstanceConfig* Config = GetContentBrowserConfig())
+	{
+		bNewState = !Config->bSearchCollections;
+		Config->bSearchCollections = bNewState;
+		UContentBrowserConfig::Get()->SaveEditorConfig();
+	}
+
+	GetMutableDefault<UContentBrowserSettings>()->SetIncludeCollectionNames(bNewState);
 	GetMutableDefault<UContentBrowserSettings>()->PostEditChange();
 
 	OnSearchOptionsChanged.ExecuteIfBound();
@@ -3109,10 +3426,19 @@
 
 bool SAssetView::IsIncludingCollectionNames() const
 {
-	return IsToggleIncludeCollectionNamesAllowed() && GetDefault<UContentBrowserSettings>()->GetIncludeCollectionNames();
+	if (!IsToggleIncludeCollectionNamesAllowed())
+	{
+		return false;
+	}
+	
+	if (const FContentBrowserInstanceConfig* Config = GetContentBrowserConfig())
+	{
+		return Config->bSearchCollections;
+	}
+	
+	return GetDefault<UContentBrowserSettings>()->GetIncludeCollectionNames();
 }
 
-
 void SAssetView::SetCurrentViewType(EAssetViewType::Type NewType)
 {
 	if ( ensure(NewType != EAssetViewType::MAX) && NewType != CurrentViewType )
@@ -3145,12 +3471,21 @@
 			RefreshFilteredItems();
 			SortList();
 		}
+
+		if (FAssetViewInstanceConfig* Config = GetAssetViewConfig())
+		{
+			Config->ViewType = (int32) NewType;
+			UAssetViewConfig::Get()->SaveEditorConfig();
+		}
 	}
 }
 
 void SAssetView::SetCurrentThumbnailSize(EThumbnailSize NewThumbnailSize)
 {
-	OnThumbnailSizeChanged(NewThumbnailSize);
+	if (ThumbnailSize != NewThumbnailSize)
+	{
+		OnThumbnailSizeChanged(NewThumbnailSize);
+	}
 }
 
 void SAssetView::SetCurrentViewTypeFromMenu(EAssetViewType::Type NewType)
@@ -4004,6 +4339,13 @@
 void SAssetView::OnThumbnailSizeChanged(EThumbnailSize NewThumbnailSize)
 {
 	ThumbnailSize = NewThumbnailSize;
+
+	if (FAssetViewInstanceConfig* Config = GetAssetViewConfig())
+	{
+		Config->ThumbnailSize = (int32) NewThumbnailSize;
+		UAssetViewConfig::Get()->SaveEditorConfig();
+	}
+
 	RefreshList();
 }
 
@@ -4393,6 +4735,13 @@
 		}
 	}
 	
+	if (FAssetViewInstanceConfig* Config = GetAssetViewConfig())
+	{
+		Config->HiddenColumns.Reset();
+		Algo::Transform(HiddenColumnNames, Config->HiddenColumns, [](const FString& Str) { return FName(*Str); });
+
+		UAssetViewConfig::Get()->SaveEditorConfig();
+	}
 }
 
 bool SAssetView::ShouldColumnGenerateWidget(const FString ColumnName) const
@@ -4413,6 +4762,26 @@
 	}
 }
 
+void SAssetView::OverrideShowEngineContent( bool bShow )
+{
+	bOverrideShowEngineContent = bShow;
+}
+
+void SAssetView::OverrideShowDeveloperContent( bool bShow )
+{
+	bOverrideShowDeveloperContent = bShow;
+}
+
+void SAssetView::OverrideShowPluginContent( bool bShow )
+{
+	bOverrideShowPluginContent = bShow;
+}
+
+void SAssetView::OverrideShowLocalizedContent( bool bShow )
+{
+	bOverrideShowLocalizedContent = bShow;
+}
+
 void SAssetView::HandleItemDataUpdated(TArrayView<const FContentBrowserItemDataUpdate> InUpdatedItems)
 {
 	if (InUpdatedItems.Num() == 0)
diff --git a/Engine/Source/Editor/ContentBrowser/Private/SContentBrowser.cpp b/Engine/Source/Editor/ContentBrowser/Private/SContentBrowser.cpp
--- a/Engine/Source/Editor/ContentBrowser/Private/SContentBrowser.cpp
+++ b/Engine/Source/Editor/ContentBrowser/Private/SContentBrowser.cpp
@@ -19,6 +19,7 @@
 #include "Containers/StringFwd.h"
 #include "Containers/StringView.h"
 #include "ContentBrowserCommands.h"
+#include "ContentBrowserConfig.h"
 #include "ContentBrowserDataFilter.h"
 #include "ContentBrowserDataSource.h"
 #include "ContentBrowserDataSubsystem.h"
@@ -195,6 +196,12 @@ SContentBrowser::~SContentBrowser()
 BEGIN_SLATE_FUNCTION_BUILD_OPTIMIZATION
 void SContentBrowser::Construct( const FArguments& InArgs, const FName& InInstanceName, const FContentBrowserConfig* Config )
 {
+	InstanceName = InInstanceName;
+
+	UContentBrowserConfig::Initialize();
+	UContentBrowserConfig::Get()->LoadEditorConfig();
+	const FContentBrowserInstanceConfig* EditorConfig = CreateEditorConfigIfRequired();
+
 	if ( InArgs._ContainingTab.IsValid() )
 	{
 		// For content browsers that are placed in tabs, save settings when the tab is closing.
@@ -203,12 +210,10 @@ void SContentBrowser::Construct( const FArguments& InArgs, const FName& InInstan
 		InArgs._ContainingTab->SetOnTabClosed( SDockTab::FOnTabClosedCallback::CreateSP( this, &SContentBrowser::OnContainingTabClosed ) );
 		InArgs._ContainingTab->SetOnTabActivated( SDockTab::FOnTabActivatedCallback::CreateSP( this, &SContentBrowser::OnContainingTabActivated ) );
 	}
-	
 
  	bIsLocked = InArgs._InitiallyLocked;
 	bCanSetAsPrimaryBrowser = Config != nullptr ? Config->bCanSetAsPrimaryBrowser : true;
 	bIsDrawer = InArgs._IsDrawer;
-	this->InstanceName = InInstanceName;
 
 	HistoryManager.SetOnApplyHistoryData(FOnApplyHistoryData::CreateSP(this, &SContentBrowser::OnApplyHistoryData));
 	HistoryManager.SetOnUpdateHistoryData(FOnUpdateHistoryData::CreateSP(this, &SContentBrowser::OnUpdateHistoryData));
@@ -577,22 +582,21 @@ void SContentBrowser::Construct( const FArguments& InArgs, const FName& InInstan
 	// Load settings if they were specified
 	LoadSettings(InInstanceName);
 
-	if( Config != nullptr )
+	if (Config != nullptr)
 	{
 		// Make sure the sources view is initially visible if we were asked to show it
-		if( ( bSourcesViewExpanded && ( !Config->bExpandSourcesView || !Config->bUseSourcesView ) ) ||
-			( !bSourcesViewExpanded && Config->bExpandSourcesView && Config->bUseSourcesView ) )
-		{
-			SourcesViewExpandClicked();
-		}
+		SetSourcesViewExpanded(Config->bExpandSourcesView && Config->bUseSourcesView);
 	}
 	else
 	{
 		// in case we do not have a config, see what the global default settings are for the Sources Panel
-		if (!bSourcesViewExpanded && GetDefault<UContentBrowserSettings>()->bOpenSourcesPanelByDefault)
+		bool bSourcesExpanded = GetDefault<UContentBrowserSettings>()->bOpenSourcesPanelByDefault;
+		if (EditorConfig != nullptr)
 		{
-			SourcesViewExpandClicked();
+			bSourcesExpanded = EditorConfig->bSourcesExpanded;
 		}
+
+		SetSourcesViewExpanded(bSourcesExpanded);
 	}
 
 	// Bindings to manage history when items are deleted
@@ -673,6 +677,11 @@ void SContentBrowser::BindCommands()
 
 EVisibility SContentBrowser::GetFavoriteFolderVisibility() const
 {
+	if (const FContentBrowserInstanceConfig* Config = GetConstInstanceConfig())
+	{
+		return Config->bShowFavorites ? EVisibility::Visible : EVisibility::Collapsed;
+	}
+
 	return GetDefault<UContentBrowserSettings>()->GetDisplayFavorites() ? EVisibility::Visible : EVisibility::Collapsed;
 }
 
@@ -688,27 +697,34 @@ EVisibility SContentBrowser::GetLockButtonVisibility() const
 
 bool SContentBrowser::IsCollectionViewDocked() const
 {
+	if (const FContentBrowserInstanceConfig* Config = GetConstInstanceConfig())
+	{
+		return Config->bCollectionsDocked;
+	}
+
 	return GetDefault<UContentBrowserSettings>()->GetDockCollections();
 }
 
 void SContentBrowser::ToggleFolderFavorite(const TArray<FString>& FolderPaths)
 {
 	bool bAddedFavorite = false;
-	for (FString FolderPath : FolderPaths)
+	for (const FString& FolderPath : FolderPaths)
 	{
 		if (ContentBrowserUtils::IsFavoriteFolder(FolderPath))
 		{
-			ContentBrowserUtils::RemoveFavoriteFolder(FolderPath, false);
+			ContentBrowserUtils::RemoveFavoriteFolder(FolderPath);
 		}
 		else
 		{
-			ContentBrowserUtils::AddFavoriteFolder(FolderPath, false);
+			ContentBrowserUtils::AddFavoriteFolder(FolderPath);
 			bAddedFavorite = true;
 		}
 	}
-	GConfig->Flush(false, GEditorPerProjectIni);
+
+	FavoritePathViewPtr->SaveSettings(GEditorPerProjectIni, SettingsIniSection, InstanceName.ToString() + TEXT(".Favorites"));
+	
 	FavoritePathViewPtr->Populate();
-	if(bAddedFavorite)
+	if (bAddedFavorite)
 	{	
 		FavoritePathViewPtr->SetSelectedPaths(FolderPaths);
 		if (GetFavoriteFolderVisibility() == EVisibility::Collapsed)
@@ -987,6 +1003,20 @@ TSharedRef<SWidget> SContentBrowser::CreateAssetView(const FContentBrowserConfig
 	return AssetViewBorder.ToSharedRef();
 }
 
+void SContentBrowser::SetFavoritesExpanded(bool bExpanded)
+{
+	if (!bExpanded)
+	{
+		FavoritesSearchToggleButton->SetExpanded(false); 
+	}
+
+	if (FContentBrowserInstanceConfig* EditorConfig = GetMutableInstanceConfig())
+	{
+		EditorConfig->bFavoritesExpanded = bExpanded;
+		UContentBrowserConfig::Get()->SaveEditorConfig();
+	}
+}
+
 TSharedRef<SWidget> SContentBrowser::CreateFavoritesView(const FContentBrowserConfig* Config)
 {
 	return
@@ -997,7 +1027,7 @@ TSharedRef<SWidget> SContentBrowser::CreateFavoritesView(const FContentBrowserCo
 		.Visibility(this, &SContentBrowser::GetFavoriteFolderVisibility)
 		.Padding(0)
 		.AllowAnimatedTransition(false)
-		.OnAreaExpansionChanged_Lambda([this](bool bIsExpanded) { if (!bIsExpanded) FavoritesSearchToggleButton->SetExpanded(false); })
+		.OnAreaExpansionChanged(this, &SContentBrowser::SetFavoritesExpanded)
 		.HeaderContent()
 		[
 			SNew(SHorizontalBox)
@@ -1048,8 +1078,24 @@ TSharedRef<SWidget> SContentBrowser::CreateFavoritesView(const FContentBrowserCo
 		];
 }
 
+void SContentBrowser::SetPathViewExpanded(bool bExpanded)
+{
+	if (!bExpanded)
+	{
+		PathSearchToggleButton->SetExpanded(false);
+	}
+
+	if (FContentBrowserInstanceConfig* EditorConfig = GetMutableInstanceConfig())
+	{
+		EditorConfig->PathView.bExpanded = bExpanded;
+		UContentBrowserConfig::Get()->SaveEditorConfig();
+	}
+}
+
 TSharedRef<SWidget> SContentBrowser::CreatePathView(const FContentBrowserConfig* Config)
 {	
+	FContentBrowserInstanceConfig* InstanceConfig = GetMutableInstanceConfig();
+
 	return
 		SAssignNew(PathArea, SExpandableArea)
 		.BorderImage(FAppStyle::Get().GetBrush("Brushes.Header"))
@@ -1057,7 +1103,7 @@ TSharedRef<SWidget> SContentBrowser::CreatePathView(const FContentBrowserConfig*
 		.HeaderPadding(FMargin(4.0f, 0.0f))
 		.Padding(0)
 		.AllowAnimatedTransition(false)
-		.OnAreaExpansionChanged_Lambda([this](bool bIsExpanded) { if (!bIsExpanded) PathSearchToggleButton->SetExpanded(false); })
+		.OnAreaExpansionChanged(this, &SContentBrowser::SetPathViewExpanded)
 		.HeaderContent()
 		[
 			SNew(SHorizontalBox)
@@ -1108,6 +1154,7 @@ TSharedRef<SWidget> SContentBrowser::CreatePathView(const FContentBrowserConfig*
 					.AddMetaData<FTagMetaData>(FTagMetaData(TEXT("ContentBrowserSources")))
 					.ExternalSearch(SourcesSearch)
 					.PluginPathFilters(PluginPathFilters)
+					.EditorConfig(InstanceConfig != nullptr ? &InstanceConfig->PathView : nullptr)
 				]
 			]
 		];
@@ -1225,42 +1272,49 @@ void SContentBrowser::ExtendViewOptionsMenu(const FContentBrowserConfig* Config)
 {
 	UToolMenu* Menu = UToolMenus::Get()->ExtendMenu("ContentBrowser.AssetViewOptions");
 
-	if (Config == nullptr || Config->bCanShowLockButton)
+	const bool bShowLockButton = Config == nullptr || Config->bCanShowLockButton;
+	const bool bShowSourcesView = Config == nullptr || Config->bUseSourcesView;
+
+	if (!bShowLockButton && !bShowSourcesView)
 	{
-		Menu->AddDynamicSection("ContentBrowserViewOptionsSection", FNewToolMenuDelegate::CreateLambda([](UToolMenu* InMenu)
+		return;
+	}
+
+	Menu->AddDynamicSection("ContentBrowserViewOptionsSection", FNewToolMenuDelegate::CreateLambda([bShowLockButton, bShowSourcesView](UToolMenu* InMenu)
+		{
+			if (UContentBrowserAssetViewContextMenuContext* Context = InMenu->FindContext<UContentBrowserAssetViewContextMenuContext>())
 			{
-				if (UContentBrowserAssetViewContextMenuContext* Context = InMenu->FindContext<UContentBrowserAssetViewContextMenuContext>())
+				if (TSharedPtr<SContentBrowser> ContentBrowser = Context->OwningContentBrowser.Pin())
 				{
-					if (TSharedPtr<SContentBrowser> ContentBrowser = Context->OwningContentBrowser.Pin())
+					if (bShowLockButton)
 					{
-						{
-							FToolMenuSection& Section = InMenu->AddSection("Locking", LOCTEXT("LockingMenuHeader", "Locking"), FToolMenuInsert("AssetViewType", EToolMenuInsertType::After));
-							Section.AddMenuEntry(
-								"ToggleLock",
-								TAttribute<FText>(ContentBrowser.ToSharedRef(), &SContentBrowser::GetLockMenuText),
-								LOCTEXT("LockToggleTooltip", "Toggle lock. If locked, this browser will ignore Find in Content Browser requests."),
-								TAttribute<FSlateIcon>(),
-								FUIAction(FExecuteAction::CreateLambda([ContentBrowser = Context->OwningContentBrowser]() {ContentBrowser.Pin()->ToggleLockClicked(); })));
-						}
+						FToolMenuSection& Section = InMenu->AddSection("Locking", LOCTEXT("LockingMenuHeader", "Locking"), FToolMenuInsert("AssetViewType", EToolMenuInsertType::After));
+						Section.AddMenuEntry(
+							"ToggleLock",
+							TAttribute<FText>(ContentBrowser.ToSharedRef(), &SContentBrowser::GetLockMenuText),
+							LOCTEXT("LockToggleTooltip", "Toggle lock. If locked, this browser will ignore Find in Content Browser requests."),
+							TAttribute<FSlateIcon>(),
+							FUIAction(FExecuteAction::CreateLambda([ContentBrowser = Context->OwningContentBrowser]() { ContentBrowser.Pin()->ToggleLockClicked(); })));
+					}
 
-						{
-							FToolMenuSection& Section = InMenu->FindOrAddSection("View");
-							Section.AddMenuEntry(
-								"ToggleSources",
-								LOCTEXT("ToggleSourcesView", "Show Sources Panel"),
-								LOCTEXT("ToggleSourcesView_Tooltip", "Show or hide the sources panel"),
-								TAttribute<FSlateIcon>(),
-								FUIAction(
-									FExecuteAction::CreateLambda([ContentBrowser = Context->OwningContentBrowser]() {ContentBrowser.Pin()->SourcesViewExpandClicked(); }),
-									FCanExecuteAction(),
-									FIsActionChecked::CreateLambda([ContentBrowser = Context->OwningContentBrowser]() {return ContentBrowser.Pin()->bSourcesViewExpanded; })),
-								EUserInterfaceActionType::Check);
-						}
+					if (bShowSourcesView)
+					{
+						FToolMenuSection& Section = InMenu->FindOrAddSection("View");
+						Section.AddMenuEntry(
+							"ToggleSources",
+							LOCTEXT("ToggleSourcesView", "Show Sources Panel"),
+							LOCTEXT("ToggleSourcesView_Tooltip", "Show or hide the sources panel"),
+							TAttribute<FSlateIcon>(),
+							FUIAction(
+								FExecuteAction::CreateLambda([ContentBrowser = Context->OwningContentBrowser]() { ContentBrowser.Pin()->SourcesViewExpandClicked(); }),
+								FCanExecuteAction(),
+								FIsActionChecked::CreateLambda([ContentBrowser = Context->OwningContentBrowser]() { return ContentBrowser.Pin()->bSourcesViewExpanded; })),
+							EUserInterfaceActionType::ToggleButton);
 					}
 				}
 			}
-		));
-	}
+		}
+	));
 }
 
 void SContentBrowser::RegisterContentBrowserToolBar()
@@ -1407,28 +1461,28 @@ void SContentBrowser::PrepareToSyncItems(TArrayView<const FContentBrowserItem> I
 			if (!bDisplayDev && ContentBrowserUtils::IsItemDeveloperContent(ItemToSync))
 			{
 				bDisplayDev = true;
-				GetMutableDefault<UContentBrowserSettings>()->SetDisplayDevelopersFolder(true, true);
+				AssetViewPtr->OverrideShowDeveloperContent(true);
 				bRepopulate = true;
 			}
 
 			if (!bDisplayEngine && ContentBrowserUtils::IsItemEngineContent(ItemToSync))
 			{
 				bDisplayEngine = true;
-				GetMutableDefault<UContentBrowserSettings>()->SetDisplayEngineFolder(true, true);
+				AssetViewPtr->OverrideShowEngineContent(true);
 				bRepopulate = true;
 			}
 
 			if (!bDisplayPlugins && ContentBrowserUtils::IsItemPluginContent(ItemToSync))
 			{
 				bDisplayPlugins = true;
-				GetMutableDefault<UContentBrowserSettings>()->SetDisplayPluginFolders(true, true);
+				AssetViewPtr->OverrideShowPluginContent(true);
 				bRepopulate = true;
 			}
 
 			if (!bDisplayLocalized && ContentBrowserUtils::IsItemLocalizedContent(ItemToSync))
 			{
 				bDisplayLocalized = true;
-				GetMutableDefault<UContentBrowserSettings>()->SetDisplayL10NFolder(true);
+				AssetViewPtr->OverrideShowLocalizedContent(true);
 				bRepopulate = true;
 			}
 
@@ -3060,38 +3114,45 @@ EVisibility SContentBrowser::GetSourcesViewVisibility() const
 	return bSourcesViewExpanded ? EVisibility::Visible : EVisibility::Collapsed;
 }
 
-FReply SContentBrowser::SourcesViewExpandClicked()
+void SContentBrowser::SetSourcesViewExpanded(bool bExpanded)
 {
-	bSourcesViewExpanded = !bSourcesViewExpanded;
+	bSourcesViewExpanded = bExpanded;
+
+	if (FContentBrowserInstanceConfig* EditorConfig = GetMutableInstanceConfig())
+	{
+		EditorConfig->bSourcesExpanded = bSourcesViewExpanded;
+		UContentBrowserConfig::Get()->SaveEditorConfig();
+	}
+
 
 	// Notify 'Sources View Expanded' delegate
 	FContentBrowserModule& ContentBrowserModule = FModuleManager::GetModuleChecked<FContentBrowserModule>( TEXT("ContentBrowser") );
 	FContentBrowserModule::FOnSourcesViewChanged& SourcesViewChangedDelegate = ContentBrowserModule.GetOnSourcesViewChanged();
-	if(SourcesViewChangedDelegate.IsBound())
+	if (SourcesViewChangedDelegate.IsBound())
 	{
 		SourcesViewChangedDelegate.Broadcast(bSourcesViewExpanded);
 	}
+}
 
+FReply SContentBrowser::SourcesViewExpandClicked()
+{
+	SetSourcesViewExpanded(!bSourcesViewExpanded);
 	return FReply::Handled();
 }
 
 EVisibility SContentBrowser::GetSourcesSwitcherVisibility() const
 {
-	return GetDefault<UContentBrowserSettings>()->GetDockCollections() ? EVisibility::Collapsed : EVisibility::Visible;
+	return IsCollectionViewDocked() ? EVisibility::Collapsed : EVisibility::Visible;
 }
 
-
 FReply SContentBrowser::OnSourcesSwitcherClicked()
 {
-	// This only works because we only have two switcher types
-	ActiveSourcesWidgetIndex = !ActiveSourcesWidgetIndex;
+	ActiveSourcesWidgetIndex = (ActiveSourcesWidgetIndex + 1) % SourcesWidgetSwitcher->GetNumWidgets();
 	SourcesWidgetSwitcher->SetActiveWidgetIndex(ActiveSourcesWidgetIndex);
 
 	return FReply::Handled();
 }
 
-
-
 void SContentBrowser::OnContentBrowserSettingsChanged(FName PropertyName)
 {
 	if (PropertyName.IsNone())
@@ -3511,7 +3572,7 @@ void SContentBrowser::UpdatePath()
 	}
 	else if ( SourcesData.HasCollections() )
 	{
-		NewSourcesWidgetIndex = GetDefault<UContentBrowserSettings>()->GetDockCollections() ? ContentBrowserSourcesWidgetSwitcherIndex::PathView : ContentBrowserSourcesWidgetSwitcherIndex::CollectionsView;
+		NewSourcesWidgetIndex = IsCollectionViewDocked() ? ContentBrowserSourcesWidgetSwitcherIndex::PathView : ContentBrowserSourcesWidgetSwitcherIndex::CollectionsView;
 
 		FCollectionManagerModule& CollectionManagerModule = FCollectionManagerModule::GetModule();
 		TArray<FCollectionNameType> CollectionPathItems;
@@ -4246,5 +4307,79 @@ void SContentBrowser::OpenNewContentBrowser()
 	const TArray<FContentBrowserItem> SelectedFolders = PathContextMenu->GetSelectedFolders();
 	FContentBrowserSingleton::Get().SyncBrowserToItems(SelectedFolders, false, true, NAME_None, true);
 }
+ 
+const FContentBrowserInstanceConfig* SContentBrowser::GetConstInstanceConfig() const
+{
+	if (InstanceName.IsNone())
+	{
+		return nullptr;
+	}
+
+	UContentBrowserConfig* Config = UContentBrowserConfig::Get();
+	if (Config == nullptr)
+	{
+		return nullptr;
+	}
+
+	const FContentBrowserInstanceConfig* InstanceConfig = Config->Instances.Find(InstanceName);
+	return InstanceConfig;
+}
+
+FContentBrowserInstanceConfig* SContentBrowser::GetMutableInstanceConfig()
+{
+	if (InstanceName.IsNone())
+	{
+		return nullptr;
+	}
+
+	UContentBrowserConfig* Config = UContentBrowserConfig::Get();
+	if (Config == nullptr)
+	{
+		return nullptr;
+	}
+
+	FContentBrowserInstanceConfig* InstanceConfig = Config->Instances.Find(InstanceName);
+	return InstanceConfig;
+}
+
+FContentBrowserInstanceConfig* SContentBrowser::CreateEditorConfigIfRequired()
+{
+	UContentBrowserConfig* Config = UContentBrowserConfig::Get();
+	if (Config == nullptr)
+	{
+		return nullptr;
+	}
+
+	FContentBrowserInstanceConfig* InstanceConfig = Config->Instances.Find(InstanceName);
+	if (InstanceConfig != nullptr)
+	{
+		return InstanceConfig;
+	}
+
+	InstanceConfig = &Config->Instances.Add(InstanceName, FContentBrowserInstanceConfig());
+
+	const UContentBrowserSettings* Settings = GetDefault<UContentBrowserSettings>();
+	InstanceConfig->bShowEngineContent = Settings->GetDisplayEngineFolder();
+	InstanceConfig->bShowDeveloperContent = Settings->GetDisplayDevelopersFolder();
+	InstanceConfig->bShowLocalizedContent = Settings->GetDisplayL10NFolder();
+	InstanceConfig->bShowPluginContent = Settings->GetDisplayPluginFolders();
+	InstanceConfig->bShowFolders = Settings->DisplayFolders;
+	InstanceConfig->bShowAllFolder = Settings->bShowAllFolder;
+	InstanceConfig->bShowEmptyFolders = Settings->DisplayEmptyFolders;
+	InstanceConfig->bShowCppFolders = Settings->GetDisplayCppFolders();
+	InstanceConfig->bOrganizeFolders = Settings->bOrganizeFolders;
+	InstanceConfig->bCollectionsDocked = Settings->GetDockCollections();
+	InstanceConfig->bFavoritesExpanded = Settings->GetDisplayFavorites();
+	InstanceConfig->bSourcesExpanded = Settings->bOpenSourcesPanelByDefault;
+	InstanceConfig->bRealTimeThumbnails = Settings->RealTimeThumbnails;
+	InstanceConfig->bSearchAssetPaths = Settings->GetIncludeAssetPaths();
+	InstanceConfig->bSearchClasses = Settings->GetIncludeClassNames();
+	InstanceConfig->bSearchCollections = Settings->GetIncludeCollectionNames();
+	InstanceConfig->bFilterRecursively = Settings->FilterRecursively;
+
+	UContentBrowserConfig::Get()->SaveEditorConfig();
+
+	return InstanceConfig;
+}
 
 #undef LOCTEXT_NAMESPACE
diff --git a/Engine/Source/Editor/ContentBrowser/Private/SContentBrowser.h b/Engine/Source/Editor/ContentBrowser/Private/SContentBrowser.h
--- a/Engine/Source/Editor/ContentBrowser/Private/SContentBrowser.h
+++ b/Engine/Source/Editor/ContentBrowser/Private/SContentBrowser.h
@@ -39,6 +39,7 @@ class SBorder;
 class SCollectionView;
 class SComboButton;
 class SExpandableArea;
+class SFavoritePathView;
 class SFilterList;
 class SPathView;
 class SSearchToggleButton;
@@ -51,6 +52,7 @@ class UToolMenu;
 struct FAssetData;
 struct FAssetSearchBoxSuggestion;
 struct FCollectionNameType;
+struct FContentBrowserInstanceConfig;
 struct FContentBrowserItem;
 struct FContentBrowserItemPath;
 struct FGeometry;
@@ -377,6 +379,9 @@ private:
 
 	/** Gets the visibility state of the asset tree */
 	EVisibility GetSourcesViewVisibility() const;
+	
+	/** Set whether or not we the sources view is expanded. */
+	void SetSourcesViewExpanded(bool bExpanded);
 
 	/** Handler for clicking the tree expand/collapse button */
 	FReply SourcesViewExpandClicked();
@@ -573,6 +578,9 @@ private:
 	TSharedRef<SWidget> CreateDockedCollectionsView(const FContentBrowserConfig* Config);
 	TSharedRef<SWidget> CreateDrawerDockButton(const FContentBrowserConfig* Config);
 
+	void SetFavoritesExpanded(bool bExpanded);
+	void SetPathViewExpanded(bool bExpanded);
+
 	/** Adds menu options to the view menu */
 	void ExtendViewOptionsMenu(const FContentBrowserConfig* Config);
 
@@ -591,6 +599,15 @@ private:
 	/** Called when the layout of the SFilterList is changing */
 	void OnFilterBarLayoutChanging(EFilterBarLayout NewLayout);
 
+	/** Fetch the const config for this content browser instance. */
+	const FContentBrowserInstanceConfig* GetConstInstanceConfig() const;
+
+	/** Fetch the mutable config for this content browser instance. */
+	FContentBrowserInstanceConfig* GetMutableInstanceConfig();
+
+	/** Initialize an editor config for this instance if one does not exist. */
+	FContentBrowserInstanceConfig* CreateEditorConfigIfRequired();
+
 private:
 
 	/** The tab that contains this browser */
@@ -618,7 +635,7 @@ private:
 	TSharedPtr<SPathView> PathViewPtr;
 
 	/** The favorites tree widget */
-	TSharedPtr<class SFavoritePathView> FavoritePathViewPtr;
+	TSharedPtr<SFavoritePathView> FavoritePathViewPtr;
 
 	/** The collection widget */
 	TSharedPtr<SCollectionView> CollectionViewPtr;
diff --git a/Engine/Source/Editor/ContentBrowser/Private/SPathView.cpp b/Engine/Source/Editor/ContentBrowser/Private/SPathView.cpp
--- a/Engine/Source/Editor/ContentBrowser/Private/SPathView.cpp
+++ b/Engine/Source/Editor/ContentBrowser/Private/SPathView.cpp
@@ -6,6 +6,7 @@
 #include "AssetToolsModule.h"
 #include "Containers/ContainerAllocationPolicies.h"
 #include "Containers/StringView.h"
+#include "ContentBrowserConfig.h"
 #include "ContentBrowserDataSource.h"
 #include "ContentBrowserDataSubsystem.h"
 #include "ContentBrowserItemData.h"
@@ -134,6 +135,8 @@ SPathView::~SPathView()
 
 void SPathView::Construct( const FArguments& InArgs )
 {
+	EditorConfig = InArgs._EditorConfig;
+
 	OnItemSelectionChanged = InArgs._OnItemSelectionChanged;
 	bAllowContextMenu = InArgs._AllowContextMenu;
 	OnGetItemContextMenu = InArgs._OnGetItemContextMenu;
@@ -978,12 +981,12 @@ void SPathView::SyncToVirtualPaths(TArrayView<const FName> VirtualPathsToSync, c
 		{
 			// Prune the current selection so that we don't unnecessarily change the path which might disorientate the user.
 			// If a parent tree item is currently selected we don't need to clear it and select the child
-			auto SelectedTreeItems = TreeViewPtr->GetSelectedItems();
+			TArray<TSharedPtr<FTreeItem>> SelectedTreeItems = TreeViewPtr->GetSelectedItems();
 
 			for (int32 Index = 0; Index < SelectedTreeItems.Num(); ++Index)
 			{
 				// For each item already selected in the tree
-				auto AlreadySelectedTreeItem = SelectedTreeItems[Index];
+				const TSharedPtr<FTreeItem>& AlreadySelectedTreeItem = SelectedTreeItems[Index];
 				if (!AlreadySelectedTreeItem.IsValid())
 				{
 					continue;
@@ -992,7 +995,7 @@ void SPathView::SyncToVirtualPaths(TArrayView<const FName> VirtualPathsToSync, c
 				// Check to see if any of the items to sync are already synced
 				for (int32 ToSyncIndex = SyncTreeItems.Num()-1; ToSyncIndex >= 0; --ToSyncIndex)
 				{
-					auto ToSyncItem = SyncTreeItems[ToSyncIndex];
+					const TSharedPtr<FTreeItem>& ToSyncItem = SyncTreeItems[ToSyncIndex];
 					if (ToSyncItem == AlreadySelectedTreeItem || ToSyncItem->IsChildOf(*AlreadySelectedTreeItem.Get()))
 					{
 						// A parent is already selected
@@ -1013,10 +1016,10 @@ void SPathView::SyncToVirtualPaths(TArrayView<const FName> VirtualPathsToSync, c
 		}
 
 		// SyncTreeItems should now only contain items which aren't already shown explicitly or implicitly (as a child)
-		for ( auto ItemIt = SyncTreeItems.CreateConstIterator(); ItemIt; ++ItemIt )
+		for (const TSharedPtr<FTreeItem>& Item : SyncTreeItems)
 		{
-			RecursiveExpandParents(*ItemIt);
-			TreeViewPtr->SetItemSelection(*ItemIt, true);
+			RecursiveExpandParents(Item);
+			TreeViewPtr->SetItemSelection(Item, true);
 		}
 	}
 
@@ -1067,11 +1070,17 @@ void SPathView::ApplyHistoryData( const FHistoryData& History )
 	SetSelectedPaths(SelectedPaths);
 }
 
-void SPathView::SaveSettings(const FString& IniFilename, const FString& IniSection, const FString& SettingsString) const
+void SPathView::SaveSettings(const FString& IniFilename, const FString& IniSection, const FString& InstanceName) const
 {
 	FString SelectedPathsString;
 	TArray< TSharedPtr<FTreeItem> > PathItems = TreeViewPtr->GetSelectedItems();
-	for ( auto PathIt = PathItems.CreateConstIterator(); PathIt; ++PathIt )
+
+	if (EditorConfig != nullptr)
+	{
+		EditorConfig->SelectedPaths.Reset(PathItems.Num());
+	}
+
+	for (const TSharedPtr<FTreeItem>& Item : PathItems)
 	{
 		if ( SelectedPathsString.Len() > 0 )
 		{
@@ -1079,26 +1088,41 @@ void SPathView::SaveSettings(const FString& IniFilename, const FString& IniSecti
 		}
 
 		FName InvariantPath;
-		IContentBrowserDataModule::Get().GetSubsystem()->TryConvertVirtualPath((*PathIt)->GetItem().GetVirtualPath(), InvariantPath);
+		IContentBrowserDataModule::Get().GetSubsystem()->TryConvertVirtualPath(Item->GetItem().GetVirtualPath(), InvariantPath);
 		InvariantPath.AppendString(SelectedPathsString);
+
+		if (EditorConfig != nullptr)
+		{
+			EditorConfig->SelectedPaths.Add(InvariantPath);
+		}
 	}
 
-	GConfig->SetString(*IniSection, *(SettingsString + TEXT(".SelectedPaths")), *SelectedPathsString, IniFilename);
+	GConfig->SetString(*IniSection, *(InstanceName + TEXT(".SelectedPaths")), *SelectedPathsString, IniFilename);
 
 	FString PluginFiltersString;
 	if (PluginPathFilters.IsValid())
 	{
-		for (int32 i=0; i < PluginPathFilters->Num(); ++i)
+		if (EditorConfig != nullptr)
+		{
+			EditorConfig->PluginFilters.Reset(PluginPathFilters->Num());
+		}
+
+		for (int32 FilterIdx = 0; FilterIdx < PluginPathFilters->Num(); ++FilterIdx)
 		{
 			if (PluginFiltersString.Len() > 0)
 			{
 				PluginFiltersString += TEXT(",");
 			}
 
-			TSharedPtr<FContentBrowserPluginFilter> Filter = StaticCastSharedPtr<FContentBrowserPluginFilter>(PluginPathFilters->GetFilterAtIndex(i));
+			TSharedPtr<FContentBrowserPluginFilter> Filter = StaticCastSharedPtr<FContentBrowserPluginFilter>(PluginPathFilters->GetFilterAtIndex(FilterIdx));
 			PluginFiltersString += Filter->GetName();
+
+			if (EditorConfig != nullptr)
+			{
+				EditorConfig->PluginFilters.Add(Filter->GetName());
+			}
 		}
-		GConfig->SetString(*IniSection, *(SettingsString + TEXT(".PluginFilters")), *PluginFiltersString, IniFilename);
+		GConfig->SetString(*IniSection, *(InstanceName + TEXT(".PluginFilters")), *PluginFiltersString, IniFilename);
 	}
 }
 
@@ -2189,9 +2213,13 @@ void SFavoritePathView::Populate(const bool bIsRefreshingFilter)
 			const FString Path = VirtualPath.ToString();
 
 			// Use the whole path so we deliberately include any children of matched parents in the filtered list
-			if (SearchBoxFolderFilter->PassesFilter(Path))
+			if (!SearchBoxFolderFilter->PassesFilter(Path))
 			{
-				ContentBrowserData->EnumerateItemsAtPath(*Path, CompiledDataFilter.ItemTypeFilter, [this, &CompiledDataFilter](FContentBrowserItemData&& InItemData)
+				continue;
+			}
+
+			ContentBrowserData->EnumerateItemsAtPath(*Path, CompiledDataFilter.ItemTypeFilter, 
+			[this, &CompiledDataFilter](FContentBrowserItemData&& InItemData)
 				{
 					UContentBrowserDataSource* ItemDataSource = InItemData.GetOwnerDataSource();
 					if (ItemDataSource->DoesItemPassFilter(InItemData, CompiledDataFilter))
@@ -2199,7 +2227,6 @@ void SFavoritePathView::Populate(const bool bIsRefreshingFilter)
 						if (TSharedPtr<FTreeItem> Item = AddFolderItem(MoveTemp(InItemData)))
 						{
 							const bool bSelectedItem = LastSelectedPaths.Contains(Item->GetItem().GetInvariantPath());
-
 							if (bSelectedItem)
 							{
 								// Tree items that match the last broadcasted paths should be re-selected them after they are added
@@ -2210,8 +2237,8 @@ void SFavoritePathView::Populate(const bool bIsRefreshingFilter)
 					}
 
 					return true;
-				});
-			}
+				}
+			);
 		}
 	}
 
@@ -2264,7 +2291,7 @@ void SFavoritePathView::LoadSettings(const FString& IniFilename, const FString&
 				const FString& InvariantPath = NewFavoritePaths[PathIdx];
 				FName VirtualPath;
 				IContentBrowserDataModule::Get().GetSubsystem()->ConvertInternalPathToVirtual(FStringView(InvariantPath), VirtualPath);
-				ContentBrowserUtils::AddFavoriteFolder(VirtualPath.ToString(), false);
+				ContentBrowserUtils::AddFavoriteFolder(VirtualPath.ToString());
 				bAddedAtLeastOnePath = true;
 			}
 		}
@@ -2287,13 +2314,13 @@ TSharedPtr<FTreeItem> SFavoritePathView::AddFolderItem(FContentBrowserItemData&&
 	// The favorite view will add all items at the root level
 
 	// Try and find an existing tree item
-	for (const TSharedPtr<FTreeItem>& PotentialTreeItem : TreeRootItems)
+	TWeakPtr<FTreeItem>* WeakExistingItem = TreeItemLookup.Find(InItem.GetVirtualPath());
+	if (WeakExistingItem != nullptr)
 	{
-		if (PotentialTreeItem->GetItem().GetVirtualPath() == InItem.GetVirtualPath())
+		if (TSharedPtr<FTreeItem> ExistingItem = WeakExistingItem->Pin())
 		{
-			// Found a match - merge the new item data
-			PotentialTreeItem->AppendItemData(InItem);
-			return PotentialTreeItem;
+			ExistingItem->AppendItemData(InItem);
+			return ExistingItem;
 		}
 	}
 
@@ -2469,11 +2496,11 @@ void SFavoritePathView::FixupFavoritesFromExternalChange(TArrayView<const AssetV
 		if (bWasFavorite)
 		{
 			// Remove the original path
-			ContentBrowserUtils::RemoveFavoriteFolder(MovedFolder.Key, false);
+			ContentBrowserUtils::RemoveFavoriteFolder(MovedFolder.Key);
 
 			// Add the new path to favorites instead
 			const FString& NewPath = MovedFolder.Value;
-			ContentBrowserUtils::AddFavoriteFolder(NewPath, false);
+			ContentBrowserUtils::AddFavoriteFolder(NewPath);
 			TSharedPtr<FTreeItem> Item = FindTreeItem(*NewPath);
 			if (Item.IsValid())
 			{
diff --git a/Engine/Source/Editor/ContentBrowser/Private/SPathView.h b/Engine/Source/Editor/ContentBrowser/Private/SPathView.h
--- a/Engine/Source/Editor/ContentBrowser/Private/SPathView.h
+++ b/Engine/Source/Editor/ContentBrowser/Private/SPathView.h
@@ -50,6 +50,7 @@ class UToolMenu;
 struct FAssetData;
 struct FGeometry;
 struct FHistoryData;
+struct FPathViewConfig;
 struct FPointerEvent;
 
 typedef TTextFilter< const FString& > FolderTextFilter;
@@ -74,6 +75,7 @@ public:
 		, _AllowReadOnlyFolders(true)
 		, _ShowFavorites(false)
 		, _SelectionMode( ESelectionMode::Multi )
+		, _EditorConfig(nullptr)
 		{}
 
 		/** Content displayed to the left of the search bar */
@@ -124,6 +126,8 @@ public:
 		/** The plugin filter collection */
 		SLATE_ARGUMENT( TSharedPtr<FPluginFilterCollectionType>, PluginPathFilters)
 
+		SLATE_ARGUMENT( FPathViewConfig*, EditorConfig )
+
 	SLATE_END_ARGS()
 
 	/** Destructor */
@@ -197,10 +201,10 @@ public:
 	void ApplyHistoryData( const FHistoryData& History );
 
 	/** Saves any settings to config that should be persistent between editor sessions */
-	virtual void SaveSettings(const FString& IniFilename, const FString& IniSection, const FString& SettingsString) const;
+	virtual void SaveSettings(const FString& IniFilename, const FString& IniSection, const FString& InstanceName) const;
 
 	/** Loads any settings to config that should be persistent between editor sessions */
-	virtual void LoadSettings(const FString& IniFilename, const FString& IniSection, const FString& SettingsString);
+	virtual void LoadSettings(const FString& IniFilename, const FString& IniSection, const FString& InstanceName);
 
 	/**
 	 * Return true if passes path block lists
@@ -473,6 +477,9 @@ private:
 
 	/** The favorites path view if one is set. */
 	TSharedPtr<SExpandableArea> FavoritesArea;
+
+	/** The config instance to use. */
+	FPathViewConfig* EditorConfig = nullptr;
 };
 
 
diff --git a/Engine/Source/Editor/ContentBrowser/Public/IContentBrowserSingleton.h b/Engine/Source/Editor/ContentBrowser/Public/IContentBrowserSingleton.h
--- a/Engine/Source/Editor/ContentBrowser/Public/IContentBrowserSingleton.h
+++ b/Engine/Source/Editor/ContentBrowser/Public/IContentBrowserSingleton.h
@@ -90,54 +90,54 @@ struct FContentBrowserConfig
 	/** The default view mode */
 	EAssetViewType::Type InitialAssetViewType;
 
+	/** Collection to view initially */
+	FCollectionNameType SelectedCollectionName;
+
 	/** If true, show the bottom toolbar which shows # of assets selected, view mode buttons, etc... */
-	bool bShowBottomToolbar;
+	bool bShowBottomToolbar : 1;
 
 	/** Indicates if this view is allowed to show classes */
-	bool bCanShowClasses;
+	bool bCanShowClasses : 1;
 
 	/** Whether the sources view for choosing folders/collections is available or not */
-	bool bUseSourcesView;
+	bool bUseSourcesView : 1;
 
 	/** Whether the sources view should initially be expanded or not */
-	bool bExpandSourcesView;
-
-	/** Collection to view initially */
-	FCollectionNameType SelectedCollectionName;
+	bool bExpandSourcesView : 1;
 
 	/** Whether the path picker is available or not */
-	bool bUsePathPicker;
+	bool bUsePathPicker : 1;
 
 	/** Whether to show filters */
-	bool bCanShowFilters;
+	bool bCanShowFilters : 1;
 
 	/** Whether to show asset search */
-	bool bCanShowAssetSearch;
+	bool bCanShowAssetSearch : 1;
 
 	/** Indicates if the 'Show folders' option should be enabled or disabled */
-	bool bCanShowFolders;
+	bool bCanShowFolders : 1;
 
 	/** Indicates if the 'Real-Time Thumbnails' option should be enabled or disabled */
-	bool bCanShowRealTimeThumbnails;
+	bool bCanShowRealTimeThumbnails : 1;
 
 	/** Indicates if the 'Show Developers' option should be enabled or disabled */
-	bool bCanShowDevelopersFolder;
+	bool bCanShowDevelopersFolder : 1;
 
 	/** Whether the 'lock' button is visible on the toolbar */
-	bool bCanShowLockButton;
+	bool bCanShowLockButton : 1;
 
 	/** Whether or not this Content Browser can be used as the Primary Browser for SyncBrowserTo functions */
-	bool bCanSetAsPrimaryBrowser;
+	bool bCanSetAsPrimaryBrowser : 1;
 
 	FContentBrowserConfig()
 		: ThumbnailLabel( EThumbnailLabel::ClassName )
 		, ThumbnailScale(0.1f)
 		, InitialAssetViewType(EAssetViewType::Tile)
+		, SelectedCollectionName( NAME_None, ECollectionShareType::CST_Local )
 		, bShowBottomToolbar(true)
 		, bCanShowClasses(true)
 		, bUseSourcesView(true)
 		, bExpandSourcesView(true)
-		, SelectedCollectionName( NAME_None, ECollectionShareType::CST_Local )
 		, bUsePathPicker(true)
 		, bCanShowFilters(true)
 		, bCanShowAssetSearch(true)
@@ -346,25 +346,25 @@ struct FPathPickerConfig
 	TArray<FSetPathPickerPathsDelegate*> SetPathsDelegates;
 
 	/** If true, the search box will gain focus when the path picker is created */
-	bool bFocusSearchBoxWhenOpened;
+	bool bFocusSearchBoxWhenOpened : 1;
 
 	/** If false, the context menu will not open when an item is right clicked */
-	bool bAllowContextMenu;
+	bool bAllowContextMenu : 1;
 
 	/** If true, will allow class folders to be shown in the picker */
-	bool bAllowClassesFolder;
+	bool bAllowClassesFolder : 1;
 
 	/** If true, will allow read-only folders to be shown in the picker */
-	bool bAllowReadOnlyFolders;
+	bool bAllowReadOnlyFolders : 1;
 
 	/** If true, will add the path specified in DefaultPath to the tree if it doesn't exist already */
-	bool bAddDefaultPath;
+	bool bAddDefaultPath : 1;
 
 	/** If true, passes virtual paths to OnPathSelected instead of internal asset paths */
-	bool bOnPathSelectedPassesVirtualPaths;
+	bool bOnPathSelectedPassesVirtualPaths : 1;
 
-	/** If true, will show the favorites selector in the picker */
-	bool bShowFavorites;
+	/** Whether or not to show the favorites selector. */
+	bool bShowFavorites : 1;
 
 	FPathPickerConfig()
 		: bFocusSearchBoxWhenOpened(true)
diff --git a/Engine/Source/Editor/ContentBrowser/Public/SAssetView.h b/Engine/Source/Editor/ContentBrowser/Public/SAssetView.h
--- a/Engine/Source/Editor/ContentBrowser/Public/SAssetView.h
+++ b/Engine/Source/Editor/ContentBrowser/Public/SAssetView.h
@@ -7,7 +7,6 @@
 #include "AssetRegistry/AssetData.h"
 #include "AssetThumbnail.h"
 #include "AssetViewSortManager.h"
-#include "AssetViewTypes.h"
 #include "Containers/Array.h"
 #include "Containers/ArrayView.h"
 #include "Containers/Map.h"
@@ -45,6 +44,7 @@
 #include "Widgets/Views/STableRow.h"
 #include "Widgets/Views/STableViewBase.h"
 
+class FAssetThumbnail;
 class FAssetViewItem;
 class FContentBrowserItemDataTemporaryContext;
 class FContentBrowserItemDataUpdate;
@@ -66,8 +66,11 @@ class STableViewBase;
 class SWidget;
 class UClass;
 class UFactory;
+class UToolMenu;
+struct FAssetViewInstanceConfig;
 struct FCharacterEvent;
 struct FCollectionNameType;
+struct FContentBrowserInstanceConfig;
 struct FFocusEvent;
 struct FGeometry;
 struct FKeyEvent;
@@ -396,6 +399,18 @@ public:
 	 */
 	void ForceShowPluginFolder( bool bEnginePlugin );
 
+	/** Sets a non-persistent override value for the Show Engine Content setting. The user can still toggle the setting manually. */
+	void OverrideShowEngineContent( bool bShow );
+
+	/** Sets a non-persistent override value for the Show Engine Content setting. The user can still toggle the setting manually. */
+	void OverrideShowDeveloperContent( bool bShow );
+
+	/** Sets a non-persistent override value for the Show Engine Content setting. The user can still toggle the setting manually. */
+	void OverrideShowPluginContent( bool bShow );
+
+	/** Sets a non-persistent override value for the Show Engine Content setting. The user can still toggle the setting manually. */
+	void OverrideShowLocalizedContent( bool bShow );
+
 	/** @return true when we are including class names in search criteria */
 	bool IsIncludingClassNames() const;
 
@@ -447,13 +462,13 @@ private:
 	void ProcessItemsPendingFilter(const double TickStartTime);
 
 	/** Creates a new tile view */
-	TSharedRef<class SAssetTileView> CreateTileView();
+	TSharedRef<SAssetTileView> CreateTileView();
 
 	/** Creates a new list view */
-	TSharedRef<class SAssetListView> CreateListView();
+	TSharedRef<SAssetListView> CreateListView();
 
 	/** Creates a new column view */
-	TSharedRef<class SAssetColumnView> CreateColumnView();
+	TSharedRef<SAssetColumnView> CreateColumnView();
 
 	const FSlateBrush* GetRevisionControlColumnIconBadge() const;
 
@@ -503,7 +518,7 @@ private:
 	static void RegisterGetViewButtonMenu();
 
 	/** Fill in menu content for when the view combo button is clicked */
-	void PopulateViewButtonMenu(class UToolMenu* Menu);
+	void PopulateViewButtonMenu(UToolMenu* Menu);
 
 	/** Toggle whether folders should be shown or not */
 	void ToggleShowFolders();
@@ -682,7 +697,7 @@ private:
 	void UpdateThumbnails();
 
 	/**  Helper function for UpdateThumbnails. Adds the specified item to the new thumbnail relevancy map and creates any thumbnails for new items. Returns the thumbnail. */
-	TSharedPtr<class FAssetThumbnail> AddItemToNewThumbnailRelevancyMap(const TSharedPtr<FAssetViewItem>& Item, TMap< TSharedPtr<FAssetViewItem>, TSharedPtr<class FAssetThumbnail> >& NewRelevantThumbnails);
+	TSharedPtr<FAssetThumbnail> AddItemToNewThumbnailRelevancyMap(const TSharedPtr<FAssetViewItem>& Item, TMap< TSharedPtr<FAssetViewItem>, TSharedPtr<FAssetThumbnail> >& NewRelevantThumbnails);
 
 	/** Handler for tree view selection changes */
 	void AssetSelectionChanged( TSharedPtr<FAssetViewItem > AssetItem, ESelectInfo::Type SelectInfo );
@@ -850,6 +865,12 @@ private:
 	/** Notification for when the content browser has completed it's initial search */
 	void HandleItemDataDiscoveryComplete();
 
+	/** Get the config struct for the owning Content Browser, if one exists. */
+	FContentBrowserInstanceConfig* GetContentBrowserConfig() const;
+
+	/** Get the config struct for this asset view, if one exists. */
+	FAssetViewInstanceConfig* GetAssetViewConfig() const;
+
 private:
 	friend class FAssetViewFrontendFilterHelper;
 
@@ -876,9 +897,9 @@ private:
 
 	/** The list view that is displaying the assets */
 	EAssetViewType::Type CurrentViewType;
-	TSharedPtr<class SAssetListView> ListView;
-	TSharedPtr<class SAssetTileView> TileView;
-	TSharedPtr<class SAssetColumnView> ColumnView;
+	TSharedPtr<SAssetListView> ListView;
+	TSharedPtr<SAssetTileView> TileView;
+	TSharedPtr<SAssetColumnView> ColumnView;
 	TSharedPtr<SBox> ViewContainer;
 
 	/** The content browser that created this asset view if any */
@@ -971,10 +992,10 @@ private:
 	TWeakPtr<FAssetViewItem> RenamingAsset;
 
 	/** Pool for maintaining and rendering thumbnails */
-	TSharedPtr<class FAssetThumbnailPool> AssetThumbnailPool;
+	TSharedPtr<FAssetThumbnailPool> AssetThumbnailPool;
 
 	/** A map of FAssetViewAsset to the thumbnail that represents it. Only items that are currently visible or within half of the FilteredAssetItems array index distance described by NumOffscreenThumbnails are in this list */
-	TMap< TSharedPtr<FAssetViewItem>, TSharedPtr<class FAssetThumbnail> > RelevantThumbnails;
+	TMap< TSharedPtr<FAssetViewItem>, TSharedPtr<FAssetThumbnail> > RelevantThumbnails;
 
 	/** The set of FAssetItems that currently have widgets displaying them. */
 	TArray< TSharedPtr<FAssetViewItem> > VisibleItems;
@@ -1007,9 +1028,6 @@ private:
 	/** Current thumbnail size */
 	EThumbnailSize ThumbnailSize;
 
-	/** Flag indicating if we will be filling the empty space in the tile view. */
-	bool bFillEmptySpaceInTileView;
-
 	/** The amount to scale each thumbnail so that the empty space is filled. */
 	float FillScale;
 
@@ -1019,59 +1037,86 @@ private:
 	/** The manager responsible for sorting assets in the view */
 	FAssetViewSortManager SortManager;
 
+	/** Flag indicating if we will be filling the empty space in the tile view. */
+	bool bFillEmptySpaceInTileView;
+
 	/** When true, selection change notifications will not be sent */
 	bool bBulkSelecting;
 
 	/** When true, the user may edit thumbnails */
-	bool bAllowThumbnailEditMode;
+	bool bAllowThumbnailEditMode : 1;
 
 	/** True when the asset view is currently allowing the user to edit thumbnails */
-	bool bThumbnailEditMode;
+	bool bThumbnailEditMode : 1;
 
 	/** Indicates if this view is allowed to show classes */
-	bool bCanShowClasses;
+	bool bCanShowClasses : 1;
 
 	/** Indicates if the 'Show Folders' option should be enabled or disabled */
-	bool bCanShowFolders;
+	bool bCanShowFolders : 1;
 
 	/** Indicates if this view is allowed to show folders that cannot be written to */
-	bool bCanShowReadOnlyFolders;
+	bool bCanShowReadOnlyFolders : 1;
 
 	/** If true, recursive filtering will be caused by applying a backend filter */
-	bool bFilterRecursivelyWithBackendFilter;
+	bool bFilterRecursivelyWithBackendFilter : 1;
 
 	/** Indicates if the 'Real-Time Thumbnails' option should be enabled or disabled */
-	bool bCanShowRealTimeThumbnails;
+	bool bCanShowRealTimeThumbnails : 1;
 
 	/** Indicates if the 'Show Developers' option should be enabled or disabled */
-	bool bCanShowDevelopersFolder;
+	bool bCanShowDevelopersFolder : 1;
 
 	/** Indicates if the 'Show Favorites' option should be enabled or disabled */
-	bool bCanShowFavorites;
+	bool bCanShowFavorites : 1;
 
 	/** Indicates if the 'Dock Collections' option should be enabled or disabled */
-	bool bCanDockCollections;
+	bool bCanDockCollections : 1;
 
 	/** If true, it will show path column in the asset view */
-	bool bShowPathInColumnView;
+	bool bShowPathInColumnView : 1;
 
 	/** If true, it will show type in the asset view */
-	bool bShowTypeInColumnView;
+	bool bShowTypeInColumnView : 1;
 
 	/** If true, it sorts by path and then name */
-	bool bSortByPathInColumnView;
+	bool bSortByPathInColumnView : 1;
 
 	/** If true, it will show type in the tile view */
-	bool bShowTypeInTileView;
+	bool bShowTypeInTileView : 1;
 
 	/** If true, engine content is always shown */
-	bool bForceShowEngineContent;
+	bool bForceShowEngineContent : 1;
 
 	/** If true, plugin content is always shown */
-	bool bForceShowPluginContent;
+	bool bForceShowPluginContent : 1;
 
 	/** If true, scrollbar is never shown, removes scroll border entirely */
-	bool bForceHideScrollbar;
+	bool bForceHideScrollbar : 1;
+
+	/** Whether to allow dragging of items */
+	bool bAllowDragging : 1;
+
+	/** Whether this asset view should allow focus on sync or not */
+	bool bAllowFocusOnSync : 1;
+
+	/** Flag set if the user is currently searching */
+	bool bUserSearching : 1;
+	
+	/** Whether or not to notify about newly selected items on on the next asset sync */
+	bool bShouldNotifyNextAssetSync : 1;
+
+	/** If set, this will be used instead of the global setting. Mostly used post-sync. */
+	TOptional<bool> bOverrideShowEngineContent;
+
+	/** If set, this will be used instead of the global setting. Mostly used post-sync. */
+	TOptional<bool> bOverrideShowDeveloperContent;
+
+	/** If set, this will be used instead of the global setting. Mostly used post-sync. */
+	TOptional<bool> bOverrideShowPluginContent;
+
+	/** If set, this will be used instead of the global setting. Mostly used post-sync. */
+	TOptional<bool> bOverrideShowLocalizedContent;
 
 	/** The current selection mode used by the asset view */
 	ESelectionMode::Type SelectionMode;
@@ -1109,18 +1154,6 @@ private:
 	/** Initial set of item categories that this view should show - may be adjusted further by things like CanShowClasses or legacy delegate bindings */
 	EContentBrowserItemCategoryFilter InitialCategoryFilter;
 
-	/** Whether to allow dragging of items */
-	bool bAllowDragging;
-
-	/** Whether this asset view should allow focus on sync or not */
-	bool bAllowFocusOnSync;
-
-	/** Flag set if the user is currently searching */
-	bool bUserSearching;
-	
-	/** Whether or not to notify about newly selected items on on the next asset sync */
-	bool bShouldNotifyNextAssetSync;
-
 	/** A struct to hold data for the deferred creation of a file or folder item */
 	struct FCreateDeferredItemData
 	{
diff --git a/Engine/Source/Editor/UnrealEd/Classes/Settings/ContentBrowserSettings.h b/Engine/Source/Editor/UnrealEd/Classes/Settings/ContentBrowserSettings.h
--- a/Engine/Source/Editor/UnrealEd/Classes/Settings/ContentBrowserSettings.h
+++ b/Engine/Source/Editor/UnrealEd/Classes/Settings/ContentBrowserSettings.h
@@ -18,7 +18,7 @@
 UCLASS(config=EditorSettings)
 class UNREALED_API UContentBrowserSettings : public UObject
 {
-	GENERATED_UCLASS_BODY()
+	GENERATED_BODY()
 
 public:
 
@@ -73,27 +73,51 @@ public:
 public:
 
 	/** Sets whether we are allowed to display the engine folder or not, optional flag for setting override instead */
-	void SetDisplayEngineFolder( bool bInDisplayEngineFolder, bool bOverride = false )
+	void SetDisplayEngineFolder( bool bInDisplayEngineFolder )
+	{
+		DisplayEngineFolder = bInDisplayEngineFolder;
+	}
+
+	UE_DEPRECATED(5.2, "The bOverride flag is no longer necessary. The override can now be set on SAssetView::OverrideShowEngineFolder().")
+	void SetDisplayEngineFolder( bool bInDisplayEngineFolder, bool bOverride )
 	{ 
-		bOverride ? OverrideDisplayEngineFolder = bInDisplayEngineFolder : DisplayEngineFolder = bInDisplayEngineFolder;
+		SetDisplayEngineFolder(bInDisplayEngineFolder);
 	}
 
 	/** Gets whether we are allowed to display the engine folder or not, optional flag ignoring the override */
-	bool GetDisplayEngineFolder( bool bExcludeOverride = false ) const
+	bool GetDisplayEngineFolder() const
+	{
+		return DisplayEngineFolder;
+	}
+
+	UE_DEPRECATED(5.2, "The bOverride flag is no longer necessary. The override can now be set on SAssetView::OverrideShowEngineFolder().")
+	bool GetDisplayEngineFolder( bool bExcludeOverride ) const
 	{ 
-		return ( ( bExcludeOverride ? false : OverrideDisplayEngineFolder ) || DisplayEngineFolder );
+		return GetDisplayEngineFolder();
 	}
 
 	/** Sets whether we are allowed to display the developers folder or not, optional flag for setting override instead */
-	void SetDisplayDevelopersFolder( bool bInDisplayDevelopersFolder, bool bOverride = false )
+	void SetDisplayDevelopersFolder( bool bInDisplayDevelopersFolder )
+	{
+		DisplayDevelopersFolder = bInDisplayDevelopersFolder;
+	}
+
+	UE_DEPRECATED(5.2, "The bOverride flag is no longer necessary. The override can now be set on SAssetView::OverrideShowDeveloperFolders().")
+	void SetDisplayDevelopersFolder( bool bInDisplayDevelopersFolder, bool bOverride )
 	{ 
-		bOverride ? OverrideDisplayDevelopersFolder = bInDisplayDevelopersFolder : DisplayDevelopersFolder = bInDisplayDevelopersFolder;
+		SetDisplayDevelopersFolder(bInDisplayDevelopersFolder);
 	}
 
 	/** Gets whether we are allowed to display the developers folder or not, optional flag ignoring the override */
-	bool GetDisplayDevelopersFolder( bool bExcludeOverride = false ) const
+	bool GetDisplayDevelopersFolder() const
+	{
+		return DisplayDevelopersFolder;
+	}
+
+	UE_DEPRECATED(5.2, "The bOverride flag is no longer necessary. The override can now be set on SAssetView::OverrideShowDeveloperFolders().")
+	bool GetDisplayDevelopersFolder( bool bExcludeOverride ) const
 	{ 
-		return ( ( bExcludeOverride ? false : OverrideDisplayDevelopersFolder ) || DisplayDevelopersFolder );
+		return GetDisplayDevelopersFolder();
 	}
 
 	/** Sets whether we are allowed to display the L10N folder (contains localized assets) or not */
@@ -108,16 +132,28 @@ public:
 		return DisplayL10NFolder;
 	}
 
-	/** Sets whether we are allowed to display the plugin folders or not, optional flag for setting override instead */
-	void SetDisplayPluginFolders( bool bInDisplayPluginFolders, bool bOverride = false )
+	/** Sets whether we are allowed to display the plugin folders or not */
+	void SetDisplayPluginFolders( bool bInDisplayPluginFolders )
+	{
+		DisplayPluginFolders = bInDisplayPluginFolders;
+	}
+
+	UE_DEPRECATED(5.2, "The bOverride flag is no longer necessary. The override can now be set on SAssetView::OverrideShowPluginFolders().")
+	void SetDisplayPluginFolders( bool bInDisplayPluginFolders, bool bOverride )
 	{ 
-		bOverride ? OverrideDisplayPluginFolders = bInDisplayPluginFolders : DisplayPluginFolders = bInDisplayPluginFolders;
+		SetDisplayPluginFolders(bInDisplayPluginFolders);
+	}
+
+	/** Gets whether we are allowed to display the plugin folders or not */
+	bool GetDisplayPluginFolders() const
+	{
+		return DisplayPluginFolders;
 	}
 
-	/** Gets whether we are allowed to display the plugin folders or not, optional flag ignoring the override */
-	bool GetDisplayPluginFolders( bool bExcludeOverride = false ) const
+	UE_DEPRECATED(5.2, "The bOverride flag is no longer necessary. The override can now be set on SAssetView::OverrideShowPluginFolders().")
+	bool GetDisplayPluginFolders( bool bExcludeOverride ) const
 	{ 
-		return ( ( bExcludeOverride ? false : OverrideDisplayPluginFolders ) || DisplayPluginFolders );
+		return GetDisplayPluginFolders();
 	}
 
 	/** Sets whether we are allowed to display favorite folders or not */
@@ -212,9 +248,6 @@ private:
 	UPROPERTY(config)
 	bool DisplayEngineFolder;
 
-	/** If true, overrides the DisplayEngine setting */
-	bool OverrideDisplayEngineFolder;
-
 	/** Whether to display the developers folder in the path view of the content browser */
 	UPROPERTY(config)
 	bool DisplayDevelopersFolder;
@@ -222,16 +255,10 @@ private:
 	UPROPERTY(config)
 	bool DisplayL10NFolder;
 
-	/** If true, overrides the DisplayDev setting */
-	bool OverrideDisplayDevelopersFolder;
-
 	/** List of plugin folders to display in the content browser. */
 	UPROPERTY(config)
 	bool DisplayPluginFolders;
 
-	/** Temporary override for the DisplayPluginFolders setting */
-	bool OverrideDisplayPluginFolders;
-
 	UPROPERTY(config)
 	bool DisplayFavorites;
 
diff --git a/Engine/Source/Editor/UnrealEd/Private/Settings/SettingsClasses.cpp b/Engine/Source/Editor/UnrealEd/Private/Settings/SettingsClasses.cpp
--- a/Engine/Source/Editor/UnrealEd/Private/Settings/SettingsClasses.cpp
+++ b/Engine/Source/Editor/UnrealEd/Private/Settings/SettingsClasses.cpp
@@ -52,13 +52,6 @@ DEFINE_LOG_CATEGORY_STATIC(LogSettingsClasses, Log, All);
 
 UContentBrowserSettings::FSettingChangedEvent UContentBrowserSettings::SettingChangedEvent;
 
-UContentBrowserSettings::UContentBrowserSettings( const FObjectInitializer& ObjectInitializer )
-	: Super(ObjectInitializer)
-	, DockCollections(true)
-{
-}
-
-
 void UContentBrowserSettings::PostEditChangeProperty( struct FPropertyChangedEvent& PropertyChangedEvent )
 {
 	Super::PostEditChangeProperty(PropertyChangedEvent);
