Added missing copyright boilerplate.

diff --git a/Engine/Source/Editor/ContentBrowser/Private/AssetView/AssetViewConfig.h b/Engine/Source/Editor/ContentBrowser/Private/AssetView/AssetViewConfig.h
--- a/Engine/Source/Editor/ContentBrowser/Private/AssetView/AssetViewConfig.h
+++ b/Engine/Source/Editor/ContentBrowser/Private/AssetView/AssetViewConfig.h
@@ -1,3 +1,5 @@
+// Copyright Epic Games, Inc. All Rights Reserved.
+
 #pragma once
 
 #include "Containers/Map.h"
