Collections are now visible in the Content Browser again.

Removed CollectionsDocked setting, since it wasn't bound to any toggle anywhere.

Removed sources widget switcher which was untoggleable.

diff --git a/Engine/Source/Editor/ContentBrowser/Private/ContentBrowserConfig.h b/Engine/Source/Editor/ContentBrowser/Private/ContentBrowserConfig.h
--- a/Engine/Source/Editor/ContentBrowser/Private/ContentBrowserConfig.h
+++ b/Engine/Source/Editor/ContentBrowser/Private/ContentBrowserConfig.h
@@ -110,9 +110,6 @@ struct FContentBrowserInstanceConfig
 	UPROPERTY()
 	bool bRealTimeThumbnails = false;
 
-	UPROPERTY()
-	bool bCollectionsDocked = false;
-
 	UPROPERTY()
 	bool bSearchClasses = false;
 
diff --git a/Engine/Source/Editor/ContentBrowser/Private/SAssetView.cpp b/Engine/Source/Editor/ContentBrowser/Private/SAssetView.cpp
--- a/Engine/Source/Editor/ContentBrowser/Private/SAssetView.cpp
+++ b/Engine/Source/Editor/ContentBrowser/Private/SAssetView.cpp
@@ -257,7 +257,6 @@ void SAssetView::Construct( const FArguments& InArgs )
 	bCanShowDevelopersFolder = InArgs._CanShowDevelopersFolder;
 
 	bCanShowFavorites = InArgs._CanShowFavorites;
-	bCanDockCollections = InArgs._CanDockCollections;
 
 	SelectionMode = InArgs._SelectionMode;
 
@@ -3248,43 +3247,6 @@ bool SAssetView::IsShowingFavorites() const
 	return GetDefault<UContentBrowserSettings>()->GetDisplayFavorites();
 }
 
-void SAssetView::ToggleDockCollections()
-{
-	check(IsToggleDockCollectionsAllowed()); 
-
-	bool bNewState = !GetDefault<UContentBrowserSettings>()->GetDockCollections();
-
-	if (FContentBrowserInstanceConfig* Config = GetContentBrowserConfig())
-	{
-		bNewState = !Config->bCollectionsDocked;
-		Config->bCollectionsDocked = bNewState;
-		UContentBrowserConfig::Get()->SaveEditorConfig();
-	}
-
-	GetMutableDefault<UContentBrowserSettings>()->SetDockCollections(bNewState);
-	GetMutableDefault<UContentBrowserSettings>()->PostEditChange();
-}
-
-bool SAssetView::IsToggleDockCollectionsAllowed() const
-{
-	return bCanDockCollections;
-}
-
-bool SAssetView::HasDockedCollections() const
-{
-	if (!IsToggleDockCollectionsAllowed())
-	{
-		return false;
-	}
-
-	if (const FContentBrowserInstanceConfig* Config = GetContentBrowserConfig())
-	{
-		return Config->bCollectionsDocked;
-	}
-
-	return GetDefault<UContentBrowserSettings>()->GetDockCollections();
-}
-
 void SAssetView::ToggleShowCppContent()
 {
 	check(IsToggleShowCppContentAllowed());
diff --git a/Engine/Source/Editor/ContentBrowser/Private/SContentBrowser.cpp b/Engine/Source/Editor/ContentBrowser/Private/SContentBrowser.cpp
--- a/Engine/Source/Editor/ContentBrowser/Private/SContentBrowser.cpp
+++ b/Engine/Source/Editor/ContentBrowser/Private/SContentBrowser.cpp
@@ -156,12 +156,6 @@ struct FSlateBrush;
 
 const FString SContentBrowser::SettingsIniSection = TEXT("ContentBrowser");
 
-namespace ContentBrowserSourcesWidgetSwitcherIndex
-{
-	static const int32 PathView = 0;
-	static const int32 CollectionsView = 1;
-}
-
 SContentBrowser::~SContentBrowser()
 {
 	// Remove the listener for when view settings are changed
@@ -245,7 +239,7 @@ void SContentBrowser::Construct( const FArguments& InArgs, const FName& InInstan
 		.AddMetaData<FTagMetaData>(FTagMetaData(TEXT("ContentBrowserCollections")))
 		.AllowCollectionDrag(true)
 		.AllowQuickAssetManagement(true)
-		.IsDocked(this, &SContentBrowser::IsCollectionViewDocked)
+		.IsDocked(true)
 		.ExternalSearch(CollectionSearch);
 
 	static const FName DefaultForegroundName("DefaultForeground");
@@ -277,7 +271,6 @@ void SContentBrowser::Construct( const FArguments& InArgs, const FName& InInstan
 		.CanShowRealTimeThumbnails(Config != nullptr ? Config->bCanShowRealTimeThumbnails : true)
 		.CanShowDevelopersFolder(Config != nullptr ? Config->bCanShowDevelopersFolder : true)
 		.CanShowFavorites(true)
-		.CanDockCollections(true)
 		.AddMetaData<FTagMetaData>(FTagMetaData(TEXT("ContentBrowserAssets")))
 		.OwningContentBrowser(SharedThis(this))
 		.OnSearchOptionsChanged(this, &SContentBrowser::HandleAssetViewSearchOptionsChanged)
@@ -466,68 +459,51 @@ void SContentBrowser::Construct( const FArguments& InArgs, const FName& InInstan
 					SNew(SBorder)
 					.Padding(FMargin(0))
 					.BorderImage(FAppStyle::GetBrush("Brushes.Recessed"))
-
 					[
-						// Note: If adding more widgets here, fix ContentBrowserSourcesWidgetSwitcherIndex and the code that uses it!
-						SAssignNew(SourcesWidgetSwitcher, SWidgetSwitcher)
-
-						// Paths View
-						+SWidgetSwitcher::Slot()
+						SAssignNew(PathFavoriteSplitterPtr, SSplitter)
+						.Clipping(EWidgetClipping::ClipToBounds)
+						.PhysicalSplitterHandleSize(2.0f)
+						.HitDetectionSplitterHandleSize(8.0f)
+						.Orientation(EOrientation::Orient_Vertical)
+						.MinimumSlotHeight(26.0f)
+						.Visibility(this, &SContentBrowser::GetSourcesViewVisibility)
+
+						+ SSplitter::Slot()
+						.SizeRule(this, &SContentBrowser::GetFavoritesAreaSizeRule)
+						.MinSize(this, &SContentBrowser::GetFavoritesAreaMinSize)
+						.Value(0.2f)
 						[
-							SAssignNew(PathFavoriteSplitterPtr, SSplitter)
-							.Clipping(EWidgetClipping::ClipToBounds)
-							.PhysicalSplitterHandleSize(2.0f)
-							.HitDetectionSplitterHandleSize(8.0f)
-							.Orientation(EOrientation::Orient_Vertical)
-							.MinimumSlotHeight(26.0f)
-							.Visibility( this, &SContentBrowser::GetSourcesViewVisibility )
-							+SSplitter::Slot()
-							.SizeRule(TAttribute<SSplitter::ESizeRule>(this, &SContentBrowser::GetFavoritesAreaSizeRule))
-							.MinSize(TAttribute<float>(this, &SContentBrowser::GetFavoritesAreaMinSize))
-							.Value(0.2f)
+							SNew(SBorder)
+							.BorderImage(FAppStyle::Get().GetBrush("Brushes.Header"))
+							.Padding(0.0f, 2.0f, 0.0f, 0.0f)
 							[
-								SNew(SBorder)
-								.BorderImage(FAppStyle::Get().GetBrush("Brushes.Header"))
-								.Padding(0.0f, 2.0f, 0.0f, 0.0f)
-								[
-									CreateFavoritesView(Config)
-								]
+								CreateFavoritesView(Config)
 							]
+						]
 								
-							+SSplitter::Slot()
-							.SizeRule(TAttribute<SSplitter::ESizeRule>(this, &SContentBrowser::GetPathAreaSizeRule))
-							.MinSize(29.0f)
-							.Value(0.8f)
-							[
-								SNew(SBorder)
-								.BorderImage(FAppStyle::Get().GetBrush("Brushes.Header"))
-								.Padding(0.0f, 2.0f, 0.0f, 0.0f)
-								[
-									CreatePathView(Config)
-								]
-							]
-
-							+SSplitter::Slot()
-							.SizeRule(TAttribute<SSplitter::ESizeRule>(this, &SContentBrowser::GetCollectionsAreaSizeRule))
-							.MinSize(TAttribute<float>(this, &SContentBrowser::GetCollectionsAreaMinSize))
-							.Value(0.4f)
+						+ SSplitter::Slot()
+						.SizeRule(this, &SContentBrowser::GetPathAreaSizeRule)
+						.MinSize(29.0f)
+						.Value(0.8f)
+						[
+							SNew(SBorder)
+							.BorderImage(FAppStyle::Get().GetBrush("Brushes.Header"))
+							.Padding(0.0f, 2.0f, 0.0f, 0.0f)
 							[
-								SNew(SBorder)
-								.BorderImage(FAppStyle::Get().GetBrush("Brushes.Header"))
-								.Padding(0.0f, 2.0f, 0.0f, 0.0f)
-								[
-									CreateDockedCollectionsView(Config)
-								]
+								CreatePathView(Config)
 							]
 						]
 
-						// Collections View
-						+SWidgetSwitcher::Slot()
+						+ SSplitter::Slot()
+						.SizeRule(this, &SContentBrowser::GetCollectionsAreaSizeRule)
+						.MinSize(29.0f)
+						.Value(0.4f)
 						[
-							SNew(SBox)
-							.Padding(FMargin(0.0f, 4.0f, 0.0f, 0.0f))
+							SNew(SBorder)
+							.BorderImage(FAppStyle::Get().GetBrush("Brushes.Header"))
+							.Padding(0.0f, 2.0f, 0.0f, 0.0f)
 							[
-								CollectionViewPtr.ToSharedRef()
+								CreateDockedCollectionsView(Config)
 							]
 						]
 					]
@@ -685,26 +661,11 @@ EVisibility SContentBrowser::GetFavoriteFolderVisibility() const
 	return GetDefault<UContentBrowserSettings>()->GetDisplayFavorites() ? EVisibility::Visible : EVisibility::Collapsed;
 }
 
-EVisibility SContentBrowser::GetDockedCollectionsVisibility() const
-{
-	return IsCollectionViewDocked() ? EVisibility::Visible : EVisibility::Collapsed;
-}
-
 EVisibility SContentBrowser::GetLockButtonVisibility() const
 {
 	return IsLocked() ? EVisibility::Visible : EVisibility::Collapsed;
 }
 
-bool SContentBrowser::IsCollectionViewDocked() const
-{
-	if (const FContentBrowserInstanceConfig* Config = GetConstInstanceConfig())
-	{
-		return Config->bCollectionsDocked;
-	}
-
-	return GetDefault<UContentBrowserSettings>()->GetDockCollections();
-}
-
 void SContentBrowser::ToggleFolderFavorite(const TArray<FString>& FolderPaths)
 {
 	bool bAddedFavorite = false;
@@ -1168,7 +1129,6 @@ TSharedRef<SWidget> SContentBrowser::CreateDockedCollectionsView(const FContentB
 		.BodyBorderImage(FAppStyle::Get().GetBrush("Brushes.Recessed"))
 		.HeaderPadding(FMargin(4.0f, 0.0f))
 		.Padding(0)
-		.Visibility(this, &SContentBrowser::GetDockedCollectionsVisibility)
 		.AllowAnimatedTransition(false)
 		.OnAreaExpansionChanged_Lambda([this](bool bIsExpanded) { if (!bIsExpanded) CollectionSearchToggleButton->SetExpanded(false); })
 		.HeaderContent()
@@ -1422,7 +1382,7 @@ SSplitter::ESizeRule SContentBrowser::GetPathAreaSizeRule() const
 SSplitter::ESizeRule SContentBrowser::GetCollectionsAreaSizeRule() const
 {
 	// Make sure the area is expanded and visible 
-	return CollectionArea->IsExpanded() && GetDockedCollectionsVisibility() == EVisibility::Visible ? SSplitter::ESizeRule::FractionOfParent : SSplitter::ESizeRule::SizeToContent;
+	return CollectionArea->IsExpanded() ? SSplitter::ESizeRule::FractionOfParent : SSplitter::ESizeRule::SizeToContent;
 }
 
 float SContentBrowser::GetFavoritesAreaMinSize() const
@@ -1430,11 +1390,6 @@ float SContentBrowser::GetFavoritesAreaMinSize() const
 	return GetFavoriteFolderVisibility() == EVisibility::Visible ? 29.0f : 0.0f;
 }
 
-float SContentBrowser::GetCollectionsAreaMinSize() const
-{
-	return GetDockedCollectionsVisibility() == EVisibility::Visible ? 29.0f : 0.0f;
-}
-
 FText SContentBrowser::GetHighlightedText() const
 {
 	return TextFilter->GetRawFilterText();
@@ -3140,19 +3095,6 @@ FReply SContentBrowser::SourcesViewExpandClicked()
 	return FReply::Handled();
 }
 
-EVisibility SContentBrowser::GetSourcesSwitcherVisibility() const
-{
-	return IsCollectionViewDocked() ? EVisibility::Collapsed : EVisibility::Visible;
-}
-
-FReply SContentBrowser::OnSourcesSwitcherClicked()
-{
-	ActiveSourcesWidgetIndex = (ActiveSourcesWidgetIndex + 1) % SourcesWidgetSwitcher->GetNumWidgets();
-	SourcesWidgetSwitcher->SetActiveWidgetIndex(ActiveSourcesWidgetIndex);
-
-	return FReply::Handled();
-}
-
 void SContentBrowser::OnContentBrowserSettingsChanged(FName PropertyName)
 {
 	if (PropertyName.IsNone())
@@ -3548,12 +3490,8 @@ void SContentBrowser::UpdatePath()
 
 	PathBreadcrumbTrail->ClearCrumbs();
 
-	int32 NewSourcesWidgetIndex = ActiveSourcesWidgetIndex;
-
 	if ( SourcesData.HasVirtualPaths() )
 	{
-		NewSourcesWidgetIndex = ContentBrowserSourcesWidgetSwitcherIndex::PathView;
-
 		UContentBrowserDataSubsystem* ContentBrowserData = IContentBrowserDataModule::Get().GetSubsystem();
 
 		TArray<FString> Crumbs;
@@ -3572,8 +3510,6 @@ void SContentBrowser::UpdatePath()
 	}
 	else if ( SourcesData.HasCollections() )
 	{
-		NewSourcesWidgetIndex = IsCollectionViewDocked() ? ContentBrowserSourcesWidgetSwitcherIndex::PathView : ContentBrowserSourcesWidgetSwitcherIndex::CollectionsView;
-
 		FCollectionManagerModule& CollectionManagerModule = FCollectionManagerModule::GetModule();
 		TArray<FCollectionNameType> CollectionPathItems;
 
@@ -3603,12 +3539,6 @@ void SContentBrowser::UpdatePath()
 		PathBreadcrumbTrail->PushCrumb(LOCTEXT("AllAssets", "All Assets"), TEXT(""));
 	}
 
-	if (ActiveSourcesWidgetIndex != NewSourcesWidgetIndex)
-	{
-		ActiveSourcesWidgetIndex = NewSourcesWidgetIndex;
-		SourcesWidgetSwitcher->SetActiveWidgetIndex(ActiveSourcesWidgetIndex);
-	}
-
 	CachedCanWriteToCurrentPath.Reset();
 }
 
@@ -4368,7 +4298,6 @@ FContentBrowserInstanceConfig* SContentBrowser::CreateEditorConfigIfRequired()
 	InstanceConfig->bShowEmptyFolders = Settings->DisplayEmptyFolders;
 	InstanceConfig->bShowCppFolders = Settings->GetDisplayCppFolders();
 	InstanceConfig->bOrganizeFolders = Settings->bOrganizeFolders;
-	InstanceConfig->bCollectionsDocked = Settings->GetDockCollections();
 	InstanceConfig->bFavoritesExpanded = Settings->GetDisplayFavorites();
 	InstanceConfig->bSourcesExpanded = Settings->bOpenSourcesPanelByDefault;
 	InstanceConfig->bRealTimeThumbnails = Settings->RealTimeThumbnails;
diff --git a/Engine/Source/Editor/ContentBrowser/Private/SContentBrowser.h b/Engine/Source/Editor/ContentBrowser/Private/SContentBrowser.h
--- a/Engine/Source/Editor/ContentBrowser/Private/SContentBrowser.h
+++ b/Engine/Source/Editor/ContentBrowser/Private/SContentBrowser.h
@@ -44,7 +44,6 @@ class SFilterList;
 class SPathView;
 class SSearchToggleButton;
 class SWidget;
-class SWidgetSwitcher;
 class UClass;
 class UContentBrowserToolbarMenuContext;
 class UFactory;
@@ -59,11 +58,10 @@ struct FGeometry;
 struct FKeyEvent;
 struct FPointerEvent;
 struct FSlateBrush;
+struct FToolMenuContext;
 
 enum class EFilterBarLayout : uint8;
 
-struct FToolMenuContext;
-
 enum class EContentBrowserViewContext : uint8
 {
 	AssetView,
@@ -386,18 +384,6 @@ private:
 	/** Handler for clicking the tree expand/collapse button */
 	FReply SourcesViewExpandClicked();
 
-	/** Gets the visibility of the source switch button */
-	EVisibility GetSourcesSwitcherVisibility() const;
-
-	/** Gets the icon used on the source switch button */
-	const FSlateBrush* GetSourcesSwitcherIcon() const;
-
-	/** Gets the tooltip text used on the source switch button */
-	FText GetSourcesSwitcherToolTipText() const;
-
-	/** Handler for clicking the source switch button */
-	FReply OnSourcesSwitcherClicked();
-
 	/** Called to handle the Content Browser settings changing */
 	void OnContentBrowserSettingsChanged(FName PropertyName);
 
@@ -542,15 +528,9 @@ private:
 	/** Gets the visibility of the favorites view */
 	EVisibility GetFavoriteFolderVisibility() const;
 
-	/** Get the visibility of the docked collections view */
-	EVisibility GetDockedCollectionsVisibility() const;
-
 	/** Get the visibility of the lock button */
 	EVisibility GetLockButtonVisibility() const;
 
-	/** Whether or not the collections view is docked or exists in its own panel in the same area as the sources view */
-	bool IsCollectionViewDocked() const;
-
 	/** Toggles the favorite status of an array of folders*/
 	void ToggleFolderFavorite(const TArray<FString>& FolderPaths);
 
@@ -594,7 +574,6 @@ private:
 
 	/** Gets the min size for various areas. When areas are not visible the min size is 0, otherwise there is a minimum size to prevent overlap */
 	float GetFavoritesAreaMinSize() const;
-	float GetCollectionsAreaMinSize() const;
 
 	/** Called when the layout of the SFilterList is changing */
 	void OnFilterBarLayoutChanging(EFilterBarLayout NewLayout);
@@ -679,9 +658,6 @@ private:
 	/** Toggle button for showing/hiding collection search area */
 	TSharedPtr<SSearchToggleButton> CollectionSearchToggleButton;
 
-	/** Index of the active sources widget */
-	int32 ActiveSourcesWidgetIndex = 0;
-
 	/** The expanded state of the asset tree */
 	bool bSourcesViewExpanded;
 
@@ -718,9 +694,6 @@ private:
 	/** Delegate used to create a new folder */
 	FOnCreateNewFolder OnCreateNewFolder;
 
-	/** Switcher between the different sources views */
-	TSharedPtr<SWidgetSwitcher> SourcesWidgetSwitcher;
-
 	/** The splitter between the path & asset view */
 	TSharedPtr<SSplitter> PathAssetSplitterPtr;
 
diff --git a/Engine/Source/Editor/ContentBrowser/Public/SAssetView.h b/Engine/Source/Editor/ContentBrowser/Public/SAssetView.h
--- a/Engine/Source/Editor/ContentBrowser/Public/SAssetView.h
+++ b/Engine/Source/Editor/ContentBrowser/Public/SAssetView.h
@@ -110,7 +110,6 @@ public:
 		, _CanShowRealTimeThumbnails(false)
 		, _CanShowDevelopersFolder(false)
 		, _CanShowFavorites(false)
-		, _CanDockCollections(false)
 		, _SelectionMode( ESelectionMode::Multi )
 		, _AllowDragging(true)
 		, _AllowFocusOnSync(true)
@@ -230,7 +229,7 @@ public:
 		SLATE_ARGUMENT(bool, CanShowFavorites)
 
 		/** Indicates if the 'Dock Collections' option should be enabled or disabled */
-		SLATE_ARGUMENT(bool, CanDockCollections)
+		SLATE_ARGUMENT_DEPRECATED(bool, CanDockCollections, 5.2, "CanDockCollections is no longer supported (now defaults to true).")
 
 		/** The selection mode the asset view should use */
 		SLATE_ARGUMENT( ESelectionMode::Type, SelectionMode )
@@ -595,15 +594,6 @@ private:
 	/** @return true when we are showing favorites */
 	bool IsShowingFavorites() const;
 
-	/** Toggle whether the collections view should be docked under the paths view */
-	void ToggleDockCollections();
-
-	/** Whether or not it's possible to dock the collections view */
-	bool IsToggleDockCollectionsAllowed() const;
-
-	/** @return true when the collections view is docked */
-	bool HasDockedCollections() const;
-
 	/** Toggle whether C++ content should be shown or not */
 	void ToggleShowCppContent();
 
@@ -1070,9 +1060,6 @@ private:
 	/** Indicates if the 'Show Favorites' option should be enabled or disabled */
 	bool bCanShowFavorites : 1;
 
-	/** Indicates if the 'Dock Collections' option should be enabled or disabled */
-	bool bCanDockCollections : 1;
-
 	/** If true, it will show path column in the asset view */
 	bool bShowPathInColumnView : 1;
 
