Content browser now restores the selected paths after a restart.

Selecting a content browser path now triggers an editor config save.

Added icon to Lock/Unlock Content Browser setting.

diff --git a/Engine/Source/Editor/ContentBrowser/Private/SContentBrowser.cpp b/Engine/Source/Editor/ContentBrowser/Private/SContentBrowser.cpp
--- a/Engine/Source/Editor/ContentBrowser/Private/SContentBrowser.cpp
+++ b/Engine/Source/Editor/ContentBrowser/Private/SContentBrowser.cpp
@@ -459,56 +459,56 @@ void SContentBrowser::Construct( const FArguments& InArgs, const FName& InInstan
 					SNew(SBorder)
 					.Padding(FMargin(0))
 					.BorderImage(FAppStyle::GetBrush("Brushes.Recessed"))
-					[
-						SAssignNew(PathFavoriteSplitterPtr, SSplitter)
-						.Clipping(EWidgetClipping::ClipToBounds)
-						.PhysicalSplitterHandleSize(2.0f)
-						.HitDetectionSplitterHandleSize(8.0f)
-						.Orientation(EOrientation::Orient_Vertical)
-						.MinimumSlotHeight(26.0f)
+						[
+							SAssignNew(PathFavoriteSplitterPtr, SSplitter)
+							.Clipping(EWidgetClipping::ClipToBounds)
+							.PhysicalSplitterHandleSize(2.0f)
+							.HitDetectionSplitterHandleSize(8.0f)
+							.Orientation(EOrientation::Orient_Vertical)
+							.MinimumSlotHeight(26.0f)
 						.Visibility(this, &SContentBrowser::GetSourcesViewVisibility)
 
 						+ SSplitter::Slot()
 						.SizeRule(this, &SContentBrowser::GetFavoritesAreaSizeRule)
 						.MinSize(this, &SContentBrowser::GetFavoritesAreaMinSize)
-						.Value(0.2f)
-						[
-							SNew(SBorder)
-							.BorderImage(FAppStyle::Get().GetBrush("Brushes.Header"))
-							.Padding(0.0f, 2.0f, 0.0f, 0.0f)
+							.Value(0.2f)
 							[
-								CreateFavoritesView(Config)
+								SNew(SBorder)
+								.BorderImage(FAppStyle::Get().GetBrush("Brushes.Header"))
+								.Padding(0.0f, 2.0f, 0.0f, 0.0f)
+								[
+									CreateFavoritesView(Config)
+								]
 							]
-						]
 								
 						+ SSplitter::Slot()
 						.SizeRule(this, &SContentBrowser::GetPathAreaSizeRule)
-						.MinSize(29.0f)
-						.Value(0.8f)
-						[
-							SNew(SBorder)
-							.BorderImage(FAppStyle::Get().GetBrush("Brushes.Header"))
-							.Padding(0.0f, 2.0f, 0.0f, 0.0f)
+							.MinSize(29.0f)
+							.Value(0.8f)
 							[
-								CreatePathView(Config)
+								SNew(SBorder)
+								.BorderImage(FAppStyle::Get().GetBrush("Brushes.Header"))
+								.Padding(0.0f, 2.0f, 0.0f, 0.0f)
+								[
+									CreatePathView(Config)
+								]
 							]
-						]
 
 						+ SSplitter::Slot()
 						.SizeRule(this, &SContentBrowser::GetCollectionsAreaSizeRule)
 						.MinSize(29.0f)
-						.Value(0.4f)
-						[
-							SNew(SBorder)
-							.BorderImage(FAppStyle::Get().GetBrush("Brushes.Header"))
-							.Padding(0.0f, 2.0f, 0.0f, 0.0f)
+							.Value(0.4f)
 							[
-								CreateDockedCollectionsView(Config)
+								SNew(SBorder)
+								.BorderImage(FAppStyle::Get().GetBrush("Brushes.Header"))
+								.Padding(0.0f, 2.0f, 0.0f, 0.0f)
+								[
+									CreateDockedCollectionsView(Config)
+								]
 							]
 						]
 					]
 				]
-			]
 
 			// Asset View
 			+ SSplitter::Slot()
@@ -533,23 +533,17 @@ void SContentBrowser::Construct( const FArguments& InArgs, const FName& InInstan
 	{
 		// Select the specified collection by default
 		FSourcesData DefaultSourcesData( Config->SelectedCollectionName );
-		TArray<FString> SelectedPaths;
 		AssetViewPtr->SetSourcesData( DefaultSourcesData );
 	}
 	else
 	{
 		// Select /Game by default
-		const FString DefaultInvariantPath = TEXT("/Game");
+		const FName DefaultInvariantPath(TEXT("/Game"));
 		FName DefaultVirtualPath;
-		IContentBrowserDataModule::Get().GetSubsystem()->ConvertInternalPathToVirtual(FStringView(DefaultInvariantPath), DefaultVirtualPath);
+		IContentBrowserDataModule::Get().GetSubsystem()->ConvertInternalPathToVirtual(DefaultInvariantPath, DefaultVirtualPath);
 
 		FSourcesData DefaultSourcesData(DefaultVirtualPath);
-		TArray<FString> SelectedPaths;
-		TArray<FString> SelectedFavoritePaths;
-		SelectedPaths.Add(DefaultVirtualPath.ToString());
-		PathViewPtr->SetSelectedPaths(SelectedPaths);
 		AssetViewPtr->SetSourcesData(DefaultSourcesData);
-		FavoritePathViewPtr->SetSelectedPaths(SelectedFavoritePaths);
 	}
 
 	// Set the initial history data
@@ -757,7 +751,7 @@ TSharedRef<SWidget> SContentBrowser::CreateLockButton(const FContentBrowserConfi
 			.Visibility(this, &SContentBrowser::GetLockButtonVisibility)
 			[
 				SNew(SImage)
-				.Image(this, &SContentBrowser::GetLockIcon)
+				.Image(this, &SContentBrowser::GetLockIconBrush)
 				.ColorAndOpacity(FSlateColor::UseStyle())
 			];
 	}
@@ -1253,8 +1247,11 @@ void SContentBrowser::ExtendViewOptionsMenu(const FContentBrowserConfig* Config)
 							"ToggleLock",
 							TAttribute<FText>(ContentBrowser.ToSharedRef(), &SContentBrowser::GetLockMenuText),
 							LOCTEXT("LockToggleTooltip", "Toggle lock. If locked, this browser will ignore Find in Content Browser requests."),
-							TAttribute<FSlateIcon>(),
-							FUIAction(FExecuteAction::CreateLambda([ContentBrowser = Context->OwningContentBrowser]() { ContentBrowser.Pin()->ToggleLockClicked(); })));
+							TAttribute<FSlateIcon>(ContentBrowser.ToSharedRef(), &SContentBrowser::GetLockIcon),
+							FUIAction(
+								FExecuteAction::CreateLambda([ContentBrowser = Context->OwningContentBrowser]() { ContentBrowser.Pin()->ToggleLockClicked(); })
+							)
+						);
 					}
 
 					if (bShowSourcesView)
@@ -1269,7 +1266,8 @@ void SContentBrowser::ExtendViewOptionsMenu(const FContentBrowserConfig* Config)
 								FExecuteAction::CreateLambda([ContentBrowser = Context->OwningContentBrowser]() { ContentBrowser.Pin()->SourcesViewExpandClicked(); }),
 								FCanExecuteAction(),
 								FIsActionChecked::CreateLambda([ContentBrowser = Context->OwningContentBrowser]() { return ContentBrowser.Pin()->bSourcesViewExpanded; })),
-							EUserInterfaceActionType::ToggleButton);
+							EUserInterfaceActionType::ToggleButton
+						);
 					}
 				}
 			}
@@ -3055,7 +3053,14 @@ FText SContentBrowser::GetLockMenuText() const
 	return IsLocked() ? LOCTEXT("ContentBrowserLockMenu_Unlock", "Unlock Content Browser") : LOCTEXT("ContentBrowserLockMenu_Lock", "Lock Content Browser");
 }
 
-const FSlateBrush* SContentBrowser::GetLockIcon() const
+FSlateIcon SContentBrowser::GetLockIcon() const
+{
+	static const FSlateIcon Unlocked = FSlateIcon(FAppStyle::GetAppStyleSetName(), "Icons.Unlock");
+	static const FSlateIcon Locked = FSlateIcon(FAppStyle::GetAppStyleSetName(), "Icons.Lock");
+	return IsLocked() ? Locked : Unlocked;
+}
+
+const FSlateBrush* SContentBrowser::GetLockIconBrush() const
 {
 	static const FName Unlock = "Icons.Unlock";
 	static const FName Lock = "Icons.Lock";
@@ -3063,7 +3068,6 @@ const FSlateBrush* SContentBrowser::GetLockIcon() const
 	return FAppStyle::Get().GetBrush(IsLocked() ? Lock : Unlock);
 }
 
-
 EVisibility SContentBrowser::GetSourcesViewVisibility() const
 {
 	return bSourcesViewExpanded ? EVisibility::Visible : EVisibility::Collapsed;
diff --git a/Engine/Source/Editor/ContentBrowser/Private/SContentBrowser.h b/Engine/Source/Editor/ContentBrowser/Private/SContentBrowser.h
--- a/Engine/Source/Editor/ContentBrowser/Private/SContentBrowser.h
+++ b/Engine/Source/Editor/ContentBrowser/Private/SContentBrowser.h
@@ -373,7 +373,10 @@ private:
 	FText GetLockMenuText() const;
 
 	/** Gets icon for the lock button */
-	const FSlateBrush* GetLockIcon() const;
+	FSlateIcon GetLockIcon() const;
+
+	/** Gets brush for the lock button */
+	const FSlateBrush* GetLockIconBrush() const;
 
 	/** Gets the visibility state of the asset tree */
 	EVisibility GetSourcesViewVisibility() const;
diff --git a/Engine/Source/Editor/ContentBrowser/Private/SPathView.cpp b/Engine/Source/Editor/ContentBrowser/Private/SPathView.cpp
--- a/Engine/Source/Editor/ContentBrowser/Private/SPathView.cpp
+++ b/Engine/Source/Editor/ContentBrowser/Private/SPathView.cpp
@@ -136,7 +136,6 @@ SPathView::~SPathView()
 void SPathView::Construct( const FArguments& InArgs )
 {
 	EditorConfig = InArgs._EditorConfig;
-
 	OnItemSelectionChanged = InArgs._OnItemSelectionChanged;
 	bAllowContextMenu = InArgs._AllowContextMenu;
 	OnGetItemContextMenu = InArgs._OnGetItemContextMenu;
@@ -419,11 +418,32 @@ void SPathView::SetPluginPathFilterActive(const TSharedRef<FContentBrowserPlugin
 	{
 		PluginPathFilters->Remove(Filter);
 	}
+
+	if (EditorConfig != nullptr)
+	{
+		if (bActive)
+		{
+			EditorConfig->PluginFilters.Add(Filter->GetName());
+		}
+		else
+		{
+			EditorConfig->PluginFilters.Remove(Filter->GetName());
+		}
+		
+		UContentBrowserConfig::Get()->SaveEditorConfig();
+	}
+}
+
+void SPathView::SetSelectedPaths(const TArray<FName>& Paths)
+{
+	TArray<FString> PathStrings;
+	Algo::Transform(Paths, PathStrings, [](const FName& Name) { return Name.ToString(); });
+	SetSelectedPaths(PathStrings);
 }
 
 void SPathView::SetSelectedPaths(const TArray<FString>& Paths)
 {
-	if ( !ensure(TreeViewPtr.IsValid()) )
+	if (!ensure(TreeViewPtr.IsValid()))
 	{
 		return;
 	}
@@ -449,10 +469,8 @@ void SPathView::SetSelectedPaths(const TArray<FString>& Paths)
 	LastSelectedPaths.Empty();
 	TreeViewPtr->ClearSelection();
 
-	for (int32 PathIdx = 0; PathIdx < Paths.Num(); ++PathIdx)
+	for (const FString& Path : Paths)
 	{
-		const FString& Path = Paths[PathIdx];
-
 		TArray<FName> PathItemList;
 		{
 			TArray<FString> PathItemListStr;
@@ -465,31 +483,29 @@ void SPathView::SetSelectedPaths(const TArray<FString>& Paths)
 			}
 		}
 
-		if ( PathItemList.Num() )
+		if (PathItemList.Num())
 		{
 			// There is at least one element in the path
 			TArray<TSharedPtr<FTreeItem>> TreeItems;
 
 			// Find the first item in the root items list
-			for ( int32 RootItemIdx = 0; RootItemIdx < TreeRootItems.Num(); ++RootItemIdx )
+			for (const TSharedPtr<FTreeItem>& TreeItem : TreeRootItems)
 			{
-				if ( TreeRootItems[RootItemIdx]->GetItem().GetItemName() == PathItemList[0] )
+				if (TreeItem->GetItem().GetItemName() == PathItemList[0])
 				{
 					// Found the first item in the path
-					TreeItems.Add(TreeRootItems[RootItemIdx]);
+					TreeItems.Add(TreeItem);
 					break;
 				}
 			}
 
 			// If found in the root items list, try to find the childmost item matching the path
-			if ( TreeItems.Num() > 0 )
+			if (TreeItems.Num() > 0)
 			{
-				for ( int32 PathItemIdx = 1; PathItemIdx < PathItemList.Num(); ++PathItemIdx )
+				for (const FName& PathItemName : PathItemList)
 				{
-					const FName PathItemName = PathItemList[PathItemIdx];
 					const TSharedPtr<FTreeItem> ChildItem = TreeItems.Last()->GetChild(PathItemName);
-
-					if ( ChildItem.IsValid() )
+					if (ChildItem.IsValid())
 					{
 						// Update tree items list
 						TreeItems.Add(ChildItem);
@@ -522,6 +538,13 @@ void SPathView::SetSelectedPaths(const TArray<FString>& Paths)
 			// No path items... skip
 		}
 	}
+
+	if (EditorConfig != nullptr)
+	{
+		EditorConfig->SelectedPaths = LastSelectedPaths.Array();
+
+		UContentBrowserConfig::Get()->SaveEditorConfig();
+	}
 }
 
 void SPathView::ClearSelection()
@@ -1075,14 +1098,9 @@ void SPathView::SaveSettings(const FString& IniFilename, const FString& IniSecti
 	FString SelectedPathsString;
 	TArray< TSharedPtr<FTreeItem> > PathItems = TreeViewPtr->GetSelectedItems();
 
-	if (EditorConfig != nullptr)
-	{
-		EditorConfig->SelectedPaths.Reset(PathItems.Num());
-	}
-
 	for (const TSharedPtr<FTreeItem>& Item : PathItems)
 	{
-		if ( SelectedPathsString.Len() > 0 )
+		if (SelectedPathsString.Len() > 0)
 		{
 			SelectedPathsString += TEXT(",");
 		}
@@ -1090,11 +1108,6 @@ void SPathView::SaveSettings(const FString& IniFilename, const FString& IniSecti
 		FName InvariantPath;
 		IContentBrowserDataModule::Get().GetSubsystem()->TryConvertVirtualPath(Item->GetItem().GetVirtualPath(), InvariantPath);
 		InvariantPath.AppendString(SelectedPathsString);
-
-		if (EditorConfig != nullptr)
-		{
-			EditorConfig->SelectedPaths.Add(InvariantPath);
-		}
 	}
 
 	GConfig->SetString(*IniSection, *(InstanceName + TEXT(".SelectedPaths")), *SelectedPathsString, IniFilename);
@@ -1102,11 +1115,6 @@ void SPathView::SaveSettings(const FString& IniFilename, const FString& IniSecti
 	FString PluginFiltersString;
 	if (PluginPathFilters.IsValid())
 	{
-		if (EditorConfig != nullptr)
-		{
-			EditorConfig->PluginFilters.Reset(PluginPathFilters->Num());
-		}
-
 		for (int32 FilterIdx = 0; FilterIdx < PluginPathFilters->Num(); ++FilterIdx)
 		{
 			if (PluginFiltersString.Len() > 0)
@@ -1116,11 +1124,6 @@ void SPathView::SaveSettings(const FString& IniFilename, const FString& IniSecti
 
 			TSharedPtr<FContentBrowserPluginFilter> Filter = StaticCastSharedPtr<FContentBrowserPluginFilter>(PluginPathFilters->GetFilterAtIndex(FilterIdx));
 			PluginFiltersString += Filter->GetName();
-
-			if (EditorConfig != nullptr)
-			{
-				EditorConfig->PluginFilters.Add(Filter->GetName());
-			}
 		}
 		GConfig->SetString(*IniSection, *(InstanceName + TEXT(".PluginFilters")), *PluginFiltersString, IniFilename);
 	}
@@ -1129,36 +1132,43 @@ void SPathView::SaveSettings(const FString& IniFilename, const FString& IniSecti
 void SPathView::LoadSettings(const FString& IniFilename, const FString& IniSection, const FString& SettingsString)
 {
 	// Selected Paths
-	FString SelectedPathsString;
-	if ( GConfig->GetString(*IniSection, *(SettingsString + TEXT(".SelectedPaths")), SelectedPathsString, IniFilename) )
+	TArray<FName> NewSelectedPaths;
+	if (EditorConfig != nullptr)
 	{
-		TArray<FString> NewSelectedPaths;
-		SelectedPathsString.ParseIntoArray(NewSelectedPaths, TEXT(","), /*bCullEmpty*/true);
-
-		UContentBrowserDataSubsystem* ContentBrowserData = IContentBrowserDataModule::Get().GetSubsystem();
-		const bool bDiscoveringAssets = ContentBrowserData->IsDiscoveringItems();
-
-		// Replace each path in NewSelectedPaths with virtual version of that path
-		for (FString& Path : NewSelectedPaths)
+		NewSelectedPaths = EditorConfig->SelectedPaths;
+	}
+	else 
+	{
+		FString SelectedPathsString;
+		if (GConfig->GetString(*IniSection, *(SettingsString + TEXT(".SelectedPaths")), SelectedPathsString, IniFilename))
 		{
-			FNameBuilder PathBuffer;
-			IContentBrowserDataModule::Get().GetSubsystem()->ConvertInternalPathToVirtual(Path, PathBuffer);
-			Path = PathBuffer;
+			TArray<FString> ParsedPaths;
+			SelectedPathsString.ParseIntoArray(ParsedPaths, TEXT(","), /*bCullEmpty*/true);
+
+			Algo::Transform(ParsedPaths, NewSelectedPaths, [](const FString& Str) { return *Str; });
 		}
+	}
 
+	// Replace each path in NewSelectedPaths with virtual version of that path
+	for (FName& Path : NewSelectedPaths)
+	{
+		IContentBrowserDataModule::Get().GetSubsystem()->ConvertInternalPathToVirtual(Path, Path);
+	}
 
+	{
 		// Batch the selection changed event
 		FScopedSelectionChangedEvent ScopedSelectionChangedEvent(SharedThis(this));
 
 		bPendingInitialPathsNeedsSelectionClear = false;
 
-		if ( bDiscoveringAssets )
+		UContentBrowserDataSubsystem* ContentBrowserData = IContentBrowserDataModule::Get().GetSubsystem();
+		if (ContentBrowserData->IsDiscoveringItems())
 		{
 			// Determine if any of the items already exist
 			bool bFoundAnyItemInTree = false;
-			for (const FString& Path : NewSelectedPaths)
+			for (const FName& Path : NewSelectedPaths)
 			{
-				if (TSharedPtr<FTreeItem> FoundItem = FindTreeItem(*Path))
+				if (TSharedPtr<FTreeItem> FoundItem = FindTreeItem(Path))
 				{
 					bFoundAnyItemInTree = true;
 					break;
@@ -1180,10 +1190,9 @@ void SPathView::LoadSettings(const FString& IniFilename, const FString& IniSecti
 
 			// If the selected paths is empty, the path was "All assets"
 			// This should handle that case properly
-			for (int32 PathIdx = 0; PathIdx < NewSelectedPaths.Num(); ++PathIdx)
+			for (const FName& Path : NewSelectedPaths)
 			{
-				const FName Path = *NewSelectedPaths[PathIdx];
-				if ( !ExplicitlyAddPathToSelection(Path) )
+				if (!ExplicitlyAddPathToSelection(Path))
 				{
 					// If we could not initially select these paths, but are still discovering assets, add them to a pending list to select them later
 					PendingInitialPaths.Add(Path);
@@ -1200,18 +1209,25 @@ void SPathView::LoadSettings(const FString& IniFilename, const FString& IniSecti
 	// Plugin Filters
 	if (PluginPathFilters.IsValid())
 	{
-		FString PluginFiltersString;
-		if (GConfig->GetString(*IniSection, *(SettingsString + TEXT(".PluginFilters")), PluginFiltersString, IniFilename))
+		TArray<FString> NewSelectedFilters;
+		if (EditorConfig != nullptr)
 		{
-			TArray<FString> NewSelectedFilters;
-			PluginFiltersString.ParseIntoArray(NewSelectedFilters, TEXT(","), /*bCullEmpty*/ true);
-
-			for (const TSharedRef<FContentBrowserPluginFilter>& Filter : AllPluginPathFilters)
+			NewSelectedFilters = EditorConfig->PluginFilters;
+		}
+		else
+		{
+			FString PluginFiltersString;
+			if (GConfig->GetString(*IniSection, *(SettingsString + TEXT(".PluginFilters")), PluginFiltersString, IniFilename))
 			{
-				bool bFilterActive = NewSelectedFilters.Contains(Filter->GetName());
-				SetPluginPathFilterActive(Filter, bFilterActive);
+				PluginFiltersString.ParseIntoArray(NewSelectedFilters, TEXT(","), /*bCullEmpty*/ true);
 			}
 		}
+
+		for (const TSharedRef<FContentBrowserPluginFilter>& Filter : AllPluginPathFilters)
+		{
+			bool bFilterActive = NewSelectedFilters.Contains(Filter->GetName());
+			SetPluginPathFilterActive(Filter, bFilterActive);
+		}
 	}
 }
 
@@ -1369,6 +1385,13 @@ void SPathView::TreeSelectionChanged( TSharedPtr< FTreeItem > TreeItem, ESelectI
 				OnItemSelectionChanged.Execute(FContentBrowserItem(), SelectInfo);
 			}
 		}
+
+		if (EditorConfig != nullptr)
+		{
+			EditorConfig->SelectedPaths = LastSelectedPaths.Array();
+
+			UContentBrowserConfig::Get()->SaveEditorConfig();
+		}
 	}
 
 	if (TreeItem.IsValid())
diff --git a/Engine/Source/Editor/ContentBrowser/Private/SPathView.h b/Engine/Source/Editor/ContentBrowser/Private/SPathView.h
--- a/Engine/Source/Editor/ContentBrowser/Private/SPathView.h
+++ b/Engine/Source/Editor/ContentBrowser/Private/SPathView.h
@@ -136,6 +136,9 @@ public:
 	/** Constructs this widget with InArgs */
 	virtual void Construct( const FArguments& InArgs );
 
+	/** Selects the closest matches to the supplied paths in the tree. "/" delimited */
+	void SetSelectedPaths(const TArray<FName>& Paths);
+
 	/** Selects the closest matches to the supplied paths in the tree. "/" delimited */
 	void SetSelectedPaths(const TArray<FString>& Paths);
 
