Content Browser: Fix an error that stopped you from being able to navigate folders by clicking

diff --git a/Engine/Source/Editor/ContentBrowser/Private/SPathView.cpp b/Engine/Source/Editor/ContentBrowser/Private/SPathView.cpp
--- a/Engine/Source/Editor/ContentBrowser/Private/SPathView.cpp
+++ b/Engine/Source/Editor/ContentBrowser/Private/SPathView.cpp
@@ -502,8 +502,9 @@ void SPathView::SetSelectedPaths(const TArray<FString>& Paths)
 			// If found in the root items list, try to find the childmost item matching the path
 			if (TreeItems.Num() > 0)
 			{
-				for (const FName& PathItemName : PathItemList)
+				for ( int32 PathItemIdx = 1; PathItemIdx < PathItemList.Num(); ++PathItemIdx )
 				{
+					const FName PathItemName = PathItemList[PathItemIdx];
 					const TSharedPtr<FTreeItem> ChildItem = TreeItems.Last()->GetChild(PathItemName);
 					if (ChildItem.IsValid())
 					{
