Add read-only access to FAssetContextMenu::SelectedItems

diff --git a/Engine/Source/Editor/ContentBrowser/Private/AssetContextMenu.h b/Engine/Source/Editor/ContentBrowser/Private/AssetContextMenu.h
--- a/Engine/Source/Editor/ContentBrowser/Private/AssetContextMenu.h
+++ b/Engine/Source/Editor/ContentBrowser/Private/AssetContextMenu.h
@@ -38,6 +38,9 @@ public:
 	/** Updates the list of currently selected items to those passed in */
 	void SetSelectedItems(TArrayView<const FContentBrowserItem> InSelectedItems);
 
+	/** Read-only access to the list of currently selected items */
+	const TArray<FContentBrowserItem>& GetSelectedItems() const { return SelectedItems; }
+
 	using FOnShowInPathsViewRequested = UContentBrowserDataMenuContext_FileMenu::FOnShowInPathsView;
 	void SetOnShowInPathsViewRequested(const FOnShowInPathsViewRequested& InOnShowInPathsViewRequested);
 
diff --git a/Engine/Source/Editor/ContentBrowser/Private/ContentBrowserMenuContexts.cpp b/Engine/Source/Editor/ContentBrowser/Private/ContentBrowserMenuContexts.cpp
--- a/Engine/Source/Editor/ContentBrowser/Private/ContentBrowserMenuContexts.cpp
+++ b/Engine/Source/Editor/ContentBrowser/Private/ContentBrowserMenuContexts.cpp
@@ -9,6 +9,7 @@
 #include "ToolMenus.h"
 #include "Algo/Copy.h"
 #include "AssetViewUtils.h"
+#include "AssetContextMenu.h"
 
 FName UContentBrowserToolbarMenuContext::GetCurrentPath() const
 {
@@ -82,6 +83,16 @@ const FAssetData* UContentBrowserAssetContextMenuContext::GetSingleSelectedAsset
 	return nullptr;
 }
 
+const TArray<FContentBrowserItem>& UContentBrowserAssetContextMenuContext::GetSelectedItems() const
+{
+	if (TSharedPtr<FAssetContextMenu> PinnedAssetContextMenu = AssetContextMenu.Pin())
+	{
+		return PinnedAssetContextMenu->GetSelectedItems();
+	}
+	static const TArray<FContentBrowserItem> Empty;
+	return Empty;
+}
+
 namespace UE::ContentBrowser
 {
 	TArray<UToolMenu*> ExtendToolMenu_AssetContextMenus(TConstArrayView<UClass*> AssetClasses)
diff --git a/Engine/Source/Editor/ContentBrowser/Public/ContentBrowserMenuContexts.h b/Engine/Source/Editor/ContentBrowser/Public/ContentBrowserMenuContexts.h
--- a/Engine/Source/Editor/ContentBrowser/Public/ContentBrowserMenuContexts.h
+++ b/Engine/Source/Editor/ContentBrowser/Public/ContentBrowserMenuContexts.h
@@ -211,6 +211,9 @@ public:
 		return Result;
 	}
 
+	/** Read-only access to the list of currently selected items */
+	const TArray<FContentBrowserItem>& GetSelectedItems() const;
+
 	/**
 	 * Finds the Content Browser MenuContext from a Menu or Section, and returns the context provided there are some
 	 * selected assets.
