Content Browser sources panel is now always open by default.

diff --git a/Engine/Source/Editor/ContentBrowser/Private/ContentBrowserConfig.h b/Engine/Source/Editor/ContentBrowser/Private/ContentBrowserConfig.h
--- a/Engine/Source/Editor/ContentBrowser/Private/ContentBrowserConfig.h
+++ b/Engine/Source/Editor/ContentBrowser/Private/ContentBrowserConfig.h
@@ -42,18 +42,6 @@ struct FPathViewConfig
 	TArray<FString> PluginFilters;
 };
 
-USTRUCT()
-struct FSourcesConfig
-{
-	GENERATED_BODY()
-
-	UPROPERTY()
-	bool bExpanded = false;
-
-	UPROPERTY()
-	bool bSearchExpanded = false;
-};
-
 USTRUCT()
 struct FContentBrowserInstanceConfig
 {
@@ -66,28 +54,25 @@ struct FContentBrowserInstanceConfig
 	FPathViewConfig PathView;
 
 	UPROPERTY()
-	FSourcesConfig Sources;
-
-	UPROPERTY()
-	bool bShowFavorites = false;
+	bool bShowFavorites = true;
 
 	UPROPERTY()
-	bool bFavoritesExpanded = false;
+	bool bFavoritesExpanded = true;
 
 	UPROPERTY()
-	bool bSourcesExpanded = false;
+	bool bSourcesExpanded = true;
 
 	UPROPERTY()
 	bool bFilterRecursively = false;
 	
 	UPROPERTY()
-	bool bShowFolders = false;
+	bool bShowFolders = true;
 
 	UPROPERTY()
-	bool bShowEmptyFolders = false;
+	bool bShowEmptyFolders = true;
 
 	UPROPERTY()
-	bool bShowAllFolder = false;
+	bool bShowAllFolder = true;
 
 	UPROPERTY()
 	bool bShowEngineContent = false;
diff --git a/Engine/Source/Editor/ContentBrowser/Private/SContentBrowser.cpp b/Engine/Source/Editor/ContentBrowser/Private/SContentBrowser.cpp
--- a/Engine/Source/Editor/ContentBrowser/Private/SContentBrowser.cpp
+++ b/Engine/Source/Editor/ContentBrowser/Private/SContentBrowser.cpp
@@ -560,7 +560,7 @@ void SContentBrowser::Construct( const FArguments& InArgs, const FName& InInstan
 	else
 	{
 		// in case we do not have a config, see what the global default settings are for the Sources Panel
-		bool bSourcesExpanded = GetDefault<UContentBrowserSettings>()->bOpenSourcesPanelByDefault;
+		bool bSourcesExpanded = true;
 		if (EditorConfig != nullptr)
 		{
 			bSourcesExpanded = EditorConfig->bSourcesExpanded;
@@ -4303,7 +4303,6 @@ FContentBrowserInstanceConfig* SContentBrowser::CreateEditorConfigIfRequired()
 	InstanceConfig->bShowCppFolders = Settings->GetDisplayCppFolders();
 	InstanceConfig->bOrganizeFolders = Settings->bOrganizeFolders;
 	InstanceConfig->bFavoritesExpanded = Settings->GetDisplayFavorites();
-	InstanceConfig->bSourcesExpanded = Settings->bOpenSourcesPanelByDefault;
 	InstanceConfig->bRealTimeThumbnails = Settings->RealTimeThumbnails;
 	InstanceConfig->bSearchAssetPaths = Settings->GetIncludeAssetPaths();
 	InstanceConfig->bSearchClasses = Settings->GetIncludeClassNames();
diff --git a/Engine/Source/Editor/ContentBrowser/Private/SContentBrowser.h b/Engine/Source/Editor/ContentBrowser/Private/SContentBrowser.h
--- a/Engine/Source/Editor/ContentBrowser/Private/SContentBrowser.h
+++ b/Engine/Source/Editor/ContentBrowser/Private/SContentBrowser.h
@@ -662,25 +662,25 @@ private:
 	TSharedPtr<SSearchToggleButton> CollectionSearchToggleButton;
 
 	/** The expanded state of the asset tree */
-	bool bSourcesViewExpanded;
+	bool bSourcesViewExpanded = true;
 
 	/** True if this browser is the primary content browser */
-	bool bIsPrimaryBrowser;
+	bool bIsPrimaryBrowser = false;
 
 	/** True if this content browser can be set to the primary browser. */
-	bool bCanSetAsPrimaryBrowser;
+	bool bCanSetAsPrimaryBrowser = true;
 
 	/** True if this content browser is an a drawer */
-	bool bIsDrawer;
-
-	/** Unique name for this Content Browser. */
-	FName InstanceName;
+	bool bIsDrawer = false;
 
 	/** True if source should not be changed from an outside source */
-	bool bIsLocked;
+	bool bIsLocked = false;
 
 	/** Cached result of CanWriteToPath to avoid recalculating it every frame */
 	mutable bool bCachedCanWriteToCurrentPath = false;
+	
+	/** Unique name for this Content Browser. */
+	FName InstanceName;
 
 	/** Path that was last used to determine bCachedCanWriteToCurrentPath */
 	mutable TOptional<FName> CachedCanWriteToCurrentPath;
diff --git a/Engine/Source/Editor/UnrealEd/Classes/Settings/ContentBrowserSettings.h b/Engine/Source/Editor/UnrealEd/Classes/Settings/ContentBrowserSettings.h
--- a/Engine/Source/Editor/UnrealEd/Classes/Settings/ContentBrowserSettings.h
+++ b/Engine/Source/Editor/UnrealEd/Classes/Settings/ContentBrowserSettings.h
@@ -6,7 +6,6 @@
 
 #pragma once
 
-
 #include "CoreMinimal.h"
 #include "UObject/ObjectMacros.h"
 #include "UObject/Object.h"
@@ -26,10 +25,6 @@ public:
 	UPROPERTY(EditAnywhere, config, Category=ContentBrowser, meta=(DisplayName = "Assets to Load at Once Before Warning", ClampMin = "1"))
 	int32 NumObjectsToLoadBeforeWarning;
 
-	/** Whether the Content Browser should open the Sources Panel by default */
-	UPROPERTY(EditAnywhere, config, Category = ContentBrowser)
-	bool bOpenSourcesPanelByDefault;
-
 	/** Whether to render thumbnails for loaded assets in real-time in the Content Browser */
 	UPROPERTY(config)
 	bool RealTimeThumbnails;
