Fixed truncation warnings in the ContentBrowser modules

diff --git a/Engine/Source/Developer/CollectionManager/CollectionManager.Build.cs b/Engine/Source/Developer/CollectionManager/CollectionManager.Build.cs
--- a/Engine/Source/Developer/CollectionManager/CollectionManager.Build.cs
+++ b/Engine/Source/Developer/CollectionManager/CollectionManager.Build.cs
@@ -6,6 +6,8 @@ public class CollectionManager : ModuleRules
 {
 	public CollectionManager(ReadOnlyTargetRules Target) : base(Target)
 	{
+		UnsafeTypeCastWarningLevel = WarningLevel.Error;
+
 		PrivateDependencyModuleNames.AddRange(
 			new string[] {
 				"Core",
diff --git a/Engine/Source/Editor/AssetDefinition/Public/AssetDefinition.h b/Engine/Source/Editor/AssetDefinition/Public/AssetDefinition.h
--- a/Engine/Source/Editor/AssetDefinition/Public/AssetDefinition.h
+++ b/Engine/Source/Editor/AssetDefinition/Public/AssetDefinition.h
@@ -57,7 +57,7 @@ struct FAssetArgs
 	template<typename ExpectedObjectType>
 	TArray<ExpectedObjectType*> LoadObjects(const TSet<FName>& LoadTags = {}, TArray<FAssetData>* OutAssetsThatFailedToLoad = nullptr) const
 	{
-		FScopedSlowTask SlowTask(Assets.Num());
+		FScopedSlowTask SlowTask((float)Assets.Num());
 	
 		TArray<ExpectedObjectType*> LoadedObjects;
 		LoadedObjects.Reserve(Assets.Num());
diff --git a/Engine/Source/Editor/ContentBrowser/ContentBrowser.Build.cs b/Engine/Source/Editor/ContentBrowser/ContentBrowser.Build.cs
--- a/Engine/Source/Editor/ContentBrowser/ContentBrowser.Build.cs
+++ b/Engine/Source/Editor/ContentBrowser/ContentBrowser.Build.cs
@@ -6,6 +6,8 @@ public class ContentBrowser : ModuleRules
 {
 	public ContentBrowser(ReadOnlyTargetRules Target) : base(Target)
 	{
+		UnsafeTypeCastWarningLevel = WarningLevel.Error;
+
 		PrivateIncludePathModuleNames.AddRange(
 			new string[] {
 				"AssetRegistry",
diff --git a/Engine/Source/Editor/ContentBrowser/Private/AssetViewWidgets.cpp b/Engine/Source/Editor/ContentBrowser/Private/AssetViewWidgets.cpp
--- a/Engine/Source/Editor/ContentBrowser/Private/AssetViewWidgets.cpp
+++ b/Engine/Source/Editor/ContentBrowser/Private/AssetViewWidgets.cpp
@@ -82,7 +82,7 @@ class FDragDropEvent;
 
 #define LOCTEXT_NAMESPACE "ContentBrowser"
 
-static float GenericThumbnailSizes[(int32)EThumbnailSize::MAX] = { 24, 32, 64, 128, 200 };
+static int32 GenericThumbnailSizes[(int32)EThumbnailSize::MAX] = { 24, 32, 64, 128, 200 };
 
 ///////////////////////////////
 // FAssetViewModeUtils
@@ -1756,12 +1756,12 @@ float SAssetListItem::GetExtraStateIconWidth() const
 
 FOptionalSize SAssetListItem::GetExtraStateIconMaxWidth() const
 {
-	return GetThumbnailBoxSize().Get() * 0.7;
+	return GetThumbnailBoxSize().Get() * 0.7f;
 }
 
 FOptionalSize SAssetListItem::GetStateIconImageSize() const
 {
-	float IconSize = FMath::TruncToFloat(GetThumbnailBoxSize().Get() * 0.3);
+	float IconSize = FMath::TruncToFloat(GetThumbnailBoxSize().Get() * 0.3f);
 	return IconSize > 12 ? IconSize : 12;
 }
 
@@ -2024,12 +2024,12 @@ float SAssetTileItem::GetExtraStateIconWidth() const
 
 FOptionalSize SAssetTileItem::GetExtraStateIconMaxWidth() const
 {
-	return GetThumbnailBoxSize().Get() * 0.8;
+	return GetThumbnailBoxSize().Get() * 0.8f;
 }
 
 FOptionalSize SAssetTileItem::GetStateIconImageSize() const
 {
-	float IconSize = FMath::TruncToFloat(GetThumbnailBoxSize().Get() * 0.2);
+	float IconSize = FMath::TruncToFloat(GetThumbnailBoxSize().Get() * 0.2f);
 	return IconSize > 12 ? IconSize : 12;
 }
 
diff --git a/Engine/Source/Editor/ContentBrowser/Private/ContentBrowserUtils.cpp b/Engine/Source/Editor/ContentBrowser/Private/ContentBrowserUtils.cpp
--- a/Engine/Source/Editor/ContentBrowser/Private/ContentBrowserUtils.cpp
+++ b/Engine/Source/Editor/ContentBrowser/Private/ContentBrowserUtils.cpp
@@ -93,7 +93,7 @@ public:
 		[
 			SNew(SBorder)
 			.BorderImage(FAppStyle::GetBrush("Menu.Background"))
-			.Padding(10)
+			.Padding(10.f)
 			.OnMouseButtonDown(this, &SContentBrowserPopup::OnBorderClicked)
 			.BorderBackgroundColor(this, &SContentBrowserPopup::GetBorderBackgroundColor)
 			[
@@ -102,7 +102,7 @@ public:
 				+SHorizontalBox::Slot()
 				.AutoWidth()
 				.VAlign(VAlign_Center)
-				.Padding(0, 0, 4, 0)
+				.Padding(0.f, 0.f, 4.f, 0.f)
 				[
 					SNew(SImage) .Image( FAppStyle::GetBrush("ContentBrowser.PopupMessageIcon") )
 				]
@@ -198,12 +198,12 @@ public:
 		[
 			SNew(SBorder)
 			. BorderImage(FAppStyle::GetBrush("Menu.Background"))
-			. Padding(10)
+			. Padding(10.f)
 			[
 				SNew(SVerticalBox)
 				+SVerticalBox::Slot()
 				.AutoHeight()
-				.Padding(0, 0, 0, 5)
+				.Padding(0.f, 0.f, 0.f, 5.f)
 				.HAlign(HAlign_Center)
 				[
 					SNew(STextBlock)
@@ -215,7 +215,7 @@ public:
 				.HAlign(HAlign_Center)
 				[
 					SNew(SUniformGridPanel)
-					.SlotPadding(3)
+					.SlotPadding(3.f)
 					+ SUniformGridPanel::Slot(0, 0)
 					.HAlign(HAlign_Fill)
 					[
diff --git a/Engine/Source/Editor/ContentBrowser/Private/DragDropHandler.cpp b/Engine/Source/Editor/ContentBrowser/Private/DragDropHandler.cpp
--- a/Engine/Source/Editor/ContentBrowser/Private/DragDropHandler.cpp
+++ b/Engine/Source/Editor/ContentBrowser/Private/DragDropHandler.cpp
@@ -244,7 +244,7 @@ void HandleDragDropMoveOrCopy(const FContentBrowserItem& InDropTargetItem, const
 	if (NumMovedOrCopiedItems > 0 && InParentWidget)
 	{
 		const FText Message = FText::Format(InMoveOrCopyMsg, NumMovedOrCopiedItems, FText::FromName(InDropTargetPath));
-		const FVector2D& CursorPos = FSlateApplication::Get().GetCursorPos();
+		const FVector2f& CursorPos = FSlateApplication::Get().GetCursorPos();
 		FSlateRect MessageAnchor(CursorPos.X, CursorPos.Y, CursorPos.X, CursorPos.Y);
 		ContentBrowserUtils::DisplayMessage(Message, MessageAnchor, InParentWidget.ToSharedRef());
 	}
diff --git a/Engine/Source/Editor/ContentBrowser/Private/SAssetPicker.cpp b/Engine/Source/Editor/ContentBrowser/Private/SAssetPicker.cpp
--- a/Engine/Source/Editor/ContentBrowser/Private/SAssetPicker.cpp
+++ b/Engine/Source/Editor/ContentBrowser/Private/SAssetPicker.cpp
@@ -369,7 +369,7 @@ void SAssetPicker::Construct( const FArguments& InArgs )
 	.AutoWidth()
 	[
 		SNew(SComboButton)
-		.ContentPadding(0)
+		.ContentPadding(0.f)
 		.ComboButtonStyle(&FAppStyle::Get().GetWidgetStyle<FComboButtonStyle>("SimpleComboButton"))
 		.OnGetMenuContent(AssetViewPtr.ToSharedRef(), &SAssetView::GetViewButtonContent)
 		.HasDownArrow(false)
diff --git a/Engine/Source/Editor/ContentBrowser/Private/SAssetView.cpp b/Engine/Source/Editor/ContentBrowser/Private/SAssetView.cpp
--- a/Engine/Source/Editor/ContentBrowser/Private/SAssetView.cpp
+++ b/Engine/Source/Editor/ContentBrowser/Private/SAssetView.cpp
@@ -195,7 +195,7 @@ void SAssetView::Construct( const FArguments& InArgs )
 	InitialNumAmortizedTasks = 0;
 	TotalAmortizeTime = 0;
 	AmortizeStartTime = 0;
-	MaxSecondsPerFrame = 0.015;
+	MaxSecondsPerFrame = 0.015f;
 
 	bFillEmptySpaceInTileView = InArgs._FillEmptySpaceInTileView;
 	FillScale = 1.0f;
@@ -223,11 +223,11 @@ void SAssetView::Construct( const FArguments& InArgs )
 	FDisplayMetrics DisplayMetrics;
 	FSlateApplication::Get().GetCachedDisplayMetrics( DisplayMetrics );
 
-	const FVector2D DisplaySize(
+	const FIntPoint DisplaySize(
 		DisplayMetrics.PrimaryDisplayWorkAreaRect.Right - DisplayMetrics.PrimaryDisplayWorkAreaRect.Left,
 		DisplayMetrics.PrimaryDisplayWorkAreaRect.Bottom - DisplayMetrics.PrimaryDisplayWorkAreaRect.Top );
 
-	ThumbnailScaleRangeScalar = ( DisplaySize.Y / 2160 );
+	ThumbnailScaleRangeScalar = (float)DisplaySize.Y / 2160.f;
 
 	// Create a thumbnail pool for rendering thumbnails	
 	AssetThumbnailPool = MakeShared<FAssetThumbnailPool>(InArgs._InitialThumbnailPoolSize);
@@ -347,7 +347,7 @@ void SAssetView::Construct( const FArguments& InArgs )
 	.Padding(0.0f)
 	[
 		SNew(SBorder)
-		.Padding(0)
+		.Padding(0.f)
 		.BorderImage(FAppStyle::Get().GetBrush("Brushes.Panel"))
 		[
 			VerticalBox
@@ -365,7 +365,7 @@ void SAssetView::Construct( const FArguments& InArgs )
 		[
 			SNew( SBox )
 			.Visibility_Lambda([this] { return InitialNumAmortizedTasks > 0 ? EVisibility::SelfHitTestInvisible : EVisibility::Collapsed; })
-			.HeightOverride( 2 )
+			.HeightOverride( 2.f )
 			[
 				SNew( SProgressBar )
 				.Percent( this, &SAssetView::GetIsWorkingProgressBarState )
@@ -429,7 +429,7 @@ void SAssetView::Construct( const FArguments& InArgs )
 	// Thumbnail edit mode banner
 	VerticalBox->AddSlot()
 	.AutoHeight()
-	.Padding(0, 4)
+	.Padding(0.f, 4.f)
 	[
 		SNew(SBorder)
 		.Visibility( this, &SAssetView::GetEditModeLabelVisibility )
@@ -439,7 +439,7 @@ void SAssetView::Construct( const FArguments& InArgs )
 			SNew(SHorizontalBox)
 			+SHorizontalBox::Slot()
 			.VAlign(VAlign_Center)
-			.Padding(4, 0, 0, 0)
+			.Padding(4.f, 0.f, 0.f, 0.f)
 			.FillWidth(1.f)
 			[
 				SNew(STextBlock)
@@ -482,7 +482,7 @@ void SAssetView::Construct( const FArguments& InArgs )
 			[
 				SNew(SComboButton)
 				.Visibility(InArgs._ShowViewOptions ? EVisibility::Visible : EVisibility::Collapsed)
-				.ContentPadding(0)
+				.ContentPadding(0.f)
 				.ButtonStyle( FAppStyle::Get(), "ToggleButton" ) // Use the tool bar item style for this button
 				.OnGetMenuContent( this, &SAssetView::GetViewButtonContent )
 				.ButtonContent()
@@ -499,7 +499,7 @@ void SAssetView::Construct( const FArguments& InArgs )
  
 					+SHorizontalBox::Slot()
 					.AutoWidth()
-					.Padding(2, 0, 0, 0)
+					.Padding(2.f, 0.f, 0.f, 0.f)
 					.VAlign(VAlign_Center)
 					[
 						SNew(STextBlock)
@@ -1201,7 +1201,7 @@ void SAssetView::CalculateFillScale( const FGeometry& AllottedGeometry )
 		// If there isn't enough room to support even a single item, don't apply a fill scale.
 		if ( Items > 0 )
 		{
-			float GapSpace = ItemWidth * ( Coverage - Items );
+			float GapSpace = ItemWidth * ( Coverage - (float)Items );
 			float ExpandAmount = GapSpace / (float)Items;
 			FillScale = ( ItemWidth + ExpandAmount ) / ItemWidth;
 			FillScale = FMath::Max( 1.0f, FillScale );
@@ -1628,7 +1628,7 @@ FReply SAssetView::OnKeyDown( const FGeometry& MyGeometry, const FKeyEvent& InKe
 	}
 	// Swallow the key-presses used by the quick-jump in OnKeyChar to avoid other things (such as the viewport commands) getting them instead
 	// eg) Pressing "W" without this would set the viewport to "translate" mode
-	else if(HandleQuickJumpKeyDown(InKeyEvent.GetCharacter(), bIsControlOrCommandDown, InKeyEvent.IsAltDown(), /*bTestOnly*/true).IsEventHandled())
+	else if(HandleQuickJumpKeyDown((TCHAR)InKeyEvent.GetCharacter(), bIsControlOrCommandDown, InKeyEvent.IsAltDown(), /*bTestOnly*/true).IsEventHandled())
 	{
 		return FReply::Handled();
 	}
@@ -3446,7 +3446,7 @@ void SAssetView::SetCurrentViewType(EAssetViewType::Type NewType)
 
 		if (FAssetViewInstanceConfig* Config = GetAssetViewConfig())
 		{
-			Config->ViewType = (int32) NewType;
+			Config->ViewType = (uint8) NewType;
 			UAssetViewConfig::Get()->SaveEditorConfig();
 		}
 	}
@@ -3630,8 +3630,7 @@ TSharedRef<ITableRow> SAssetView::MakeListViewWidget(TSharedPtr<FAssetViewItem>
 		TSharedPtr<FAssetThumbnail>& AssetThumbnail = RelevantThumbnails.FindOrAdd(AssetItem);
 		if (!AssetThumbnail)
 		{
-			const float ThumbnailResolution = ListViewThumbnailResolution;
-			AssetThumbnail = MakeShared<FAssetThumbnail>(FAssetData(), ThumbnailResolution, ThumbnailResolution, AssetThumbnailPool);
+			AssetThumbnail = MakeShared<FAssetThumbnail>(FAssetData(), ListViewThumbnailResolution, ListViewThumbnailResolution, AssetThumbnailPool);
 			AssetItem->GetItem().UpdateThumbnail(*AssetThumbnail);
 			AssetThumbnail->GetViewportRenderTargetTexture(); // Access the texture once to trigger it to render
 		}
@@ -3644,7 +3643,7 @@ TSharedRef<ITableRow> SAssetView::MakeListViewWidget(TSharedPtr<FAssetViewItem>
 				SNew(SAssetListItem)
 					.AssetThumbnail(AssetThumbnail)
 					.AssetItem(AssetItem)
-					.ThumbnailPadding(ListViewThumbnailPadding)
+					.ThumbnailPadding((float)ListViewThumbnailPadding)
 					.ItemHeight(this, &SAssetView::GetListViewItemHeight)
 					.OnRenameBegin(this, &SAssetView::AssetRenameBegin)
 					.OnRenameCommit(this, &SAssetView::AssetRenameCommit)
@@ -3685,7 +3684,7 @@ TSharedRef<ITableRow> SAssetView::MakeTileViewWidget(TSharedPtr<FAssetViewItem>
 		TSharedRef<SAssetTileItem> Item =
 			SNew(SAssetTileItem)
 			.AssetItem(AssetItem)
-			.ThumbnailPadding(TileViewThumbnailPadding)
+			.ThumbnailPadding((float)TileViewThumbnailPadding)
 			.ItemWidth(this, &SAssetView::GetTileViewItemWidth)
 			.OnRenameBegin(this, &SAssetView::AssetRenameBegin)
 			.OnRenameCommit(this, &SAssetView::AssetRenameCommit)
@@ -3705,8 +3704,7 @@ TSharedRef<ITableRow> SAssetView::MakeTileViewWidget(TSharedPtr<FAssetViewItem>
 		TSharedPtr<FAssetThumbnail>& AssetThumbnail = RelevantThumbnails.FindOrAdd(AssetItem);
 		if (!AssetThumbnail)
 		{
-			const float ThumbnailResolution = TileViewThumbnailResolution;
-			AssetThumbnail = MakeShared<FAssetThumbnail>(FAssetData(), ThumbnailResolution, ThumbnailResolution, AssetThumbnailPool);
+			AssetThumbnail = MakeShared<FAssetThumbnail>(FAssetData(), TileViewThumbnailResolution, TileViewThumbnailResolution, AssetThumbnailPool);
 			AssetItem->GetItem().UpdateThumbnail(*AssetThumbnail);
 			AssetThumbnail->GetViewportRenderTargetTexture(); // Access the texture once to trigger it to render
 		}
@@ -3721,7 +3719,7 @@ TSharedRef<ITableRow> SAssetView::MakeTileViewWidget(TSharedPtr<FAssetViewItem>
 			SNew(SAssetTileItem)
 			.AssetThumbnail(AssetThumbnail)
 			.AssetItem(AssetItem)
-			.ThumbnailPadding(TileViewThumbnailPadding)
+			.ThumbnailPadding((float)TileViewThumbnailPadding)
 			.CurrentThumbnailSize(this, &SAssetView::GetThumbnailSize)
 			.ItemWidth(this, &SAssetView::GetTileViewItemWidth)
 			.OnRenameBegin(this, &SAssetView::AssetRenameBegin)
@@ -3814,7 +3812,7 @@ void SAssetView::UpdateThumbnails()
 	int32 MinVisibleItemIdx = INDEX_NONE;
 	int32 MaxVisibleItemIdx = INDEX_NONE;
 
-	const int32 HalfNumOffscreenThumbnails = NumOffscreenThumbnails * 0.5;
+	const int32 HalfNumOffscreenThumbnails = NumOffscreenThumbnails / 2;
 	for ( auto ItemIt = VisibleItems.CreateConstIterator(); ItemIt; ++ItemIt )
 	{
 		int32 ItemIdx = FilteredAssetItems.Find(*ItemIt);
@@ -3910,7 +3908,7 @@ TSharedPtr<FAssetThumbnail> SAssetView::AddItemToNewThumbnailRelevancyMap(const
 		}
 
 		// The thumbnail newly relevant, create a new thumbnail
-		const float ThumbnailResolution = CurrentThumbnailSize * MaxThumbnailScale;
+		const int32 ThumbnailResolution = FMath::TruncToInt((float)CurrentThumbnailSize * MaxThumbnailScale);
 		Thumbnail = MakeShared<FAssetThumbnail>(FAssetData(), ThumbnailResolution, ThumbnailResolution, AssetThumbnailPool);
 		Item->GetItem().UpdateThumbnail(*Thumbnail);
 		Thumbnail->GetViewportRenderTargetTexture(); // Access the texture once to trigger it to render
@@ -4314,7 +4312,7 @@ void SAssetView::OnThumbnailSizeChanged(EThumbnailSize NewThumbnailSize)
 
 	if (FAssetViewInstanceConfig* Config = GetAssetViewConfig())
 	{
-		Config->ThumbnailSize = (int32) NewThumbnailSize;
+		Config->ThumbnailSize = (uint8) NewThumbnailSize;
 		UAssetViewConfig::Get()->SaveEditorConfig();
 	}
 
@@ -4387,17 +4385,17 @@ float SAssetView::GetTileViewTypeNameHeight() const
 
 float SAssetView::GetListViewItemHeight() const
 {
-	return (ListViewThumbnailSize + ListViewThumbnailPadding * 2) * FMath::Lerp(MinThumbnailScale, MaxThumbnailScale, GetThumbnailScale());
+	return (float)(ListViewThumbnailSize + ListViewThumbnailPadding * 2) * FMath::Lerp(MinThumbnailScale, MaxThumbnailScale, GetThumbnailScale());
 }
 
 float SAssetView::GetTileViewItemHeight() const
 {
-	return ((TileViewNameHeight + GetTileViewTypeNameHeight()) * FMath::Lerp(MinThumbnailScale, MaxThumbnailScale, GetThumbnailScale())) + GetTileViewItemBaseHeight() * FillScale;
+	return (((float)TileViewNameHeight + GetTileViewTypeNameHeight()) * FMath::Lerp(MinThumbnailScale, MaxThumbnailScale, GetThumbnailScale())) + GetTileViewItemBaseHeight() * FillScale;
 }
 
 float SAssetView::GetTileViewItemBaseHeight() const
 {
-	return (TileViewThumbnailSize + TileViewThumbnailPadding * 2) * FMath::Lerp(MinThumbnailScale, MaxThumbnailScale, GetThumbnailScale());
+	return (float)(TileViewThumbnailSize + TileViewThumbnailPadding * 2) * FMath::Lerp(MinThumbnailScale, MaxThumbnailScale, GetThumbnailScale());
 }
 
 float SAssetView::GetTileViewItemWidth() const
@@ -4407,7 +4405,7 @@ float SAssetView::GetTileViewItemWidth() const
 
 float SAssetView::GetTileViewItemBaseWidth() const //-V524
 {
-	return ( TileViewThumbnailSize + TileViewThumbnailPadding * 2 ) * FMath::Lerp( MinThumbnailScale, MaxThumbnailScale, GetThumbnailScale() );
+	return (float)( TileViewThumbnailSize + TileViewThumbnailPadding * 2 ) * FMath::Lerp( MinThumbnailScale, MaxThumbnailScale, GetThumbnailScale() );
 }
 
 EColumnSortMode::Type SAssetView::GetColumnSortMode(const FName ColumnId) const
diff --git a/Engine/Source/Editor/ContentBrowser/Private/SCollectionView.cpp b/Engine/Source/Editor/ContentBrowser/Private/SCollectionView.cpp
--- a/Engine/Source/Editor/ContentBrowser/Private/SCollectionView.cpp
+++ b/Engine/Source/Editor/ContentBrowser/Private/SCollectionView.cpp
@@ -1061,7 +1061,7 @@ void SCollectionView::DeleteCollectionItems( const TArray<TSharedPtr<FCollection
 		else
 		{
 			// Display a warning
-			const FVector2D& CursorPos = FSlateApplication::Get().GetCursorPos();
+			const FVector2f& CursorPos = FSlateApplication::Get().GetCursorPos();
 			FSlateRect MessageAnchor(CursorPos.X, CursorPos.Y, CursorPos.X, CursorPos.Y);
 			ContentBrowserUtils::DisplayMessage(
 				FText::Format( LOCTEXT("CollectionDestroyFailed", "Failed to destroy collection. {0}"), CollectionManagerModule.Get().GetLastError() ),
diff --git a/Engine/Source/Editor/ContentBrowser/Private/SContentBrowser.cpp b/Engine/Source/Editor/ContentBrowser/Private/SContentBrowser.cpp
--- a/Engine/Source/Editor/ContentBrowser/Private/SContentBrowser.cpp
+++ b/Engine/Source/Editor/ContentBrowser/Private/SContentBrowser.cpp
@@ -981,7 +981,7 @@ TSharedRef<SWidget> SContentBrowser::CreateFavoritesView(const FContentBrowserCo
 		.BodyBorderImage(FAppStyle::Get().GetBrush("Brushes.Recessed"))
 		.HeaderPadding(FMargin(4.0f, 0.0f))
 		.Visibility(this, &SContentBrowser::GetFavoriteFolderVisibility)
-		.Padding(0)
+		.Padding(0.f)
 		.AllowAnimatedTransition(false)
 		.OnAreaExpansionChanged(this, &SContentBrowser::SetFavoritesExpanded)
 		.HeaderContent()
@@ -1057,7 +1057,7 @@ TSharedRef<SWidget> SContentBrowser::CreatePathView(const FContentBrowserConfig*
 		.BorderImage(FAppStyle::Get().GetBrush("Brushes.Header"))
 		.BodyBorderImage(FAppStyle::Get().GetBrush("Brushes.Recessed"))
 		.HeaderPadding(FMargin(4.0f, 0.0f))
-		.Padding(0)
+		.Padding(0.f)
 		.AllowAnimatedTransition(false)
 		.OnAreaExpansionChanged(this, &SContentBrowser::SetPathViewExpanded)
 		.HeaderContent()
@@ -1123,7 +1123,7 @@ TSharedRef<SWidget> SContentBrowser::CreateDockedCollectionsView(const FContentB
 		.BorderImage(FAppStyle::Get().GetBrush("Brushes.Header"))
 		.BodyBorderImage(FAppStyle::Get().GetBrush("Brushes.Recessed"))
 		.HeaderPadding(FMargin(4.0f, 0.0f))
-		.Padding(0)
+		.Padding(0.f)
 		.AllowAnimatedTransition(false)
 		.OnAreaExpansionChanged_Lambda([this](bool bIsExpanded) { if (!bIsExpanded) CollectionSearchToggleButton->SetExpanded(false); })
 		.HeaderContent()
@@ -1327,7 +1327,7 @@ void SContentBrowser::RegisterContentBrowserToolBar()
 				SNew(SButton)
 				.ButtonStyle(FAppStyle::Get(), "SimpleButton")
 				.ToolTipText(LOCTEXT("SaveDirtyPackagesTooltip", "Save all modified assets."))
-				.ContentPadding(2)
+				.ContentPadding(2.f)
 				.OnClicked(ContentBrowser, &SContentBrowser::OnSaveClicked)
 				.AddMetaData<FTagMetaData>(FTagMetaData(TEXT("ContentBrowserSaveDirtyPackages")))
 				[
@@ -2555,9 +2555,9 @@ TSharedRef<SWidget> SContentBrowser::GetPathPickerContent()
 	PathPickerConfig.bOnPathSelectedPassesVirtualPaths = true;
 
 	return SNew(SBox)
-		.WidthOverride(300)
-		.HeightOverride(500)
-		.Padding(4)
+		.WidthOverride(300.f)
+		.HeightOverride(500.f)
+		.Padding(4.f)
 		[
 			SNew(SVerticalBox)
 
@@ -2571,7 +2571,7 @@ TSharedRef<SWidget> SContentBrowser::GetPathPickerContent()
 			// Collection View
 			+SVerticalBox::Slot()
 			.AutoHeight()
-			.Padding(0, 6, 0, 0)
+			.Padding(0.f, 6.f, 0.f, 0.f)
 			[
 				SNew(SCollectionView)
 				.AllowCollectionButtons(false)
diff --git a/Engine/Source/Editor/ContentBrowser/Private/SPathView.cpp b/Engine/Source/Editor/ContentBrowser/Private/SPathView.cpp
--- a/Engine/Source/Editor/ContentBrowser/Private/SPathView.cpp
+++ b/Engine/Source/Editor/ContentBrowser/Private/SPathView.cpp
@@ -1742,7 +1742,7 @@ bool SPathView::VerifyFolderNameChanged(const TSharedPtr< FTreeItem >& TreeItem,
 	return true;
 }
 
-void SPathView::FolderNameChanged( const TSharedPtr< FTreeItem >& TreeItem, const FString& ProposedName, const FVector2D& MessageLocation, const ETextCommit::Type CommitType )
+void SPathView::FolderNameChanged( const TSharedPtr< FTreeItem >& TreeItem, const FString& ProposedName, const UE::Slate::FDeprecateVector2DParameter& MessageLocation, const ETextCommit::Type CommitType )
 {
 	bool bSuccess = false;
 	FText ErrorMessage;
@@ -2166,7 +2166,7 @@ TSharedRef<SWidget> SPathView::CreateFavoritesView()
 		.BorderImage(FAppStyle::Get().GetBrush("Brushes.Header"))
 		.BodyBorderImage(FAppStyle::Get().GetBrush("Brushes.Recessed"))
 		.HeaderPadding(FMargin(4.0f, 4.0f))
-		.Padding(0)
+		.Padding(0.f)
 		.AllowAnimatedTransition(false)
 		.InitiallyCollapsed(true)
 		.HeaderContent()
diff --git a/Engine/Source/Editor/ContentBrowser/Private/SPathView.h b/Engine/Source/Editor/ContentBrowser/Private/SPathView.h
--- a/Engine/Source/Editor/ContentBrowser/Private/SPathView.h
+++ b/Engine/Source/Editor/ContentBrowser/Private/SPathView.h
@@ -33,6 +33,7 @@
 #include "Templates/UnrealTemplate.h"
 #include "Types/SlateConstants.h"
 #include "Types/SlateEnums.h"
+#include "Types/SlateVector2.h"
 #include "UObject/NameTypes.h"
 #include "Widgets/DeclarativeSyntaxSupport.h"
 #include "Widgets/SCompoundWidget.h"
@@ -280,7 +281,7 @@ protected:
 	void GetChildrenForTree(TSharedPtr< FTreeItem > TreeItem, TArray< TSharedPtr<FTreeItem> >& OutChildren);
 
 	/** Handler for when a name was given to a new folder */
-	void FolderNameChanged(const TSharedPtr< FTreeItem >& TreeItem, const FString& ProposedName, const FVector2D& MessageLocation, const ETextCommit::Type CommitType);
+	void FolderNameChanged(const TSharedPtr< FTreeItem >& TreeItem, const FString& ProposedName, const UE::Slate::FDeprecateVector2DParameter& MessageLocation, const ETextCommit::Type CommitType);
 
 	/** Handler used to verify the name of a new folder */
 	bool VerifyFolderNameChanged(const TSharedPtr< FTreeItem >& TreeItem, const FString& ProposedName, FText& OutErrorMessage) const;
diff --git a/Engine/Source/Editor/ContentBrowser/Private/SThumbnailEditModeTools.cpp b/Engine/Source/Editor/ContentBrowser/Private/SThumbnailEditModeTools.cpp
--- a/Engine/Source/Editor/ContentBrowser/Private/SThumbnailEditModeTools.cpp
+++ b/Engine/Source/Editor/ContentBrowser/Private/SThumbnailEditModeTools.cpp
@@ -60,11 +60,11 @@ void SThumbnailEditModeTools::Construct( const FArguments& InArgs, const TShared
 		+SHorizontalBox::Slot()
 		.VAlign(VAlign_Top)
 		.HAlign(HAlign_Left)
-		.Padding(1)
+		.Padding(1.f)
 		[
 			SNew(SButton)
 			.Visibility(this, &SThumbnailEditModeTools::GetPrimitiveToolsVisibility)
-			.ContentPadding(0)
+			.ContentPadding(0.f)
 			.ButtonStyle(FAppStyle::Get(), "SimpleButton")
 			.OnClicked(this, &SThumbnailEditModeTools::ChangePrimitive)
 			.ToolTipText(LOCTEXT("CyclePrimitiveThumbnailShapes", "Cycle through primitive shape for this thumbnail"))
@@ -76,11 +76,11 @@ void SThumbnailEditModeTools::Construct( const FArguments& InArgs, const TShared
 		+SHorizontalBox::Slot()
 		.HAlign(HAlign_Right)
 		.VAlign(VAlign_Top)
-		.Padding(1)
+		.Padding(1.f)
 		[
 			SNew(SButton)
 			.Visibility(this, &SThumbnailEditModeTools::GetPrimitiveToolsResetToDefaultVisibility)
-			.ContentPadding(0)
+			.ContentPadding(0.f)
 			.ButtonStyle(FAppStyle::Get(), "SimpleButton")
 			.OnClicked(this, &SThumbnailEditModeTools::ResetToDefault)
 			.ToolTipText(LOCTEXT("ResetThumbnailToDefault", "Resets thumbnail to the default"))
@@ -141,7 +141,7 @@ FReply SThumbnailEditModeTools::ChangePrimitive()
 	USceneThumbnailInfoWithPrimitive* ThumbnailInfo = GetSceneThumbnailInfoWithPrimitive();
 	if ( ThumbnailInfo )
 	{
-		uint8 PrimitiveIdx = ThumbnailInfo->PrimitiveType + 1;
+		uint8 PrimitiveIdx = ThumbnailInfo->PrimitiveType.GetIntValue() + 1;
 		if ( PrimitiveIdx >= TPT_MAX )
 		{
 			if ( ThumbnailInfo->PreviewMesh.IsValid() )
@@ -206,7 +206,7 @@ FReply SThumbnailEditModeTools::OnMouseButtonDown( const FGeometry& MyGeometry,
 			if (RenderInfo != nullptr && RenderInfo->Renderer != nullptr)
 			{
 				bModifiedThumbnailWhileDragging = false;
-				DragStartLocation = FIntPoint(MouseEvent.GetScreenSpacePosition().X, MouseEvent.GetScreenSpacePosition().Y);
+				DragStartLocation = FIntPoint(FMath::TruncToInt32(MouseEvent.GetScreenSpacePosition().X), FMath::TruncToInt32(MouseEvent.GetScreenSpacePosition().Y));
 
 				return FReply::Handled().CaptureMouse(AsShared()).UseHighPrecisionMouseMovement(AsShared()).PreventThrottling();
 			}
diff --git a/Engine/Source/Editor/ContentBrowser/Private/SourcesViewWidgets.h b/Engine/Source/Editor/ContentBrowser/Private/SourcesViewWidgets.h
--- a/Engine/Source/Editor/ContentBrowser/Private/SourcesViewWidgets.h
+++ b/Engine/Source/Editor/ContentBrowser/Private/SourcesViewWidgets.h
@@ -35,7 +35,7 @@ struct FSlateFontInfo;
 class SAssetTreeItem : public SCompoundWidget
 {
 public:
-	DECLARE_DELEGATE_FourParams( FOnNameChanged, const TSharedPtr<FTreeItem>& /*TreeItem*/, const FString& /*InProposedName*/, const FVector2D& /*MessageLocation*/, const ETextCommit::Type /*CommitType*/);
+	DECLARE_DELEGATE_FourParams( FOnNameChanged, const TSharedPtr<FTreeItem>& /*TreeItem*/, const FString& /*InProposedName*/, const UE::Slate::FDeprecateVector2DParameter& /*MessageLocation*/, const ETextCommit::Type /*CommitType*/);
 	DECLARE_DELEGATE_RetVal_ThreeParams( bool, FOnVerifyNameChanged, const TSharedPtr<FTreeItem>& /*TreeItem*/, const FString& /*InProposedName*/, FText& /*OutErrorMessage*/);
 
 	SLATE_BEGIN_ARGS( SAssetTreeItem )
diff --git a/Engine/Source/Editor/ContentBrowserData/ContentBrowserData.Build.cs b/Engine/Source/Editor/ContentBrowserData/ContentBrowserData.Build.cs
--- a/Engine/Source/Editor/ContentBrowserData/ContentBrowserData.Build.cs
+++ b/Engine/Source/Editor/ContentBrowserData/ContentBrowserData.Build.cs
@@ -6,6 +6,8 @@ public class ContentBrowserData : ModuleRules
 {
 	public ContentBrowserData(ReadOnlyTargetRules Target) : base(Target)
 	{
+		UnsafeTypeCastWarningLevel = WarningLevel.Error;
+
 		PublicDependencyModuleNames.AddRange(
 			new string[] {
 				"Core",
diff --git a/Engine/Source/Editor/ContentBrowserData/Private/ContentBrowserVirtualPathTree.cpp b/Engine/Source/Editor/ContentBrowserData/Private/ContentBrowserVirtualPathTree.cpp
--- a/Engine/Source/Editor/ContentBrowserData/Private/ContentBrowserVirtualPathTree.cpp
+++ b/Engine/Source/Editor/ContentBrowserData/Private/ContentBrowserVirtualPathTree.cpp
@@ -365,7 +365,7 @@ EContentBrowserPathType FContentBrowserVirtualPathTree::TryConvertVirtualPathToI
 	{
 		if (PathCharPtr >= PathStrEnd || *PathCharPtr == TEXT('/'))
 		{
-			const FName CheckPath(PathCharPtr - PathStr, PathStr);
+			const FName CheckPath(int32(PathCharPtr - PathStr), PathStr);
 			bool bCheckPathIsFullyVirtual = false;
 			if (PathExists(CheckPath, bCheckPathIsFullyVirtual))
 			{
diff --git a/Engine/Source/Editor/ContentBrowserData/Public/ContentBrowserVirtualPathTree.h b/Engine/Source/Editor/ContentBrowserData/Public/ContentBrowserVirtualPathTree.h
--- a/Engine/Source/Editor/ContentBrowserData/Public/ContentBrowserVirtualPathTree.h
+++ b/Engine/Source/Editor/ContentBrowserData/Public/ContentBrowserVirtualPathTree.h
@@ -48,9 +48,9 @@ public:
 	/** Recursively enumerates all child paths from the specified base path relative to this node */
 	bool EnumerateSubPaths(FName BasePath, TFunctionRef<bool(FName, FName)> Callback, bool bRecurse = true) const;
 
-	uint32 GetAllocatedSize(void) const
+	SIZE_T GetAllocatedSize(void) const
 	{
-		uint32 AllocatedSize = ParentPathToChildPaths.GetAllocatedSize() + ChildPathToParentPath.GetAllocatedSize() + VirtualToInternalMounts.GetAllocatedSize();
+		SIZE_T AllocatedSize = ParentPathToChildPaths.GetAllocatedSize() + ChildPathToParentPath.GetAllocatedSize() + VirtualToInternalMounts.GetAllocatedSize();
 
 		for (const TPair<FName, TSet<FName>>& Pair : ParentPathToChildPaths)
 		{
diff --git a/Engine/Source/Runtime/Slate/Public/Widgets/Views/SListView.h b/Engine/Source/Runtime/Slate/Public/Widgets/Views/SListView.h
--- a/Engine/Source/Runtime/Slate/Public/Widgets/Views/SListView.h
+++ b/Engine/Source/Runtime/Slate/Public/Widgets/Views/SListView.h
@@ -513,8 +513,8 @@ public:
 						ensure(NumItemsPerLine > 0);
 
 						// calculate total number of rows and row of current index (1 index)
-						const int32 NumLines = FMath::CeilToInt(NumItems / (float)NumItemsPerLine);
-						const int32 CurLine = FMath::CeilToInt((CurSelectionIndex + 1) / (float)NumItemsPerLine);
+						const int32 NumLines = FMath::CeilToInt((float)NumItems / (float)NumItemsPerLine);
+						const int32 CurLine = FMath::CeilToInt((float)(CurSelectionIndex + 1) / (float)NumItemsPerLine);
 
 						// if not on final row, assume a jagged list and select the final item
 						if (CurLine < NumLines)
@@ -2107,7 +2107,7 @@ protected:
 						if (WidgetTopLeft.ScrollAxis < ListViewTopLeft.ScrollAxis)
 						{
 							// This entry is clipped at the top/left, so simply set it as the new scroll offset target to bump it down into view
-							NewScrollOffset = static_cast<double>(IndexOfItem - NavigationScrollOffset);
+							NewScrollOffset = static_cast<double>(IndexOfItem) - NavigationScrollOffset;
 						}
 						else
 						{
diff --git a/Engine/Source/Runtime/Slate/Public/Widgets/Views/STileView.h b/Engine/Source/Runtime/Slate/Public/Widgets/Views/STileView.h
--- a/Engine/Source/Runtime/Slate/Public/Widgets/Views/STileView.h
+++ b/Engine/Source/Runtime/Slate/Public/Widgets/Views/STileView.h
@@ -393,7 +393,7 @@ protected:
 		const bool bWholeListVisible = this->DesiredScrollOffset == 0 && this->bWasAtEndOfList;
 		if (InAllowOverscroll == EAllowOverscroll::Yes && this->Overscroll.ShouldApplyOverscroll(this->DesiredScrollOffset == 0, this->bWasAtEndOfList, ScrollByAmountInSlateUnits))
 		{
-			const float UnclampedScrollDelta = ScrollByAmountInSlateUnits / GetNumItemsPerLine();
+			const float UnclampedScrollDelta = ScrollByAmountInSlateUnits / (float)GetNumItemsPerLine();
 			const float ActuallyScrolledBy = this->Overscroll.ScrollBy(MyGeometry, UnclampedScrollDelta);
 			if (ActuallyScrolledBy != 0.0f)
 			{
@@ -403,7 +403,7 @@ protected:
 		}
 		else if (!bWholeListVisible)
 		{
-			const double NewScrollOffset = this->DesiredScrollOffset + ((ScrollByAmountInSlateUnits * GetNumItemsPerLine()) / GetTileDimensions().ScrollAxis);
+			const double NewScrollOffset = this->DesiredScrollOffset + ((ScrollByAmountInSlateUnits * (float)GetNumItemsPerLine()) / GetTileDimensions().ScrollAxis);
 
 			return this->ScrollTo( (float)NewScrollOffset );
 		}
@@ -466,9 +466,9 @@ protected:
 					// Center the line in the view area
 					NewLineOffset -= NumLinesInView * 0.5f;
 					// Convert the line offset into an item offset
-					double NewScrollOffset = NewLineOffset * NumItemsPerLine;
+					double NewScrollOffset = NewLineOffset * (double)NumItemsPerLine;
 					// And clamp the scroll offset within the allowed limits
-					NewScrollOffset = FMath::Clamp(NewScrollOffset, 0., (double)(GetNumItemsBeingObserved() - NumItemsPerLine * NumLinesInView));
+					NewScrollOffset = FMath::Clamp(NewScrollOffset, 0., (double)GetNumItemsBeingObserved() - (double)NumItemsPerLine * NumLinesInView);
 
 					this->SetScrollOffset((float)NewScrollOffset);
 				}
@@ -478,14 +478,14 @@ protected:
 					if (LineOfItem == MinDisplayedLine)
 					{
 						// This line is clipped at the top/left, so set it as the new offset
-						this->SetScrollOffset(LineOfItem * NumItemsPerLine - (this->FixedLineScrollOffset.IsSet() && LineOfItem > 0 ? 0.f : this->NavigationScrollOffset));
+						this->SetScrollOffset((float)(LineOfItem * NumItemsPerLine) - (this->FixedLineScrollOffset.IsSet() && LineOfItem > 0 ? 0.f : this->NavigationScrollOffset));
 					}
 					else if (LineOfItem == MaxDisplayedLine)
 					{
 						// This line is clipped at the end, so we need to advance just enough to bring it fully into view
 						// Since all tiles are required to be of the same size, this is straightforward
-						const float NewLineOffset = LineOfItem - NumLinesInView + 1.f + (this->FixedLineScrollOffset.IsSet() ? 0.f : this->NavigationScrollOffset);
-						this->SetScrollOffset(NewLineOffset * NumItemsPerLine);
+						const float NewLineOffset = (float)LineOfItem - NumLinesInView + 1.f + (this->FixedLineScrollOffset.IsSet() ? 0.f : this->NavigationScrollOffset);
+						this->SetScrollOffset(NewLineOffset * (float)NumItemsPerLine);
 					}
 				}
 
