Added a filter to show only the unsupported asset file items

diff --git a/Engine/Source/Editor/ContentBrowser/Private/FrontendFilters.cpp b/Engine/Source/Editor/ContentBrowser/Private/FrontendFilters.cpp
--- a/Engine/Source/Editor/ContentBrowser/Private/FrontendFilters.cpp
+++ b/Engine/Source/Editor/ContentBrowser/Private/FrontendFilters.cpp
@@ -1548,4 +1548,14 @@ bool FFrontendFilter_VirtualizedData::PassesFilter(FAssetFilterType InItem) cons
 	{
 		return false;
 	}
-}
\ No newline at end of file
+}
+
+FFrontendFilter_Unsupported::FFrontendFilter_Unsupported(TSharedPtr<FFrontendFilterCategory> InCategory)
+	:FFrontendFilter (InCategory)
+{
+}
+
+bool FFrontendFilter_Unsupported::PassesFilter(FAssetFilterType InItem) const
+{
+	return !InItem.IsSupported();
+}
diff --git a/Engine/Source/Editor/ContentBrowser/Private/SFilterList.cpp b/Engine/Source/Editor/ContentBrowser/Private/SFilterList.cpp
--- a/Engine/Source/Editor/ContentBrowser/Private/SFilterList.cpp
+++ b/Engine/Source/Editor/ContentBrowser/Private/SFilterList.cpp
@@ -84,6 +84,7 @@ void SFilterList::Construct( const FArguments& InArgs )
 	AllFrontendFilters_Internal.Add(MakeShareable(new FFrontendFilter_Recent(DefaultCategory)));
 	AllFrontendFilters_Internal.Add( MakeShareable(new FFrontendFilter_NotSourceControlled(DefaultCategory)) );
 	AllFrontendFilters_Internal.Add(MakeShareable(new FFrontendFilter_VirtualizedData(DefaultCategory)));
+	AllFrontendFilters_Internal.Add(MakeShared<FFrontendFilter_Unsupported>(DefaultCategory));
 
 	// Add any global user-defined frontend filters
 	for (TObjectIterator<UContentBrowserFrontEndFilterExtension> ExtensionIt(RF_NoFlags); ExtensionIt; ++ExtensionIt)
diff --git a/Engine/Source/Editor/ContentBrowser/Public/FrontendFilters.h b/Engine/Source/Editor/ContentBrowser/Public/FrontendFilters.h
--- a/Engine/Source/Editor/ContentBrowser/Public/FrontendFilters.h
+++ b/Engine/Source/Editor/ContentBrowser/Public/FrontendFilters.h
@@ -426,4 +426,20 @@ private:
 	virtual bool PassesFilter(FAssetFilterType InItem) const override;
 };
 
+/** A filter that displays only assets the are mark as unsupported */
+class CONTENTBROWSER_API FFrontendFilter_Unsupported : public FFrontendFilter
+{
+public:
+	FFrontendFilter_Unsupported(TSharedPtr<FFrontendFilterCategory> InCategory);
+	~FFrontendFilter_Unsupported() = default;
+
+	// FFrontendFilter implementation
+	virtual FString GetName() const override { return TEXT("Unsupported"); }
+	virtual FText GetDisplayName() const override { return LOCTEXT("FrontendFilter_Unsupported", "Unsupported"); }
+	virtual FText GetToolTipText() const override { return LOCTEXT("FrontendFilter_UnsupportedTooltip", "Show only assets that are not supported by this project."); }
+
+	// IFilter implementation
+	virtual bool PassesFilter(FAssetFilterType InItem) const override;
+};
+
 #undef LOCTEXT_NAMESPACE
