Content Browser: Improve UX around unsupported assets

diff --git a/Engine/Plugins/Editor/ContentBrowser/ContentBrowserAssetDataSource/Source/ContentBrowserAssetDataSource/Private/AssetFileContextMenu.cpp b/Engine/Plugins/Editor/ContentBrowser/ContentBrowserAssetDataSource/Source/ContentBrowserAssetDataSource/Private/AssetFileContextMenu.cpp
--- a/Engine/Plugins/Editor/ContentBrowser/ContentBrowserAssetDataSource/Source/ContentBrowserAssetDataSource/Private/AssetFileContextMenu.cpp
+++ b/Engine/Plugins/Editor/ContentBrowser/ContentBrowserAssetDataSource/Source/ContentBrowserAssetDataSource/Private/AssetFileContextMenu.cpp
@@ -79,6 +79,13 @@ void FAssetFileContextMenu::MakeContextMenu(UToolMenu* InMenu, const TArray<FAss
 	SelectedAssets = InSelectedAssets;
 	OnShowAssetsInPathsView = InOnShowAssetsInPathsView;
 
+	// We don't want the regular menu options if the selection contains unsupported assets
+	const UContentBrowserDataMenuContext_FileMenu* Context = InMenu->FindContext<UContentBrowserDataMenuContext_FileMenu>();
+	if(Context && Context->bContainsUnsupportedAssets)
+	{
+		return;
+	}
+	
 	if (SelectedAssets.Num() > 0)
 	{
 		AddMenuOptions(InMenu);
diff --git a/Engine/Source/Editor/ContentBrowser/Private/AssetContextMenu.cpp b/Engine/Source/Editor/ContentBrowser/Private/AssetContextMenu.cpp
--- a/Engine/Source/Editor/ContentBrowser/Private/AssetContextMenu.cpp
+++ b/Engine/Source/Editor/ContentBrowser/Private/AssetContextMenu.cpp
@@ -191,6 +191,16 @@ TSharedRef<SWidget> FAssetContextMenu::MakeContextMenu(TArrayView<const FContent
 		}
 		FAssetToolsModule& AssetToolsModule = FModuleManager::LoadModuleChecked<FAssetToolsModule>("AssetTools");
 		const TSharedRef<FPathPermissionList>& WritableFolderPermission = AssetToolsModule.Get().GetWritableFolderPermissionList();
+		
+		ContextObject->bContainsUnsupportedAssets = false;
+		for (const FContentBrowserItem& SelectedItem : SelectedItems)
+		{
+			if(!SelectedItem.IsSupported())
+			{
+				ContextObject->bContainsUnsupportedAssets = true;
+				break;
+			}
+		}
 
 		const TArray<FAssetData>& SelectedAssets = ContextObject->SelectedAssets;
 
@@ -307,6 +317,7 @@ TSharedRef<SWidget> FAssetContextMenu::MakeContextMenu(TArrayView<const FContent
 		DataContextObject->SelectedItems = SelectedItems;
 		DataContextObject->SelectedCollections = SourcesData.Collections;
 		DataContextObject->bCanBeModified = ContextObject->bCanBeModified;
+		DataContextObject->bContainsUnsupportedAssets = ContextObject->bContainsUnsupportedAssets;
 		DataContextObject->ParentWidget = AssetView;
 		DataContextObject->OnShowInPathsView = OnShowInPathsViewRequested;
 		DataContextObject->OnRefreshView = OnAssetViewRefreshRequested;
diff --git a/Engine/Source/Editor/ContentBrowser/Private/AssetViewWidgets.cpp b/Engine/Source/Editor/ContentBrowser/Private/AssetViewWidgets.cpp
--- a/Engine/Source/Editor/ContentBrowser/Private/AssetViewWidgets.cpp
+++ b/Engine/Source/Editor/ContentBrowser/Private/AssetViewWidgets.cpp
@@ -771,6 +771,7 @@ TSharedRef<SWidget> SAssetViewItem::CreateToolTipWidget() const
 			const FText NameText = GetNameText();
 			const FText ClassText = FText::Format(LOCTEXT("ClassName", "({0})"), GetAssetClassText());
 			FText PublicStateText = LOCTEXT("PublicAssetState", "Public");
+			FName PublicStateTextBorder = "ContentBrowser.TileViewTooltip.PillBorder";
 
 			// Create a box to hold every line of info in the body of the tooltip
 			TSharedRef<SVerticalBox> InfoBox = SNew(SVerticalBox);
@@ -825,6 +826,12 @@ TSharedRef<SWidget> SAssetViewItem::CreateToolTipWidget() const
 				PublicStateText = LOCTEXT("ReadOnlyAssetState", "Read Only");
 			}
 
+			if(!AssetItem->GetItem().IsSupported())
+			{
+				PublicStateText = LOCTEXT("UnsupportedAssetState", "Unsupported");
+				PublicStateTextBorder = "ContentBrowser.TileViewTooltip.UnsupportedAssetPillBorder";
+			}
+
 			// Add tags
 			for (const auto& DisplayTagItem : CachedDisplayTags)
 			{
@@ -900,7 +907,7 @@ TSharedRef<SWidget> SAssetViewItem::CreateToolTipWidget() const
 							.HAlign(HAlign_Right) 
 							[
 								SNew(SBorder)
-								.BorderImage(FAppStyle::GetBrush("ContentBrowser.TileViewTooltip.PillBorder"))
+								.BorderImage(FAppStyle::GetBrush(PublicStateTextBorder))
 								.Padding(FMargin(12.0f, 2.0f, 12.0f, 2.0f))
 								.Visibility(bIsPublicAssetUIEnabled ? EVisibility::Visible : EVisibility::Hidden)
 								[
@@ -912,6 +919,18 @@ TSharedRef<SWidget> SAssetViewItem::CreateToolTipWidget() const
 							]
 						]
 
+						+ SVerticalBox::Slot()
+						.AutoHeight()
+						[
+							SNew(STextBlock)
+							.Visibility_Lambda([this]()
+							{
+								return AssetItem->GetItem().IsSupported() ? EVisibility::Collapsed : EVisibility::Visible;
+							})
+							.Text(LOCTEXT("UnsupportedAssetDescriptionText", "This type of asset is not allowed in this project. Delete unsupported assets to avoid errors."))
+							.ColorAndOpacity(FStyleColors::Warning)
+						]
+
 						+ SVerticalBox::Slot()
 						.AutoHeight()
 						[
@@ -1501,6 +1520,11 @@ FSlateColor SAssetViewItem::GetAssetColor() const
 				}
 			}
 		}
+
+		if(!AssetItem->GetItem().IsSupported())
+		{
+			return FSlateColor::UseForeground();
+		}
 	}
 	return ContentBrowserUtils::GetDefaultColor();
 }
@@ -1579,6 +1603,11 @@ void SAssetListItem::Construct( const FArguments& InArgs )
 		ThumbnailConfig.HighlightedText = InArgs._HighlightText;
 		ThumbnailConfig.HintColorAndOpacity = InArgs._ThumbnailHintColorAndOpacity;
 
+		if(!AssetItem->GetItem().IsSupported())
+		{
+			ThumbnailConfig.ClassThumbnailBrushOverride = FName("Icons.WarningWithColor.Thumbnail");
+		}
+
 		{
 			FContentBrowserItemDataAttributeValue ColorAttributeValue = AssetItem->GetItem().GetItemAttribute(ContentBrowserItemAttributes::ItemColor);
 			if (ColorAttributeValue.IsValid())
@@ -1821,6 +1850,11 @@ void SAssetTileItem::Construct( const FArguments& InArgs )
 		ThumbnailConfig.Padding= FMargin(2.0f);
 		ThumbnailConfig.GenericThumbnailSize = MakeAttributeSP(this, &SAssetTileItem::GetGenericThumbnailSize);
 
+		if(!AssetItem->GetItem().IsSupported())
+		{
+			ThumbnailConfig.ClassThumbnailBrushOverride = FName("Icons.WarningWithColor.Thumbnail");
+		}
+		
 		{
 			FContentBrowserItemDataAttributeValue ColorAttributeValue = AssetItem->GetItem().GetItemAttribute(ContentBrowserItemAttributes::ItemColor);
 			if (ColorAttributeValue.IsValid())
@@ -2321,7 +2355,14 @@ TSharedRef<SWidget> SAssetColumnItem::GenerateWidgetForColumn( const FName& Colu
 		}
 		else
 		{
-			IconBrush = FAppStyle::GetBrush("ContentBrowser.ColumnViewAssetIcon");
+			if(!AssetItem->GetItem().IsSupported())
+			{
+				IconBrush = FAppStyle::GetBrush("Icons.WarningWithColor");
+			}
+			else
+			{
+				IconBrush = FAppStyle::GetBrush("ContentBrowser.ColumnViewAssetIcon");
+			}
 		}
 
 		// Make icon overlays (eg, SCC and dirty status) a reasonable size in relation to the icon size (note: it is assumed this icon is square)
diff --git a/Engine/Source/Editor/ContentBrowser/Public/ContentBrowserMenuContexts.h b/Engine/Source/Editor/ContentBrowser/Public/ContentBrowserMenuContexts.h
--- a/Engine/Source/Editor/ContentBrowser/Public/ContentBrowserMenuContexts.h
+++ b/Engine/Source/Editor/ContentBrowser/Public/ContentBrowserMenuContexts.h
@@ -63,6 +63,9 @@ public:
 	UPROPERTY()
 	bool bCanBeModified;
 
+	UPROPERTY(BlueprintReadOnly, Category = ContentBrowser)
+	bool bContainsUnsupportedAssets = true;
+
 	//UE_DEPRECATED(5.2, "GetSelectedObjects has been deprecated.  We no longer automatically load assets on right click.  Please use SelectedAssets and determine whatever you need for your context menu options without actually loading the assets.  When you finally need all or a subset of the selected assets use LoadSelectedAssets or LoadSelectedAssetsIf")
 	UFUNCTION(BlueprintCallable, Category="Tool Menus", meta=(DeprecatedFunction, DeprecationMessage = "GetSelectedObjects has been deprecated.  We no longer automatically load assets on right click.  If you can work without loading the assets, please use SelectedAssets.  Otherwise call LoadSelectedObjects"))
 	TArray<UObject*> GetSelectedObjects() const
diff --git a/Engine/Source/Editor/ContentBrowserData/Public/ContentBrowserDataMenuContexts.h b/Engine/Source/Editor/ContentBrowserData/Public/ContentBrowserDataMenuContexts.h
--- a/Engine/Source/Editor/ContentBrowserData/Public/ContentBrowserDataMenuContexts.h
+++ b/Engine/Source/Editor/ContentBrowserData/Public/ContentBrowserDataMenuContexts.h
@@ -74,6 +74,9 @@ public:
 	UPROPERTY(BlueprintReadOnly, Category = ContentBrowser)
 	bool bCanBeModified = true;
 
+	UPROPERTY(BlueprintReadOnly, Category = ContentBrowser)
+	bool bContainsUnsupportedAssets = true;
+
 	TWeakPtr<SWidget> ParentWidget;
 
 	FOnShowInPathsView OnShowInPathsView;
diff --git a/Engine/Source/Editor/UnrealEd/Private/AssetThumbnail.cpp b/Engine/Source/Editor/UnrealEd/Private/AssetThumbnail.cpp
--- a/Engine/Source/Editor/UnrealEd/Private/AssetThumbnail.cpp
+++ b/Engine/Source/Editor/UnrealEd/Private/AssetThumbnail.cpp
@@ -495,7 +495,7 @@ private:
 		if(!bHasRenderedThumbnail)
 		{
 			const FSlateBrush* ClassThumbnailBrush = GetClassThumbnailBrush();
-			if( ClassThumbnailBrush && ThumbnailClass.Get() )
+			if( (ClassThumbnailBrush && ThumbnailClass.Get())  || !ClassThumbnailBrushOverride.IsNone() )
 			{
 				return EVisibility::Visible;
 			}
diff --git a/Engine/Source/Runtime/SlateCore/Private/Styling/StarshipCoreStyle.cpp b/Engine/Source/Runtime/SlateCore/Private/Styling/StarshipCoreStyle.cpp
--- a/Engine/Source/Runtime/SlateCore/Private/Styling/StarshipCoreStyle.cpp
+++ b/Engine/Source/Runtime/SlateCore/Private/Styling/StarshipCoreStyle.cpp
@@ -321,6 +321,8 @@ TSharedRef<ISlateStyle> FStarshipCoreStyle::Create()
 		Style->Set("Icons.Warning.Large", new IMAGE_BRUSH_SVG("Starship/Common/alert-triangle-large", Icon32x32));
 		Style->Set("Icons.WarningWithColor.Large", new IMAGE_BRUSH_SVG("Starship/Common/alert-triangle-large", Icon32x32, FStyleColors::Warning));
 
+		Style->Set("Icons.WarningWithColor.Thumbnail", new IMAGE_BRUSH_SVG("Starship/Common/alert-triangle-64", Icon64x64));
+		
 		Style->Set("Icons.AlertCircle", new IMAGE_BRUSH_SVG("Starship/Common/alert-circle", Icon16x16));
 		Style->Set("Icons.AlertCircleWithColor", new IMAGE_BRUSH_SVG("Starship/Common/alert-circle", Icon16x16, FStyleColors::Warning));
 		
