Remove the capacity to copy and paste the unsupported assets.

diff --git a/Engine/Source/Editor/ContentBrowser/Private/SAssetView.cpp b/Engine/Source/Editor/ContentBrowser/Private/SAssetView.cpp
--- a/Engine/Source/Editor/ContentBrowser/Private/SAssetView.cpp
+++ b/Engine/Source/Editor/ContentBrowser/Private/SAssetView.cpp
@@ -1585,6 +1585,7 @@ FReply SAssetView::OnKeyDown( const FGeometry& MyGeometry, const FKeyEvent& InKe
 
 			// Get assets and copy them
 			TArray<UObject*> AssetsToCopy;
+			FAssetToolsModule& AssetToolsModule = FModuleManager::LoadModuleChecked<FAssetToolsModule>("AssetTools");
 			for (const FString& AssetPath : AssetPathsSplit)
 			{
 				// Validate string
@@ -1593,15 +1594,19 @@ FReply SAssetView::OnKeyDown( const FGeometry& MyGeometry, const FKeyEvent& InKe
 				FString PackageName;
 				if (IsValidObjectPath(AssetPath, ObjectClassName, ObjectPath, PackageName))
 				{
-					FLinkerInstancingContext InstancingContext({ ULevel::LoadAllExternalObjectsTag });
-					UObject* ObjectToCopy = LoadObject<UObject>(nullptr, *ObjectPath, nullptr, LOAD_None, nullptr, &InstancingContext);
-					if (ObjectToCopy && !ObjectToCopy->IsA(UClass::StaticClass()))
+					// Only duplicate the objects of the supported classes.
+					if (AssetToolsModule.Get().GetAssetClassPathPermissionList(EAssetClassAction::ViewAsset)->PassesStartsWithFilter(ObjectClassName))
 					{
-						AssetsToCopy.Add(ObjectToCopy);
+						FLinkerInstancingContext InstancingContext({ ULevel::LoadAllExternalObjectsTag });
+						UObject* ObjectToCopy = LoadObject<UObject>(nullptr, *ObjectPath, nullptr, LOAD_None, nullptr, &InstancingContext);
+						if (ObjectToCopy && !ObjectToCopy->IsA(UClass::StaticClass()))
+						{
+							AssetsToCopy.Add(ObjectToCopy);
+						}
 					}
 				}
 			}
-
+			
 			if (AssetsToCopy.Num())
 			{
 				UContentBrowserDataSubsystem* ContentBrowserData = IContentBrowserDataModule::Get().GetSubsystem();
