Selecting a folder in the Content Browser with an empty ContentBrowser.json config file no longer occasionally crashes.

The EditorConfig argument was being passed to SPathView as a raw pointer, but opening a new Content Browser tab or opening the drawer could cause the config array to be resized and invalidate the pointer.

Instead, the content browser instance name is now passed to SPathView and the config is looked up each time.

diff --git a/Engine/Source/Editor/ContentBrowser/Private/SContentBrowser.cpp b/Engine/Source/Editor/ContentBrowser/Private/SContentBrowser.cpp
--- a/Engine/Source/Editor/ContentBrowser/Private/SContentBrowser.cpp
+++ b/Engine/Source/Editor/ContentBrowser/Private/SContentBrowser.cpp
@@ -1050,8 +1050,6 @@ void SContentBrowser::SetPathViewExpanded(bool bExpanded)
 
 TSharedRef<SWidget> SContentBrowser::CreatePathView(const FContentBrowserConfig* Config)
 {	
-	FContentBrowserInstanceConfig* InstanceConfig = GetMutableInstanceConfig();
-
 	return
 		SAssignNew(PathArea, SExpandableArea)
 		.BorderImage(FAppStyle::Get().GetBrush("Brushes.Header"))
@@ -1110,7 +1108,7 @@ TSharedRef<SWidget> SContentBrowser::CreatePathView(const FContentBrowserConfig*
 					.AddMetaData<FTagMetaData>(FTagMetaData(TEXT("ContentBrowserSources")))
 					.ExternalSearch(SourcesSearch)
 					.PluginPathFilters(PluginPathFilters)
-					.EditorConfig(InstanceConfig != nullptr ? &InstanceConfig->PathView : nullptr)
+					.OwningContentBrowserName(InstanceName)
 				]
 			]
 		];
diff --git a/Engine/Source/Editor/ContentBrowser/Private/SPathView.cpp b/Engine/Source/Editor/ContentBrowser/Private/SPathView.cpp
--- a/Engine/Source/Editor/ContentBrowser/Private/SPathView.cpp
+++ b/Engine/Source/Editor/ContentBrowser/Private/SPathView.cpp
@@ -135,7 +135,7 @@ SPathView::~SPathView()
 
 void SPathView::Construct( const FArguments& InArgs )
 {
-	EditorConfig = InArgs._EditorConfig;
+	OwningContentBrowserName = InArgs._OwningContentBrowserName;
 	OnItemSelectionChanged = InArgs._OnItemSelectionChanged;
 	bAllowContextMenu = InArgs._AllowContextMenu;
 	OnGetItemContextMenu = InArgs._OnGetItemContextMenu;
@@ -419,7 +419,7 @@ void SPathView::SetPluginPathFilterActive(const TSharedRef<FContentBrowserPlugin
 		PluginPathFilters->Remove(Filter);
 	}
 
-	if (EditorConfig != nullptr)
+	if (FPathViewConfig* EditorConfig = GetEditorConfig())
 	{
 		if (bActive)
 		{
@@ -434,6 +434,22 @@ void SPathView::SetPluginPathFilterActive(const TSharedRef<FContentBrowserPlugin
 	}
 }
 
+FPathViewConfig* SPathView::GetEditorConfig() const
+{
+	if (OwningContentBrowserName.IsNone())
+	{
+		return nullptr;
+	}
+
+	FContentBrowserInstanceConfig* Config = UContentBrowserConfig::Get()->Instances.Find(OwningContentBrowserName);
+	if (Config == nullptr)
+	{
+		return nullptr;
+	}
+	
+	return &Config->PathView;
+}
+
 void SPathView::SetSelectedPaths(const TArray<FName>& Paths)
 {
 	TArray<FString> PathStrings;
@@ -540,7 +556,7 @@ void SPathView::SetSelectedPaths(const TArray<FString>& Paths)
 		}
 	}
 
-	if (EditorConfig != nullptr)
+	if (FPathViewConfig* EditorConfig = GetEditorConfig())
 	{
 		EditorConfig->SelectedPaths = LastSelectedPaths.Array();
 
@@ -1134,7 +1150,7 @@ void SPathView::LoadSettings(const FString& IniFilename, const FString& IniSecti
 {
 	// Selected Paths
 	TArray<FName> NewSelectedPaths;
-	if (EditorConfig != nullptr)
+	if (FPathViewConfig* EditorConfig = GetEditorConfig())
 	{
 		NewSelectedPaths = EditorConfig->SelectedPaths;
 	}
@@ -1211,7 +1227,7 @@ void SPathView::LoadSettings(const FString& IniFilename, const FString& IniSecti
 	if (PluginPathFilters.IsValid())
 	{
 		TArray<FString> NewSelectedFilters;
-		if (EditorConfig != nullptr)
+		if (FPathViewConfig* EditorConfig = GetEditorConfig())
 		{
 			NewSelectedFilters = EditorConfig->PluginFilters;
 		}
@@ -1387,7 +1403,7 @@ void SPathView::TreeSelectionChanged( TSharedPtr< FTreeItem > TreeItem, ESelectI
 			}
 		}
 
-		if (EditorConfig != nullptr)
+		if (FPathViewConfig* EditorConfig = GetEditorConfig())
 		{
 			EditorConfig->SelectedPaths = LastSelectedPaths.Array();
 
diff --git a/Engine/Source/Editor/ContentBrowser/Private/SPathView.h b/Engine/Source/Editor/ContentBrowser/Private/SPathView.h
--- a/Engine/Source/Editor/ContentBrowser/Private/SPathView.h
+++ b/Engine/Source/Editor/ContentBrowser/Private/SPathView.h
@@ -76,7 +76,6 @@ public:
 		, _AllowReadOnlyFolders(true)
 		, _ShowFavorites(false)
 		, _SelectionMode( ESelectionMode::Multi )
-		, _EditorConfig(nullptr)
 		{}
 
 		/** Content displayed to the left of the search bar */
@@ -127,7 +126,8 @@ public:
 		/** The plugin filter collection */
 		SLATE_ARGUMENT( TSharedPtr<FPluginFilterCollectionType>, PluginPathFilters)
 
-		SLATE_ARGUMENT( FPathViewConfig*, EditorConfig )
+		/** The instance name of the owning content browser. */
+		SLATE_ARGUMENT( FName, OwningContentBrowserName )
 
 	SLATE_END_ARGS()
 
@@ -303,6 +303,9 @@ protected:
 	/** Clear all root items and clear selection */
 	void ClearTreeItems();
 
+	/** Get this path view's editor config if OwningContentBrowserName is set. */
+	FPathViewConfig* GetEditorConfig() const;
+
 private:
 	/** Selects the given path only if it exists. Returns true if selected. */
 	bool ExplicitlyAddPathToSelection(const FName Path);
@@ -483,7 +486,7 @@ private:
 	TSharedPtr<SExpandableArea> FavoritesArea;
 
 	/** The config instance to use. */
-	FPathViewConfig* EditorConfig = nullptr;
+	FName OwningContentBrowserName;
 };
 
 
