Content Browser: Folder names now wrap properly with ellipses similar to asset names

diff --git a/Engine/Source/Editor/ContentBrowser/Private/AssetViewWidgets.cpp b/Engine/Source/Editor/ContentBrowser/Private/AssetViewWidgets.cpp
--- a/Engine/Source/Editor/ContentBrowser/Private/AssetViewWidgets.cpp
+++ b/Engine/Source/Editor/ContentBrowser/Private/AssetViewWidgets.cpp
@@ -1920,6 +1920,7 @@ void SAssetTileItem::Construct( const FArguments& InArgs )
 							+ SVerticalBox::Slot()
 							.Padding(2.0f,2.0f,0.0f,0.0f)
 							.VAlign(VAlign_Top)
+							.HAlign(IsFolder() ? HAlign_Center : HAlign_Left)
 							[
 								SNew(SBox)
 								.MaxDesiredHeight(this, &SAssetTileItem::GetNameAreaMaxDesiredHeight)
@@ -1933,7 +1934,6 @@ void SAssetTileItem::Construct( const FArguments& InArgs )
 									.HighlightText(InArgs._HighlightText)
 									.IsSelected(InArgs._IsSelectedExclusively)
 									.IsReadOnly(this, &SAssetTileItem::IsNameReadOnly)
-									.Justification(IsFolder() ? ETextJustify::Center : ETextJustify::Left)
 									.LineBreakPolicy(FBreakIterator::CreateCamelCaseBreakIterator())
 									.OverflowPolicy(ETextOverflowPolicy::Ellipsis)
 									.ColorAndOpacity(this, &SAssetTileItem::GetNameAreaTextColor)
