Browsing to assets where the current settings would cause the asset to be hidden now updates the settings of the active Content Browser to show the asset(s). This is done by directly changing the setting as many systems are listening for Content Browser settings changes to refresh their display to the user.

diff --git a/Engine/Source/Editor/ContentBrowser/Private/SAssetView.cpp b/Engine/Source/Editor/ContentBrowser/Private/SAssetView.cpp
--- a/Engine/Source/Editor/ContentBrowser/Private/SAssetView.cpp
+++ b/Engine/Source/Editor/ContentBrowser/Private/SAssetView.cpp
@@ -3053,8 +3053,6 @@ void SAssetView::ToggleShowPluginContent()
 {
 	check(IsToggleShowPluginContentAllowed());
 
-	bOverrideShowPluginContent.Reset();
-
 	bool bNewState = !GetDefault<UContentBrowserSettings>()->GetDisplayPluginFolders();
 
 	if (FContentBrowserInstanceConfig* EditorConfig = GetContentBrowserConfig())
@@ -3075,11 +3073,6 @@ bool SAssetView::IsShowingPluginContent() const
 		return true;
 	}
 
-	if (bOverrideShowPluginContent.IsSet())
-	{
-		return bOverrideShowPluginContent.GetValue();
-	}
-
 	if (const FContentBrowserInstanceConfig* Config = GetContentBrowserConfig())
 	{
 		return Config->bShowPluginContent;
@@ -3092,8 +3085,6 @@ void SAssetView::ToggleShowEngineContent()
 {
 	check(IsToggleShowEngineContentAllowed());
 
-	bOverrideShowEngineContent.Reset();
-
 	bool bNewState = !GetDefault<UContentBrowserSettings>()->GetDisplayEngineFolder();
 
 	if (FContentBrowserInstanceConfig* EditorConfig = GetContentBrowserConfig())
@@ -3114,11 +3105,6 @@ bool SAssetView::IsShowingEngineContent() const
 		return true;
 	}
 
-	if (bOverrideShowEngineContent.IsSet())
-	{
-		return bOverrideShowEngineContent.GetValue();
-	}
-
 	if (const FContentBrowserInstanceConfig* Config = GetContentBrowserConfig())
 	{
 		return Config->bShowEngineContent;
@@ -3131,8 +3117,6 @@ void SAssetView::ToggleShowDevelopersContent()
 {
 	check(IsToggleShowDevelopersContentAllowed());
 
-	bOverrideShowDeveloperContent.Reset();
-
 	bool bNewState = !GetDefault<UContentBrowserSettings>()->GetDisplayDevelopersFolder();
 
 	if (FContentBrowserInstanceConfig* EditorConfig = GetContentBrowserConfig())
@@ -3168,11 +3152,6 @@ bool SAssetView::IsShowingDevelopersContent() const
 		return false;
 	}
 
-	if (bOverrideShowDeveloperContent.IsSet())
-	{
-		return bOverrideShowDeveloperContent.GetValue();
-	}
-
 	if (const FContentBrowserInstanceConfig* Config = GetContentBrowserConfig())
 	{
 		return Config->bShowDeveloperContent;
@@ -3187,8 +3166,6 @@ void SAssetView::ToggleShowLocalizedContent()
 
 	bool bNewState = !GetDefault<UContentBrowserSettings>()->GetDisplayL10NFolder();
 
-	bOverrideShowLocalizedContent.Reset();
-
 	if (FContentBrowserInstanceConfig* Config = GetContentBrowserConfig())
 	{
 		bNewState = !Config->bShowLocalizedContent;
@@ -3212,11 +3189,6 @@ bool SAssetView::IsShowingLocalizedContent() const
 		return false;
 	}
 
-	if (bOverrideShowLocalizedContent.IsSet())
-	{
-		return bOverrideShowLocalizedContent.GetValue();
-	}
-
 	if (const FContentBrowserInstanceConfig* Config = GetContentBrowserConfig())
 	{
 		return Config->bShowLocalizedContent;
@@ -4737,24 +4709,37 @@ void SAssetView::ForceShowPluginFolder(bool bEnginePlugin)
 	}
 }
 
-void SAssetView::OverrideShowEngineContent( bool bShow )
+void SAssetView::OverrideShowEngineContent()
 {
-	bOverrideShowEngineContent = bShow;
+	if (!IsShowingEngineContent())
+	{
+		ToggleShowEngineContent();
+	}
+
 }
 
-void SAssetView::OverrideShowDeveloperContent( bool bShow )
+void SAssetView::OverrideShowDeveloperContent()
 {
-	bOverrideShowDeveloperContent = bShow;
+	if (!IsShowingDevelopersContent())
+	{
+		ToggleShowDevelopersContent();
+	}
 }
 
-void SAssetView::OverrideShowPluginContent( bool bShow )
+void SAssetView::OverrideShowPluginContent()
 {
-	bOverrideShowPluginContent = bShow;
+	if (!IsShowingPluginContent())
+	{
+		ToggleShowPluginContent();
+	}
 }
 
-void SAssetView::OverrideShowLocalizedContent( bool bShow )
+void SAssetView::OverrideShowLocalizedContent()
 {
-	bOverrideShowLocalizedContent = bShow;
+	if (!IsShowingLocalizedContent())
+	{
+		ToggleShowLocalizedContent();
+	}
 }
 
 void SAssetView::HandleItemDataUpdated(TArrayView<const FContentBrowserItemDataUpdate> InUpdatedItems)
diff --git a/Engine/Source/Editor/ContentBrowser/Private/SContentBrowser.cpp b/Engine/Source/Editor/ContentBrowser/Private/SContentBrowser.cpp
--- a/Engine/Source/Editor/ContentBrowser/Private/SContentBrowser.cpp
+++ b/Engine/Source/Editor/ContentBrowser/Private/SContentBrowser.cpp
@@ -1406,6 +1406,13 @@ void SContentBrowser::PrepareToSyncItems(TArrayView<const FContentBrowserItem> I
 	bool bDisplayEngine = GetDefault<UContentBrowserSettings>()->GetDisplayEngineFolder();
 	bool bDisplayPlugins = GetDefault<UContentBrowserSettings>()->GetDisplayPluginFolders();
 	bool bDisplayLocalized = GetDefault<UContentBrowserSettings>()->GetDisplayL10NFolder();
+
+	// Keep track of any of the settings changing so we can let the user know
+	bool bDisplayDevChanged = false;
+	bool bDisplayEngineChanged = false;
+	bool bDisplayPluginsChanged = false;
+	bool bDisplayLocalizedChanged = false;
+
 	if ( !bDisplayDev || !bDisplayEngine || !bDisplayPlugins || !bDisplayLocalized )
 	{
 		for (const FContentBrowserItem& ItemToSync : ItemsToSync)
@@ -1413,29 +1420,33 @@ void SContentBrowser::PrepareToSyncItems(TArrayView<const FContentBrowserItem> I
 			if (!bDisplayDev && ContentBrowserUtils::IsItemDeveloperContent(ItemToSync))
 			{
 				bDisplayDev = true;
-				AssetViewPtr->OverrideShowDeveloperContent(true);
+				AssetViewPtr->OverrideShowDeveloperContent();
 				bRepopulate = true;
+				bDisplayDevChanged = true;
 			}
 
 			if (!bDisplayEngine && ContentBrowserUtils::IsItemEngineContent(ItemToSync))
 			{
 				bDisplayEngine = true;
-				AssetViewPtr->OverrideShowEngineContent(true);
+				AssetViewPtr->OverrideShowEngineContent();
 				bRepopulate = true;
+				bDisplayEngineChanged = true;
 			}
 
 			if (!bDisplayPlugins && ContentBrowserUtils::IsItemPluginContent(ItemToSync))
 			{
 				bDisplayPlugins = true;
-				AssetViewPtr->OverrideShowPluginContent(true);
+				AssetViewPtr->OverrideShowPluginContent();
 				bRepopulate = true;
+				bDisplayPluginsChanged = true;
 			}
 
 			if (!bDisplayLocalized && ContentBrowserUtils::IsItemLocalizedContent(ItemToSync))
 			{
 				bDisplayLocalized = true;
-				AssetViewPtr->OverrideShowLocalizedContent(true);
+				AssetViewPtr->OverrideShowLocalizedContent();
 				bRepopulate = true;
+				bDisplayLocalizedChanged = true;
 			}
 
 			if (bDisplayDev && bDisplayEngine && bDisplayPlugins && bDisplayLocalized)
@@ -1463,6 +1474,41 @@ void SContentBrowser::PrepareToSyncItems(TArrayView<const FContentBrowserItem> I
 	// If we have auto-enabled any flags or found a non-existant path, force a refresh
 	if (bRepopulate)
 	{
+		// let the user know if one of their settings is being changed to be able to show the sync targets
+		if (bDisplayDevChanged || bDisplayEngineChanged || bDisplayPluginsChanged || bDisplayLocalizedChanged)
+		{
+			TArray<FText> SettingsText;
+			if (bDisplayDevChanged)
+			{
+				SettingsText.Add(LOCTEXT("ShowDeveloperContent", "Show Developer Content"));
+			}
+			if (bDisplayEngineChanged)
+			{
+				SettingsText.Add(LOCTEXT("ShowEngineContent", "Show Engine Content"));
+			}
+			if (bDisplayPluginsChanged)
+			{
+				SettingsText.Add(LOCTEXT("ShowPluginContent", "Show Plugin Content"));
+			}
+			if (bDisplayLocalizedChanged)
+			{
+				SettingsText.Add(LOCTEXT("ShowLocalizedContent", "Show Localized Content"));
+			}
+			FTextBuilder NotificationBuilder;
+			const FText NotificationPrefix = FText::Format(
+				LOCTEXT("AssetRequiresFilterChanges", "To show {0}|plural(one=this asset,other=these assets), the following {1}|plural(one=setting has,other=settings have) been enabled for the active Content Browser:\n"),
+				ItemsToSync.Num(),
+				SettingsText.Num());
+			NotificationBuilder.AppendLine(NotificationPrefix);
+			NotificationBuilder.Indent();
+			for ( const FText& SettingsTextEntry : SettingsText)
+			{
+				NotificationBuilder.AppendLine(SettingsTextEntry);
+			}
+	
+			FNotificationInfo NotificationInfo(NotificationBuilder.ToText());
+			FSlateNotificationManager::Get().AddNotification(NotificationInfo);
+		}
 		PathViewPtr->Populate();
 		FavoritePathViewPtr->Populate();
 	}
diff --git a/Engine/Source/Editor/ContentBrowser/Public/SAssetView.h b/Engine/Source/Editor/ContentBrowser/Public/SAssetView.h
--- a/Engine/Source/Editor/ContentBrowser/Public/SAssetView.h
+++ b/Engine/Source/Editor/ContentBrowser/Public/SAssetView.h
@@ -402,17 +402,17 @@ public:
 	 */
 	void ForceShowPluginFolder( bool bEnginePlugin );
 
-	/** Sets a non-persistent override value for the Show Engine Content setting. The user can still toggle the setting manually. */
-	void OverrideShowEngineContent( bool bShow );
+	/** Enables the Show Engine Content setting for the active Content Browser. The user can still toggle the setting manually. */
+	void OverrideShowEngineContent();
 
-	/** Sets a non-persistent override value for the Show Engine Content setting. The user can still toggle the setting manually. */
-	void OverrideShowDeveloperContent( bool bShow );
+	/** Enables the Show Developer Content setting for the active Content Browser. The user can still toggle the setting manually. */
+	void OverrideShowDeveloperContent();
 
-	/** Sets a non-persistent override value for the Show Engine Content setting. The user can still toggle the setting manually. */
-	void OverrideShowPluginContent( bool bShow );
+	/** Enables the Show Plugin Content setting for the active Content Browser. The user can still toggle the setting manually. */
+	void OverrideShowPluginContent();
 
-	/** Sets a non-persistent override value for the Show Engine Content setting. The user can still toggle the setting manually. */
-	void OverrideShowLocalizedContent( bool bShow );
+	/** Enables the Show Localized Content setting for the active Content Browser. The user can still toggle the setting manually. */
+	void OverrideShowLocalizedContent();
 
 	/** @return true when we are including class names in search criteria */
 	bool IsIncludingClassNames() const;
@@ -1097,18 +1097,6 @@ private:
 	/** Whether or not to notify about newly selected items on on the next asset sync */
 	bool bShouldNotifyNextAssetSync : 1;
 
-	/** If set, this will be used instead of the global setting. Mostly used post-sync. */
-	TOptional<bool> bOverrideShowEngineContent;
-
-	/** If set, this will be used instead of the global setting. Mostly used post-sync. */
-	TOptional<bool> bOverrideShowDeveloperContent;
-
-	/** If set, this will be used instead of the global setting. Mostly used post-sync. */
-	TOptional<bool> bOverrideShowPluginContent;
-
-	/** If set, this will be used instead of the global setting. Mostly used post-sync. */
-	TOptional<bool> bOverrideShowLocalizedContent;
-
 	/** The current selection mode used by the asset view */
 	ESelectionMode::Type SelectionMode;
 
