Stopping an ensure when selecting something in the content browser that no longer exists

#rnx

diff --git a/Engine/Source/Editor/ContentBrowser/Private/AssetContextMenu.cpp b/Engine/Source/Editor/ContentBrowser/Private/AssetContextMenu.cpp
--- a/Engine/Source/Editor/ContentBrowser/Private/AssetContextMenu.cpp
+++ b/Engine/Source/Editor/ContentBrowser/Private/AssetContextMenu.cpp
@@ -374,10 +374,15 @@ void FAssetContextMenu::RegisterContextMenu(const FName MenuName)
 			{
 				UContentBrowserAssetContextMenuContext* Context = InSection.FindContext<UContentBrowserAssetContextMenuContext>();
 				PRAGMA_DISABLE_DEPRECATION_WARNINGS
-				if (Context && Context->CommonAssetTypeActions.IsValid() && Context->CommonAssetTypeActions.Pin()->ShouldCallGetActions())
-				{
-					Context->CommonAssetTypeActions.Pin()->GetActions(Context->LoadSelectedObjectsIfNeeded(), InSection);
-				}
+					if (Context && Context->CommonAssetTypeActions.IsValid() && Context->CommonAssetTypeActions.Pin()->ShouldCallGetActions())
+					{
+						TArray<UObject*> SelectedObjects = Context->LoadSelectedObjectsIfNeeded();						
+						//  It's possible for an unloaded object to be selected if the content browser is out of date, in that case it is unnecessary to call `GetActions`
+						if (SelectedObjects.Num() > 0)
+						{
+							Context->CommonAssetTypeActions.Pin()->GetActions(SelectedObjects, InSection);
+						}
+					}
 				PRAGMA_ENABLE_DEPRECATION_WARNINGS
 			}));
 
@@ -385,10 +390,14 @@ void FAssetContextMenu::RegisterContextMenu(const FName MenuName)
 			{
 				UContentBrowserAssetContextMenuContext* Context = InMenu->FindContext<UContentBrowserAssetContextMenuContext>();
 				PRAGMA_DISABLE_DEPRECATION_WARNINGS
-				if (Context && Context->CommonAssetTypeActions.IsValid() && Context->CommonAssetTypeActions.Pin()->ShouldCallGetActions())
-				{
-					Context->CommonAssetTypeActions.Pin()->GetActions(Context->LoadSelectedObjectsIfNeeded(), MenuBuilder);
-				}
+					if (Context && Context->CommonAssetTypeActions.IsValid() && Context->CommonAssetTypeActions.Pin()->ShouldCallGetActions())
+					{
+						TArray<UObject*> SelectedObjects = Context->LoadSelectedObjectsIfNeeded();
+						if (SelectedObjects.Num() > 0)
+						{
+							Context->CommonAssetTypeActions.Pin()->GetActions(SelectedObjects, MenuBuilder);
+						}
+					}
 				PRAGMA_ENABLE_DEPRECATION_WARNINGS
 			}));
 		}		
