Fall back to EditorConfig version of content browser settings where it exists, removing some settings from per-instance settings

diff --git a/Engine/Source/Editor/ContentBrowser/Private/ContentBrowserConfig.h b/Engine/Source/Editor/ContentBrowser/Private/ContentBrowserConfig.h
--- a/Engine/Source/Editor/ContentBrowser/Private/ContentBrowserConfig.h
+++ b/Engine/Source/Editor/ContentBrowser/Private/ContentBrowserConfig.h
@@ -71,9 +71,6 @@ struct FContentBrowserInstanceConfig
 	UPROPERTY()
 	bool bShowEmptyFolders = true;
 
-	UPROPERTY()
-	bool bShowAllFolder = true;
-
 	UPROPERTY()
 	bool bShowEngineContent = false;
 
@@ -89,12 +86,6 @@ struct FContentBrowserInstanceConfig
 	UPROPERTY()
 	bool bShowCppFolders = false;
 
-	UPROPERTY()
-	bool bOrganizeFolders = false;
-
-	UPROPERTY()
-	bool bRealTimeThumbnails = false;
-
 	UPROPERTY()
 	bool bSearchClasses = false;
 
diff --git a/Engine/Source/Editor/ContentBrowser/Private/NewAssetOrClassContextMenu.cpp b/Engine/Source/Editor/ContentBrowser/Private/NewAssetOrClassContextMenu.cpp
--- a/Engine/Source/Editor/ContentBrowser/Private/NewAssetOrClassContextMenu.cpp
+++ b/Engine/Source/Editor/ContentBrowser/Private/NewAssetOrClassContextMenu.cpp
@@ -2,6 +2,7 @@
 
 #include "NewAssetOrClassContextMenu.h"
 
+#include "ContentBrowserConfig.h"
 #include "ContentBrowserDataMenuContexts.h"
 #include "ContentBrowserDataSubsystem.h"
 #include "Framework/Commands/UIAction.h"
@@ -41,10 +42,10 @@ void FNewAssetOrClassContextMenu::MakeContextMenu(
 
 	// Get Content
 	FToolMenuSection& GetContentSection = Menu->AddSection("ContentBrowserGetContent", LOCTEXT("GetContentMenuHeading", "Get Content"));
+	UContentBrowserDataMenuContext_AddNewMenu* AddNewMenuContext = Menu->FindContext<UContentBrowserDataMenuContext_AddNewMenu>();
 
 	if ( InOnGetContentRequested.IsBound() )
 	{
-		UContentBrowserDataMenuContext_AddNewMenu* AddNewMenuContext = Menu->FindContext<UContentBrowserDataMenuContext_AddNewMenu>();
 		if (AddNewMenuContext && AddNewMenuContext->bCanBeModified && AddNewMenuContext->bContainsValidPackagePath)
 		{
 			GetContentSection.AddMenuEntry(
@@ -63,8 +64,15 @@ void FNewAssetOrClassContextMenu::MakeContextMenu(
 		return;
 	}
 
+	bool bDisplayFolders = GetDefault<UContentBrowserSettings>()->DisplayFolders;
+	// check to see if we have an instance config that overrides the default in UContentBrowserSettings
+	if (AddNewMenuContext && AddNewMenuContext->OwningInstanceConfig)
+	{
+		bDisplayFolders = AddNewMenuContext->OwningInstanceConfig->bShowFolders;
+	}
+
 	// New Folder
-	if(InOnNewFolderRequested.IsBound() && GetDefault<UContentBrowserSettings>()->DisplayFolders)
+	if(InOnNewFolderRequested.IsBound() && bDisplayFolders)
 	{
 		{
 			FToolMenuSection& Section = Menu->AddSection("ContentBrowserNewFolder", LOCTEXT("FolderMenuHeading", "Folder") );
diff --git a/Engine/Source/Editor/ContentBrowser/Private/SAssetView.cpp b/Engine/Source/Editor/ContentBrowser/Private/SAssetView.cpp
--- a/Engine/Source/Editor/ContentBrowser/Private/SAssetView.cpp
+++ b/Engine/Source/Editor/ContentBrowser/Private/SAssetView.cpp
@@ -2172,21 +2172,10 @@ void SAssetView::ToggleShowAllFolder()
 	const bool bNewValue = !IsShowingAllFolder();
 	GetMutableDefault<UContentBrowserSettings>()->bShowAllFolder = bNewValue;
 	GetMutableDefault<UContentBrowserSettings>()->PostEditChange();
-
-	if (FContentBrowserInstanceConfig* Config = GetContentBrowserConfig())
-	{
-		Config->bShowAllFolder = bNewValue;
-		UContentBrowserConfig::Get()->SaveEditorConfig();
-	}
 }
 
 bool SAssetView::IsShowingAllFolder() const
 {
-	if (FContentBrowserInstanceConfig* Config = GetContentBrowserConfig())
-	{
-		return Config->bShowAllFolder;
-	}
-
 	return GetDefault<UContentBrowserSettings>()->bShowAllFolder;
 }
 
@@ -2195,21 +2184,10 @@ void SAssetView::ToggleOrganizeFolders()
 	const bool bNewValue = !IsOrganizingFolders();
 	GetMutableDefault<UContentBrowserSettings>()->bOrganizeFolders = bNewValue;
 	GetMutableDefault<UContentBrowserSettings>()->PostEditChange();
-
-	if (FContentBrowserInstanceConfig* Config = GetContentBrowserConfig())
-	{
-		Config->bOrganizeFolders = bNewValue;
-		UContentBrowserConfig::Get()->SaveEditorConfig();
-	}
 }
 
 bool SAssetView::IsOrganizingFolders() const
 {
-	if (FContentBrowserInstanceConfig* Config = GetContentBrowserConfig())
-	{
-		return Config->bOrganizeFolders;
-	}
-
 	return GetDefault<UContentBrowserSettings>()->bOrganizeFolders;
 }
 
@@ -3021,12 +2999,6 @@ void SAssetView::ToggleRealTimeThumbnails()
 
 	GetMutableDefault<UContentBrowserSettings>()->RealTimeThumbnails = bNewState;
 	GetMutableDefault<UContentBrowserSettings>()->PostEditChange();
-
-	if (FContentBrowserInstanceConfig* Config = GetContentBrowserConfig())
-	{
-		Config->bRealTimeThumbnails = bNewState;
-		UContentBrowserConfig::Get()->SaveEditorConfig();
-	}
 }
 
 bool SAssetView::CanShowRealTimeThumbnails() const
@@ -3040,11 +3012,6 @@ bool SAssetView::IsShowingRealTimeThumbnails() const
 	{
 		return false;
 	}
-	
-	if (const FContentBrowserInstanceConfig* Config = GetContentBrowserConfig())
-	{
-		return Config->bRealTimeThumbnails;
-	}
 
 	return GetDefault<UContentBrowserSettings>()->RealTimeThumbnails;
 }
diff --git a/Engine/Source/Editor/ContentBrowser/Private/SContentBrowser.cpp b/Engine/Source/Editor/ContentBrowser/Private/SContentBrowser.cpp
--- a/Engine/Source/Editor/ContentBrowser/Private/SContentBrowser.cpp
+++ b/Engine/Source/Editor/ContentBrowser/Private/SContentBrowser.cpp
@@ -1402,10 +1402,20 @@ void SContentBrowser::PrepareToSyncItems(TArrayView<const FContentBrowserItem> I
 	bool bRepopulate = false;
 
 	// Check to see if any of the assets require certain folders to be visible
-	bool bDisplayDev = GetDefault<UContentBrowserSettings>()->GetDisplayDevelopersFolder();
-	bool bDisplayEngine = GetDefault<UContentBrowserSettings>()->GetDisplayEngineFolder();
-	bool bDisplayPlugins = GetDefault<UContentBrowserSettings>()->GetDisplayPluginFolders();
-	bool bDisplayLocalized = GetDefault<UContentBrowserSettings>()->GetDisplayL10NFolder();
+	const UContentBrowserSettings* ContentBrowserSettings = GetDefault<UContentBrowserSettings>();
+	bool bDisplayDev = ContentBrowserSettings->GetDisplayDevelopersFolder();
+	bool bDisplayEngine = ContentBrowserSettings->GetDisplayEngineFolder();
+	bool bDisplayPlugins = ContentBrowserSettings->GetDisplayPluginFolders();
+	bool bDisplayLocalized = ContentBrowserSettings->GetDisplayL10NFolder();
+
+	// check to see if we have an instance config that overrides the default in UContentBrowserSettings
+	if (const FContentBrowserInstanceConfig* EditorConfig = GetConstInstanceConfig())
+	{
+		bDisplayDev = EditorConfig->bShowDeveloperContent;
+		bDisplayEngine = EditorConfig->bShowEngineContent;
+		bDisplayPlugins = EditorConfig->bShowPluginContent;
+		bDisplayLocalized = EditorConfig->bShowLocalizedContent;
+	}
 
 	// Keep track of any of the settings changing so we can let the user know
 	bool bDisplayDevChanged = false;
@@ -2475,7 +2485,12 @@ TSharedRef<SWidget> SContentBrowser::OnGetCrumbDelimiterContent(const FString& C
 	else if( SourcesData.HasVirtualPaths() )
 	{
 		const UContentBrowserSettings* ContentBrowserSettings = GetDefault<UContentBrowserSettings>();
-		const bool bDisplayEmpty = ContentBrowserSettings->DisplayEmptyFolders;
+		bool bDisplayEmpty = ContentBrowserSettings->DisplayEmptyFolders;
+		// check to see if we have an instance config that overrides the default in UContentBrowserSettings
+		if (const FContentBrowserInstanceConfig* EditorConfig = GetConstInstanceConfig())
+		{
+			bDisplayEmpty = EditorConfig->bShowEmptyFolders;
+		}
 
 		UContentBrowserDataSubsystem* ContentBrowserData = IContentBrowserDataModule::Get().GetSubsystem();
 
@@ -2709,7 +2724,7 @@ void SContentBrowser::AppendNewMenuContextObjects(const EContentBrowserDataMenuC
 		DataContextObject->OnBeginItemCreation = UContentBrowserDataMenuContext_AddNewMenu::FOnBeginItemCreation::CreateSP(this, &SContentBrowser::NewFileItemRequested);
 		DataContextObject->bCanBeModified = bCanBeModified;
 		DataContextObject->bContainsValidPackagePath = bContainsValidPackagePath;
-
+		DataContextObject->OwningInstanceConfig = GetConstInstanceConfig();
 		InOutMenuContext.AddObject(DataContextObject);
 	}
 }
@@ -4343,12 +4358,9 @@ FContentBrowserInstanceConfig* SContentBrowser::CreateEditorConfigIfRequired()
 	InstanceConfig->bShowLocalizedContent = Settings->GetDisplayL10NFolder();
 	InstanceConfig->bShowPluginContent = Settings->GetDisplayPluginFolders();
 	InstanceConfig->bShowFolders = Settings->DisplayFolders;
-	InstanceConfig->bShowAllFolder = Settings->bShowAllFolder;
 	InstanceConfig->bShowEmptyFolders = Settings->DisplayEmptyFolders;
 	InstanceConfig->bShowCppFolders = Settings->GetDisplayCppFolders();
-	InstanceConfig->bOrganizeFolders = Settings->bOrganizeFolders;
 	InstanceConfig->bFavoritesExpanded = Settings->GetDisplayFavorites();
-	InstanceConfig->bRealTimeThumbnails = Settings->RealTimeThumbnails;
 	InstanceConfig->bSearchAssetPaths = Settings->GetIncludeAssetPaths();
 	InstanceConfig->bSearchClasses = Settings->GetIncludeClassNames();
 	InstanceConfig->bSearchCollections = Settings->GetIncludeCollectionNames();
diff --git a/Engine/Source/Editor/ContentBrowser/Private/SPathView.cpp b/Engine/Source/Editor/ContentBrowser/Private/SPathView.cpp
--- a/Engine/Source/Editor/ContentBrowser/Private/SPathView.cpp
+++ b/Engine/Source/Editor/ContentBrowser/Private/SPathView.cpp
@@ -419,22 +419,22 @@ void SPathView::SetPluginPathFilterActive(const TSharedRef<FContentBrowserPlugin
 		PluginPathFilters->Remove(Filter);
 	}
 
-	if (FPathViewConfig* EditorConfig = GetEditorConfig())
+	if (FPathViewConfig* PathViewConfig = GetPathViewConfig())
 	{
 		if (bActive)
 		{
-			EditorConfig->PluginFilters.Add(Filter->GetName());
+			PathViewConfig->PluginFilters.Add(Filter->GetName());
 		}
 		else
 		{
-			EditorConfig->PluginFilters.Remove(Filter->GetName());
+			PathViewConfig->PluginFilters.Remove(Filter->GetName());
 		}
 		
 		UContentBrowserConfig::Get()->SaveEditorConfig();
 	}
 }
 
-FPathViewConfig* SPathView::GetEditorConfig() const
+FPathViewConfig* SPathView::GetPathViewConfig() const
 {
 	if (OwningContentBrowserName.IsNone())
 	{
@@ -450,6 +450,22 @@ FPathViewConfig* SPathView::GetEditorConfig() const
 	return &Config->PathView;
 }
 
+FContentBrowserInstanceConfig* SPathView::GetContentBrowserConfig() const
+{
+	if (OwningContentBrowserName.IsNone())
+	{
+		return nullptr;
+	}
+
+	FContentBrowserInstanceConfig* Config = UContentBrowserConfig::Get()->Instances.Find(OwningContentBrowserName);
+	if (Config == nullptr)
+	{
+		return nullptr;
+	}
+
+	return Config;
+}
+
 void SPathView::SetSelectedPaths(const TArray<FName>& Paths)
 {
 	TArray<FString> PathStrings;
@@ -556,9 +572,9 @@ void SPathView::SetSelectedPaths(const TArray<FString>& Paths)
 		}
 	}
 
-	if (FPathViewConfig* EditorConfig = GetEditorConfig())
+	if (FPathViewConfig* PathViewConfig = GetPathViewConfig())
 	{
-		EditorConfig->SelectedPaths = LastSelectedPaths.Array();
+		PathViewConfig->SelectedPaths = LastSelectedPaths.Array();
 
 		UContentBrowserConfig::Get()->SaveEditorConfig();
 	}
@@ -848,6 +864,12 @@ void SPathView::RenameFolderItem(const FContentBrowserItem& InItem)
 FContentBrowserDataCompiledFilter SPathView::CreateCompiledFolderFilter() const
 {
 	const UContentBrowserSettings* ContentBrowserSettings = GetDefault<UContentBrowserSettings>();
+	bool bDisplayPluginFolders = ContentBrowserSettings->GetDisplayPluginFolders();
+	// check to see if we have an instance config that overrides the default in UContentBrowserSettings
+	if (FContentBrowserInstanceConfig* EditorConfig = GetContentBrowserConfig())
+	{
+		bDisplayPluginFolders = EditorConfig->bShowPluginContent;
+	}
 
 	FContentBrowserDataFilter DataFilter;
 	DataFilter.bRecursivePaths = true;
@@ -866,7 +888,7 @@ FContentBrowserDataCompiledFilter SPathView::CreateCompiledFolderFilter() const
 		CombinedFolderPermissionList->Append(*CustomFolderPermissionList);
 	}
 
-	if (PluginPathFilters.IsValid() && PluginPathFilters->Num() > 0 && ContentBrowserSettings->GetDisplayPluginFolders())
+	if (PluginPathFilters.IsValid() && PluginPathFilters->Num() > 0 && bDisplayPluginFolders)
 	{
 		TArray<TSharedRef<IPlugin>> Plugins = IPluginManager::Get().GetEnabledPluginsWithContent();
 		for (const TSharedRef<IPlugin>& Plugin : Plugins)
@@ -899,8 +921,15 @@ FContentBrowserDataCompiledFilter SPathView::CreateCompiledFolderFilter() const
 EContentBrowserItemCategoryFilter SPathView::GetContentBrowserItemCategoryFilter() const
 {
 	const UContentBrowserSettings* ContentBrowserSettings = GetDefault<UContentBrowserSettings>();
+	bool bDisplayCppFolders = ContentBrowserSettings->GetDisplayCppFolders();
+	// check to see if we have an instance config that overrides the default in UContentBrowserSettings
+	if (FContentBrowserInstanceConfig* EditorConfig = GetContentBrowserConfig())
+	{
+		bDisplayCppFolders = EditorConfig->bShowCppFolders;
+	}
+
 	EContentBrowserItemCategoryFilter ItemCategoryFilter = InitialCategoryFilter;
-	if (bAllowClassesFolder && ContentBrowserSettings->GetDisplayCppFolders())
+	if (bAllowClassesFolder && bDisplayCppFolders)
 	{
 		ItemCategoryFilter |= EContentBrowserItemCategoryFilter::IncludeClasses;
 	}
@@ -916,11 +945,25 @@ EContentBrowserItemCategoryFilter SPathView::GetContentBrowserItemCategoryFilter
 EContentBrowserItemAttributeFilter SPathView::GetContentBrowserItemAttributeFilter() const
 {
 	const UContentBrowserSettings* ContentBrowserSettings = GetDefault<UContentBrowserSettings>();
+	bool bDisplayEngineContent = ContentBrowserSettings->GetDisplayEngineFolder();
+	bool bDisplayPluginContent = ContentBrowserSettings->GetDisplayPluginFolders();
+	bool bDisplayDevelopersContent = ContentBrowserSettings->GetDisplayDevelopersFolder();
+	bool bDisplayL10NContent = ContentBrowserSettings->GetDisplayL10NFolder();
+	
+	// check to see if we have an instance config that overrides the defaults in UContentBrowserSettings
+	if (FContentBrowserInstanceConfig* EditorConfig = GetContentBrowserConfig())
+	{
+		bDisplayEngineContent = EditorConfig->bShowEngineContent;
+		bDisplayPluginContent = EditorConfig->bShowPluginContent;
+		bDisplayDevelopersContent = EditorConfig->bShowDeveloperContent;
+		bDisplayL10NContent = EditorConfig->bShowLocalizedContent;
+	}
+	
 	return EContentBrowserItemAttributeFilter::IncludeProject
-			| (ContentBrowserSettings->GetDisplayEngineFolder() ? EContentBrowserItemAttributeFilter::IncludeEngine : EContentBrowserItemAttributeFilter::IncludeNone)
-			| (ContentBrowserSettings->GetDisplayPluginFolders() ? EContentBrowserItemAttributeFilter::IncludePlugins : EContentBrowserItemAttributeFilter::IncludeNone)
-			| (ContentBrowserSettings->GetDisplayDevelopersFolder() ? EContentBrowserItemAttributeFilter::IncludeDeveloper : EContentBrowserItemAttributeFilter::IncludeNone)
-			| (ContentBrowserSettings->GetDisplayL10NFolder() ? EContentBrowserItemAttributeFilter::IncludeLocalized : EContentBrowserItemAttributeFilter::IncludeNone);
+			| (bDisplayEngineContent ? EContentBrowserItemAttributeFilter::IncludeEngine : EContentBrowserItemAttributeFilter::IncludeNone)
+			| (bDisplayPluginContent ? EContentBrowserItemAttributeFilter::IncludePlugins : EContentBrowserItemAttributeFilter::IncludeNone)
+			| (bDisplayDevelopersContent ? EContentBrowserItemAttributeFilter::IncludeDeveloper : EContentBrowserItemAttributeFilter::IncludeNone)
+			| (bDisplayL10NContent ? EContentBrowserItemAttributeFilter::IncludeLocalized : EContentBrowserItemAttributeFilter::IncludeNone);
 }
 
 bool SPathView::InternalPathPassesBlockLists(const FStringView InInternalPath, const int32 InAlreadyCheckedDepth) const
@@ -947,7 +990,15 @@ bool SPathView::InternalPathPassesBlockLists(const FStringView InInternalPath, c
 	if (InAlreadyCheckedDepth < 1 && PluginPathFilters.IsValid() && PluginPathFilters->Num() > 0)
 	{
 		const UContentBrowserSettings* ContentBrowserSettings = GetDefault<UContentBrowserSettings>();
-		if (ContentBrowserSettings->GetDisplayPluginFolders())
+		bool bDisplayPluginFolders = ContentBrowserSettings->GetDisplayPluginFolders();
+
+		// check to see if we have an instance config that overrides the default in UContentBrowserSettings
+		if (FContentBrowserInstanceConfig* EditorConfig = GetContentBrowserConfig())
+		{
+			bDisplayPluginFolders = EditorConfig->bShowPluginContent;
+		}
+
+		if (bDisplayPluginFolders)
 		{
 			const FStringView FirstFolderName = FPathViews::GetMountPointNameFromPath(InInternalPath);
 			if (TSharedPtr<IPlugin> Plugin = IPluginManager::Get().FindPlugin(FirstFolderName))
@@ -1150,9 +1201,9 @@ void SPathView::LoadSettings(const FString& IniFilename, const FString& IniSecti
 {
 	// Selected Paths
 	TArray<FName> NewSelectedPaths;
-	if (FPathViewConfig* EditorConfig = GetEditorConfig())
+	if (FPathViewConfig* PathViewConfig = GetPathViewConfig())
 	{
-		NewSelectedPaths = EditorConfig->SelectedPaths;
+		NewSelectedPaths = PathViewConfig->SelectedPaths;
 	}
 	else 
 	{
@@ -1227,9 +1278,9 @@ void SPathView::LoadSettings(const FString& IniFilename, const FString& IniSecti
 	if (PluginPathFilters.IsValid())
 	{
 		TArray<FString> NewSelectedFilters;
-		if (FPathViewConfig* EditorConfig = GetEditorConfig())
+		if (FPathViewConfig* PathViewConfig = GetPathViewConfig())
 		{
-			NewSelectedFilters = EditorConfig->PluginFilters;
+			NewSelectedFilters = PathViewConfig->PluginFilters;
 		}
 		else
 		{
@@ -1403,9 +1454,9 @@ void SPathView::TreeSelectionChanged( TSharedPtr< FTreeItem > TreeItem, ESelectI
 			}
 		}
 
-		if (FPathViewConfig* EditorConfig = GetEditorConfig())
+		if (FPathViewConfig* PathViewConfig = GetPathViewConfig())
 		{
-			EditorConfig->SelectedPaths = LastSelectedPaths.Array();
+			PathViewConfig->SelectedPaths = LastSelectedPaths.Array();
 
 			UContentBrowserConfig::Get()->SaveEditorConfig();
 		}
@@ -1487,7 +1538,12 @@ void SPathView::Populate(const bool bIsRefreshingFilter)
 	// Populate the view
 	{
 		const UContentBrowserSettings* ContentBrowserSettings = GetDefault<UContentBrowserSettings>();
-		const bool bDisplayEmpty = ContentBrowserSettings->DisplayEmptyFolders;
+		bool bDisplayEmpty = ContentBrowserSettings->DisplayEmptyFolders;
+		// check to see if we have an instance config that overrides the default in UContentBrowserSettings
+		if (FContentBrowserInstanceConfig* EditorConfig = GetContentBrowserConfig())
+		{
+			bDisplayEmpty = EditorConfig->bShowEmptyFolders;
+		}
 
 		UContentBrowserDataSubsystem* ContentBrowserData = IContentBrowserDataModule::Get().GetSubsystem();
 		const FContentBrowserDataCompiledFilter CompiledDataFilter = CreateCompiledFolderFilter();
@@ -1927,7 +1983,12 @@ void SPathView::HandleItemDataUpdated(TArrayView<const FContentBrowserItemDataUp
 	const double HandleItemDataUpdatedStartTime = FPlatformTime::Seconds();
 
 	const UContentBrowserSettings* ContentBrowserSettings = GetDefault<UContentBrowserSettings>();
-	const bool bDisplayEmpty = ContentBrowserSettings->DisplayEmptyFolders;
+	bool bDisplayEmpty = ContentBrowserSettings->DisplayEmptyFolders;
+	// check to see if we have an instance config that overrides the default in UContentBrowserSettings
+	if (FContentBrowserInstanceConfig* EditorConfig = GetContentBrowserConfig())
+	{
+		bDisplayEmpty = EditorConfig->bShowEmptyFolders;
+	}
 
 	UContentBrowserDataSubsystem* ContentBrowserData = IContentBrowserDataModule::Get().GetSubsystem();
 
@@ -2083,10 +2144,20 @@ void SPathView::HandleSettingChanged(FName PropertyName)
 		}
 
 		// If the dev or engine folder has become visible and we're inside it...
-		const bool bDisplayDev = GetDefault<UContentBrowserSettings>()->GetDisplayDevelopersFolder();
-		const bool bDisplayEngine = GetDefault<UContentBrowserSettings>()->GetDisplayEngineFolder();
-		const bool bDisplayPlugins = GetDefault<UContentBrowserSettings>()->GetDisplayPluginFolders();
-		const bool bDisplayL10N = GetDefault<UContentBrowserSettings>()->GetDisplayL10NFolder();
+		const UContentBrowserSettings* ContentBrowserSettings = GetDefault<UContentBrowserSettings>();
+		bool bDisplayDev = ContentBrowserSettings->GetDisplayDevelopersFolder();
+		bool bDisplayEngine = ContentBrowserSettings->GetDisplayEngineFolder();
+		bool bDisplayPlugins = ContentBrowserSettings->GetDisplayPluginFolders();
+		bool bDisplayL10N = ContentBrowserSettings->GetDisplayL10NFolder();
+		// check to see if we have an instance config that overrides the default in UContentBrowserSettings
+		if (FContentBrowserInstanceConfig* EditorConfig = GetContentBrowserConfig())
+		{
+			bDisplayDev = EditorConfig->bShowDeveloperContent;
+			bDisplayEngine = EditorConfig->bShowEngineContent;
+			bDisplayPlugins = EditorConfig->bShowPluginContent;
+			bDisplayL10N = EditorConfig->bShowLocalizedContent;
+		}
+
 		if (bDisplayDev || bDisplayEngine || bDisplayPlugins || bDisplayL10N)
 		{
 			const TArray<FContentBrowserItem> NewSelectedItems = GetSelectedFolderItems();
diff --git a/Engine/Source/Editor/ContentBrowser/Private/SPathView.h b/Engine/Source/Editor/ContentBrowser/Private/SPathView.h
--- a/Engine/Source/Editor/ContentBrowser/Private/SPathView.h
+++ b/Engine/Source/Editor/ContentBrowser/Private/SPathView.h
@@ -53,6 +53,7 @@ struct FGeometry;
 struct FHistoryData;
 struct FPathViewConfig;
 struct FPointerEvent;
+struct FContentBrowserInstanceConfig;
 
 typedef TTextFilter< const FString& > FolderTextFilter;
 
@@ -304,7 +305,10 @@ protected:
 	void ClearTreeItems();
 
 	/** Get this path view's editor config if OwningContentBrowserName is set. */
-	FPathViewConfig* GetEditorConfig() const;
+	FPathViewConfig* GetPathViewConfig() const;
+
+	/** Get this path view's content browser instance config if OwningContentBrowserName is set. */
+	FContentBrowserInstanceConfig* GetContentBrowserConfig() const;
 
 private:
 	/** Selects the given path only if it exists. Returns true if selected. */
diff --git a/Engine/Source/Editor/ContentBrowserData/Public/ContentBrowserDataMenuContexts.h b/Engine/Source/Editor/ContentBrowserData/Public/ContentBrowserDataMenuContexts.h
--- a/Engine/Source/Editor/ContentBrowserData/Public/ContentBrowserDataMenuContexts.h
+++ b/Engine/Source/Editor/ContentBrowserData/Public/ContentBrowserDataMenuContexts.h
@@ -8,6 +8,7 @@
 #include "ContentBrowserDataMenuContexts.generated.h"
 
 class SWidget;
+struct FContentBrowserInstanceConfig;
 
 UENUM()
 enum class EContentBrowserDataMenuContext_AddNewMenuDomain : uint8
@@ -39,6 +40,8 @@ public:
 	EContentBrowserDataMenuContext_AddNewMenuDomain OwnerDomain = EContentBrowserDataMenuContext_AddNewMenuDomain::Toolbar;
 
 	FOnBeginItemCreation OnBeginItemCreation;
+
+	const FContentBrowserInstanceConfig* OwningInstanceConfig;
 };
 
 UCLASS()
