Editor - Fix some favorites being deleted in content browser if plugin not mounted when settings loaded

diff --git a/Engine/Source/Editor/ContentBrowser/Private/ContentBrowserUtils.cpp b/Engine/Source/Editor/ContentBrowser/Private/ContentBrowserUtils.cpp
--- a/Engine/Source/Editor/ContentBrowser/Private/ContentBrowserUtils.cpp
+++ b/Engine/Source/Editor/ContentBrowser/Private/ContentBrowserUtils.cpp
@@ -701,12 +701,32 @@ bool ContentBrowserUtils::CanRenameFromPathView(TWeakPtr<SPathView> PathView, FT
 
 bool ContentBrowserUtils::IsFavoriteFolder(const FString& FolderPath)
 {
-	return FContentBrowserSingleton::Get().FavoriteFolderPaths.Contains(ConvertVirtualPathToInvariantPathString(FolderPath));
+	return IsFavoriteFolder(FContentBrowserItemPath(FolderPath, EContentBrowserPathType::Virtual));
 }
 
-void ContentBrowserUtils::AddFavoriteFolder(const FString& FolderPath)
+bool ContentBrowserUtils::IsFavoriteFolder(const FContentBrowserItemPath& FolderPath)
 {
-	const FString InvariantFolder = ConvertVirtualPathToInvariantPathString(FolderPath);
+	if (ensure(FolderPath.HasInternalPath()))
+	{
+		return FContentBrowserSingleton::Get().FavoriteFolderPaths.Contains(FolderPath.GetInternalPathString());
+	}
+
+	return false;
+}
+
+void ContentBrowserUtils::AddFavoriteFolder(const FString& FolderPath, bool bFlushConfig /*= true*/)
+{
+	AddFavoriteFolder(FContentBrowserItemPath(FolderPath, EContentBrowserPathType::Virtual));
+}
+
+void ContentBrowserUtils::AddFavoriteFolder(const FContentBrowserItemPath& FolderPath)
+{
+	if (!ensure(FolderPath.HasInternalPath()))
+	{
+		return;
+	}
+
+	const FString InvariantFolder = FolderPath.GetInternalPathString();
 
 	FContentBrowserSingleton::Get().FavoriteFolderPaths.AddUnique(InvariantFolder);
 
@@ -719,15 +739,15 @@ void ContentBrowserUtils::AddFavoriteFolder(const FString& FolderPath)
 
 	FContentBrowserSingleton::Get().BroadcastFavoritesChanged();
 }
-	
-void ContentBrowserUtils::AddFavoriteFolder(const FString& FolderPath, bool bFlushConfig)
-{
-	AddFavoriteFolder(FolderPath);
-}
 
-void ContentBrowserUtils::RemoveFavoriteFolder(const FString& FolderPath)
+void ContentBrowserUtils::RemoveFavoriteFolder(const FContentBrowserItemPath& FolderPath)
 {
-	const FString InvariantFolder = ConvertVirtualPathToInvariantPathString(FolderPath);
+	if (!FolderPath.HasInternalPath())
+	{
+		return;
+	}
+
+	const FString InvariantFolder = FolderPath.GetInternalPathString();
 
 	FContentBrowserSingleton::Get().FavoriteFolderPaths.Remove(InvariantFolder);
 
@@ -743,7 +763,7 @@ void ContentBrowserUtils::RemoveFavoriteFolder(const FString& FolderPath)
 
 void ContentBrowserUtils::RemoveFavoriteFolder(const FString& FolderPath, bool bFlushConfig)
 {
-	RemoveFavoriteFolder(FolderPath);
+	RemoveFavoriteFolder(FContentBrowserItemPath(FolderPath, EContentBrowserPathType::Virtual));
 }
 
 const TArray<FString>& ContentBrowserUtils::GetFavoriteFolders()
diff --git a/Engine/Source/Editor/ContentBrowser/Private/ContentBrowserUtils.h b/Engine/Source/Editor/ContentBrowser/Private/ContentBrowserUtils.h
--- a/Engine/Source/Editor/ContentBrowser/Private/ContentBrowserUtils.h
+++ b/Engine/Source/Editor/ContentBrowser/Private/ContentBrowserUtils.h
@@ -24,6 +24,7 @@ struct FARFilter;
 struct FAssetData;
 struct FContentBrowserDataFilter;
 struct FContentBrowserItem;
+struct FContentBrowserItemPath;
 
 namespace ContentBrowserUtils
 {
@@ -95,19 +96,17 @@ namespace ContentBrowserUtils
 	bool CanRenameFromPathView(TWeakPtr<SPathView> PathView, FText* OutErrorMsg = nullptr);
 
 	/** Returns if this folder has been marked as a favorite folder */
+	UE_DEPRECATED(5.3, "Use function that takes FContentBrowserItemPath instead.")
 	bool IsFavoriteFolder(const FString& FolderPath);
+	bool IsFavoriteFolder(const FContentBrowserItemPath& FolderPath);
 
-	/** Add a favorite folder. This should be a full virtual path. */
-	void AddFavoriteFolder(const FString& FolderPath);
+	UE_DEPRECATED(5.3, "Use function that takes FContentBrowserItemPath instead.")
+	void AddFavoriteFolder(const FString& FolderPath, bool bFlushConfig = true);
+	void AddFavoriteFolder(const FContentBrowserItemPath& FolderPath);
 
-	UE_DEPRECATED(5.2, "The bFlushConfig parameter is ignored, since this function does not persist the favorites to config.")
-	void AddFavoriteFolder(const FString& FolderPath, bool bFlushConfig);
-
-	/** Remove a favorite folder. This should be a full virtual path. */
-	void RemoveFavoriteFolder(const FString& FolderPath);
-
-	UE_DEPRECATED(5.2, "The bFlushConfig parameter is ignored, since this function does not persist the favorites to config.")
-	void RemoveFavoriteFolder(const FString& FolderPath, bool bFlushConfig);
+	UE_DEPRECATED(5.3, "Use function that takes FContentBrowserItemPath instead.")
+	void RemoveFavoriteFolder(const FString& FolderPath, bool bFlushConfig = true);
+	void RemoveFavoriteFolder(const FContentBrowserItemPath& FolderPath);
 
 	const TArray<FString>& GetFavoriteFolders();
 
diff --git a/Engine/Source/Editor/ContentBrowser/Private/PathContextMenu.cpp b/Engine/Source/Editor/ContentBrowser/Private/PathContextMenu.cpp
--- a/Engine/Source/Editor/ContentBrowser/Private/PathContextMenu.cpp
+++ b/Engine/Source/Editor/ContentBrowser/Private/PathContextMenu.cpp
@@ -194,8 +194,9 @@ void FPathContextMenu::MakePathViewContextMenu(UToolMenu* Menu)
 			}			
 
 			FString SelectedFolderPath = SelectedFolders[0].GetVirtualPath().ToString();
+			FContentBrowserItemPath SelectedFolderItemPath(SelectedFolders[0].GetVirtualPath(), EContentBrowserPathType::Virtual);
 			// If this folder is already favorited, show the option to remove from favorites
-			if (ContentBrowserUtils::IsFavoriteFolder(SelectedFolderPath))
+			if (ContentBrowserUtils::IsFavoriteFolder(SelectedFolderItemPath))
 			{
 				// Remove from favorites
 				Section.AddMenuEntry(
diff --git a/Engine/Source/Editor/ContentBrowser/Private/SContentBrowser.cpp b/Engine/Source/Editor/ContentBrowser/Private/SContentBrowser.cpp
--- a/Engine/Source/Editor/ContentBrowser/Private/SContentBrowser.cpp
+++ b/Engine/Source/Editor/ContentBrowser/Private/SContentBrowser.cpp
@@ -666,13 +666,14 @@ void SContentBrowser::ToggleFolderFavorite(const TArray<FString>& FolderPaths)
 	bool bAddedFavorite = false;
 	for (const FString& FolderPath : FolderPaths)
 	{
-		if (ContentBrowserUtils::IsFavoriteFolder(FolderPath))
+		const FContentBrowserItemPath ItemPath(FolderPath, EContentBrowserPathType::Virtual);
+		if (ContentBrowserUtils::IsFavoriteFolder(ItemPath))
 		{
-			ContentBrowserUtils::RemoveFavoriteFolder(FolderPath);
+			ContentBrowserUtils::RemoveFavoriteFolder(ItemPath);
 		}
 		else
 		{
-			ContentBrowserUtils::AddFavoriteFolder(FolderPath);
+			ContentBrowserUtils::AddFavoriteFolder(ItemPath);
 			bAddedFavorite = true;
 		}
 	}
diff --git a/Engine/Source/Editor/ContentBrowser/Private/SPathView.cpp b/Engine/Source/Editor/ContentBrowser/Private/SPathView.cpp
--- a/Engine/Source/Editor/ContentBrowser/Private/SPathView.cpp
+++ b/Engine/Source/Editor/ContentBrowser/Private/SPathView.cpp
@@ -2401,13 +2401,15 @@ void SFavoritePathView::LoadSettings(const FString& IniFilename, const FString&
 		{
 			// If the selected paths is empty, the path was "All assets"
 			// This should handle that case properly
-			for (int32 PathIdx = 0; PathIdx < NewFavoritePaths.Num(); ++PathIdx)
+			for (const FString& InvariantPath : NewFavoritePaths)
 			{
-				const FString& InvariantPath = NewFavoritePaths[PathIdx];
-				FName VirtualPath;
-				IContentBrowserDataModule::Get().GetSubsystem()->ConvertInternalPathToVirtual(FStringView(InvariantPath), VirtualPath);
-				ContentBrowserUtils::AddFavoriteFolder(VirtualPath.ToString());
-				bAddedAtLeastOnePath = true;
+				FStringView InvariantPathView(InvariantPath);
+				InvariantPathView.TrimStartAndEndInline();
+				if (!InvariantPathView.IsEmpty() && InvariantPathView != TEXT("None"))
+				{
+					ContentBrowserUtils::AddFavoriteFolder(FContentBrowserItemPath(InvariantPathView, EContentBrowserPathType::Internal));
+					bAddedAtLeastOnePath = true;
+				}
 			}
 		}
 
@@ -2582,13 +2584,13 @@ void SFavoritePathView::HandleItemDataUpdated(TArrayView<const FContentBrowserIt
 				AddFolderItem(MoveTemp(ItemData));
 			}
 
-			ContentBrowserUtils::RemoveFavoriteFolder(ItemDataUpdate.GetPreviousVirtualPath().ToString());
+			ContentBrowserUtils::RemoveFavoriteFolder(FContentBrowserItemPath(ItemDataUpdate.GetPreviousVirtualPath(), EContentBrowserPathType::Virtual));
 		}
 		break;
 
 		case EContentBrowserItemUpdateType::Removed:
 			RemoveFolderItem(ItemData);
-			ContentBrowserUtils::RemoveFavoriteFolder(ItemData.GetVirtualPath().ToString());
+			ContentBrowserUtils::RemoveFavoriteFolder(FContentBrowserItemPath(ItemData.GetVirtualPath(), EContentBrowserPathType::Virtual));
 			break;
 
 		default:
@@ -2603,22 +2605,23 @@ void SFavoritePathView::HandleItemDataUpdated(TArrayView<const FContentBrowserIt
 bool SFavoritePathView::PathIsFilteredFromViewBySearch(const FString& InPath) const
 {
 	return SPathView::PathIsFilteredFromViewBySearch(InPath)
-		&& ContentBrowserUtils::IsFavoriteFolder(InPath);
+		&& ContentBrowserUtils::IsFavoriteFolder(FContentBrowserItemPath(InPath, EContentBrowserPathType::Virtual));
 }
 
 void SFavoritePathView::FixupFavoritesFromExternalChange(TArrayView<const AssetViewUtils::FMovedContentFolder> MovedFolders)
 {
 	for (const AssetViewUtils::FMovedContentFolder& MovedFolder : MovedFolders)
 	{
-		const bool bWasFavorite = ContentBrowserUtils::IsFavoriteFolder(MovedFolder.Key);
+		FContentBrowserItemPath ItemPath(MovedFolder.Key, EContentBrowserPathType::Virtual);
+		const bool bWasFavorite = ContentBrowserUtils::IsFavoriteFolder(ItemPath);
 		if (bWasFavorite)
 		{
 			// Remove the original path
-			ContentBrowserUtils::RemoveFavoriteFolder(MovedFolder.Key);
+			ContentBrowserUtils::RemoveFavoriteFolder(ItemPath);
 
 			// Add the new path to favorites instead
 			const FString& NewPath = MovedFolder.Value;
-			ContentBrowserUtils::AddFavoriteFolder(NewPath);
+			ContentBrowserUtils::AddFavoriteFolder(FContentBrowserItemPath(NewPath, EContentBrowserPathType::Virtual));
 			TSharedPtr<FTreeItem> Item = FindTreeItem(*NewPath);
 			if (Item.IsValid())
 			{
