CIS - Fix deprecation warning after merge

diff --git a/Engine/Source/Editor/ContentBrowser/Private/SContentBrowser.cpp b/Engine/Source/Editor/ContentBrowser/Private/SContentBrowser.cpp
--- a/Engine/Source/Editor/ContentBrowser/Private/SContentBrowser.cpp
+++ b/Engine/Source/Editor/ContentBrowser/Private/SContentBrowser.cpp
@@ -3421,7 +3421,7 @@ void SContentBrowser::HandleDeleteFavorite(TSharedPtr<SWidget> ParentWidget)
 			{
 				for (const FContentBrowserItem& Folder : SelectedFolders)
 				{
-					ContentBrowserUtils::RemoveFavoriteFolder(Folder.GetVirtualPath().ToString());
+					ContentBrowserUtils::RemoveFavoriteFolder(FContentBrowserItemPath(Folder.GetVirtualPath(), EContentBrowserPathType::Virtual));
 				}
 
 				GConfig->Flush(false, GEditorPerProjectIni);
