Content Browser - Fix favorites to virtual folders such as /All/Plugins not working anymore

diff --git a/Engine/Source/Editor/ContentBrowser/Private/ContentBrowserUtils.cpp b/Engine/Source/Editor/ContentBrowser/Private/ContentBrowserUtils.cpp
--- a/Engine/Source/Editor/ContentBrowser/Private/ContentBrowserUtils.cpp
+++ b/Engine/Source/Editor/ContentBrowser/Private/ContentBrowserUtils.cpp
@@ -699,6 +699,25 @@ bool ContentBrowserUtils::CanRenameFromPathView(TWeakPtr<SPathView> PathView, FT
 	return false;
 }
 
+FName ContentBrowserUtils::GetInvariantPath(const FContentBrowserItemPath& ItemPath)
+{
+	if (!ItemPath.HasInternalPath())
+	{
+		FName InvariantPath;
+		const EContentBrowserPathType AssetPathType = IContentBrowserDataModule::Get().GetSubsystem()->TryConvertVirtualPath(ItemPath.GetVirtualPathName(), InvariantPath);
+		if (AssetPathType == EContentBrowserPathType::Virtual)
+		{
+			return InvariantPath;
+		}
+		else
+		{
+			return NAME_None;
+		}
+	}
+
+	return ItemPath.GetInternalPathName();
+}
+
 bool ContentBrowserUtils::IsFavoriteFolder(const FString& FolderPath)
 {
 	return IsFavoriteFolder(FContentBrowserItemPath(FolderPath, EContentBrowserPathType::Virtual));
@@ -706,9 +725,10 @@ bool ContentBrowserUtils::IsFavoriteFolder(const FString& FolderPath)
 
 bool ContentBrowserUtils::IsFavoriteFolder(const FContentBrowserItemPath& FolderPath)
 {
-	if (ensure(FolderPath.HasInternalPath()))
+	const FName InvariantPath = ContentBrowserUtils::GetInvariantPath(FolderPath);
+	if (!InvariantPath.IsNone())
 	{
-		return FContentBrowserSingleton::Get().FavoriteFolderPaths.Contains(FolderPath.GetInternalPathString());
+		return FContentBrowserSingleton::Get().FavoriteFolderPaths.Contains(InvariantPath.ToString());
 	}
 
 	return false;
@@ -721,12 +741,13 @@ void ContentBrowserUtils::AddFavoriteFolder(const FString& FolderPath, bool bFlu
 
 void ContentBrowserUtils::AddFavoriteFolder(const FContentBrowserItemPath& FolderPath)
 {
-	if (!ensure(FolderPath.HasInternalPath()))
+	const FName InvariantPath = ContentBrowserUtils::GetInvariantPath(FolderPath);
+	if (InvariantPath.IsNone())
 	{
 		return;
 	}
 
-	const FString InvariantFolder = FolderPath.GetInternalPathString();
+	const FString InvariantFolder = InvariantPath.ToString();
 
 	FContentBrowserSingleton::Get().FavoriteFolderPaths.AddUnique(InvariantFolder);
 
@@ -742,14 +763,13 @@ void ContentBrowserUtils::AddFavoriteFolder(const FContentBrowserItemPath& Folde
 
 void ContentBrowserUtils::RemoveFavoriteFolder(const FContentBrowserItemPath& FolderPath)
 {
-	if (!FolderPath.HasInternalPath())
+	const FName InvariantPath = ContentBrowserUtils::GetInvariantPath(FolderPath);
+	if (InvariantPath.IsNone())
 	{
 		return;
 	}
 
-	const FString InvariantFolder = FolderPath.GetInternalPathString();
-
-	FContentBrowserSingleton::Get().FavoriteFolderPaths.Remove(InvariantFolder);
+	FString InvariantFolder = InvariantPath.ToString();
 
 	if (UContentBrowserConfig* EditorConfig = UContentBrowserConfig::Get())
 	{
diff --git a/Engine/Source/Editor/ContentBrowser/Private/ContentBrowserUtils.h b/Engine/Source/Editor/ContentBrowser/Private/ContentBrowserUtils.h
--- a/Engine/Source/Editor/ContentBrowser/Private/ContentBrowserUtils.h
+++ b/Engine/Source/Editor/ContentBrowser/Private/ContentBrowserUtils.h
@@ -95,6 +95,9 @@ namespace ContentBrowserUtils
 	bool CanDeleteFromPathView(TWeakPtr<SPathView> PathView, FText* OutErrorMsg = nullptr);
 	bool CanRenameFromPathView(TWeakPtr<SPathView> PathView, FText* OutErrorMsg = nullptr);
 
+	/** Returns internal path if it has one, otherwise strips /All prefix from virtual path*/
+	FName GetInvariantPath(const FContentBrowserItemPath& ItemPath);
+
 	/** Returns if this folder has been marked as a favorite folder */
 	UE_DEPRECATED(5.3, "Use function that takes FContentBrowserItemPath instead.")
 	bool IsFavoriteFolder(const FString& FolderPath);
