Fixed a crash where a ThumbnailPool could be deleted while it is ticking because the asset view that owned it was destroyed while generating a thumbnail.

The fix was simply to make the asset views use the shared pool. Greatly increased the size to of shared thumbnail pool from the thumbnail manager to acommodate the asset views.

diff --git a/Engine/Source/Editor/ContentBrowser/Private/SAssetView.cpp b/Engine/Source/Editor/ContentBrowser/Private/SAssetView.cpp
--- a/Engine/Source/Editor/ContentBrowser/Private/SAssetView.cpp
+++ b/Engine/Source/Editor/ContentBrowser/Private/SAssetView.cpp
@@ -50,6 +50,7 @@
 #include "SlateOptMacros.h"
 #include "Styling/AppStyle.h"
 #include "Textures/SlateIcon.h"
+#include "ThumbnailRendering/ThumbnailManager.h"
 #include "ToolMenus.h"
 #include "UObject/UnrealType.h"
 #include "Widgets/Images/SImage.h"
@@ -184,9 +185,6 @@ SAssetView::~SAssetView()
 		// Clear the frontend filter changed delegate
 		FrontendFilters->OnChanged().RemoveAll( this );
 	}
-
-	// Release all rendering resources being held onto
-	AssetThumbnailPool.Reset();
 }
 
 
@@ -229,8 +227,8 @@ void SAssetView::Construct( const FArguments& InArgs )
 
 	ThumbnailScaleRangeScalar = (float)DisplaySize.Y / 2160.f;
 
-	// Create a thumbnail pool for rendering thumbnails	
-	AssetThumbnailPool = MakeShared<FAssetThumbnailPool>(InArgs._InitialThumbnailPoolSize);
+	// Use the shared ThumbnailPool for the rendering of thumbnails
+	AssetThumbnailPool = UThumbnailManager::Get().GetSharedThumbnailPool();
 	NumOffscreenThumbnails = 64;
 	ListViewThumbnailResolution = 128;
 	ListViewThumbnailSize = 64;
diff --git a/Engine/Source/Editor/ContentBrowser/Public/SAssetView.h b/Engine/Source/Editor/ContentBrowser/Public/SAssetView.h
--- a/Engine/Source/Editor/ContentBrowser/Public/SAssetView.h
+++ b/Engine/Source/Editor/ContentBrowser/Public/SAssetView.h
@@ -99,7 +99,6 @@ public:
 		, _AllowThumbnailHintLabel(true)
 		, _InitialViewType(EAssetViewType::Tile)
 		, _InitialThumbnailSize(EThumbnailSize::Medium)
-		, _InitialThumbnailPoolSize(1024)
 		, _ShowBottomToolbar(true)
 		, _ShowViewOptions(true)
 		, _AllowThumbnailEditMode(false)
@@ -196,9 +195,6 @@ public:
 		/** Initial thumbnail size */
 		SLATE_ARGUMENT( EThumbnailSize, InitialThumbnailSize )
 
-		/** Initial thumbnail pool size */
-		SLATE_ARGUMENT( uint32, InitialThumbnailPoolSize )
-
 		/** Should the toolbar indicating number of selected assets, mode switch buttons, etc... be shown? */
 		SLATE_ARGUMENT( bool, ShowBottomToolbar )
 
diff --git a/Engine/Source/Editor/UMGEditor/Private/Library/SLibraryViewModel.cpp b/Engine/Source/Editor/UMGEditor/Private/Library/SLibraryViewModel.cpp
--- a/Engine/Source/Editor/UMGEditor/Private/Library/SLibraryViewModel.cpp
+++ b/Engine/Source/Editor/UMGEditor/Private/Library/SLibraryViewModel.cpp
@@ -166,14 +166,9 @@ void FWidgetTemplateListViewModel::ConstructListView(TArray<TSharedPtr<FWidgetTe
 
 	if (!AssetViewPtr)
 	{
-		// Reserve just enough thumbnails, plus a few extra if this category has user widgets
-		uint32 ThumbnailPoolSize = NumAssets + NumClasses;
-		ThumbnailPoolSize += NumAssets > 0 ? 32 : 0;
-
 		AssetViewPtr = SNew(SAssetView)
 			.InitialCategoryFilter(EContentBrowserItemCategoryFilter::IncludeAssets | EContentBrowserItemCategoryFilter::IncludeClasses)
 			.InitialSourcesData(*LibrarySourceData)
-			.InitialThumbnailPoolSize(ThumbnailPoolSize)
 			.InitialThumbnailSize(EThumbnailSize::Small)
 			.FrontendFilters(TemplatesFilter)
 			.ForceShowEngineContent(true)
diff --git a/Engine/Source/Editor/UnrealEd/Classes/ThumbnailRendering/ThumbnailManager.h b/Engine/Source/Editor/UnrealEd/Classes/ThumbnailRendering/ThumbnailManager.h
--- a/Engine/Source/Editor/UnrealEd/Classes/ThumbnailRendering/ThumbnailManager.h
+++ b/Engine/Source/Editor/UnrealEd/Classes/ThumbnailRendering/ThumbnailManager.h
@@ -154,8 +154,7 @@ public:
 	UNREALED_API virtual void UnregisterCustomRenderer(UClass* Class);
 
 	/** 
-	 * Returns the thumbnail pool which should be used for most thumbnails in the editor. 
-	 * Unless you are rendering a huge amount of thumbnails at once this shared pool should be used
+	 * Returns the thumbnail pool which should be used for most thumbnails in the editor.
 	 */
 	UNREALED_API TSharedPtr<FAssetThumbnailPool> GetSharedThumbnailPool() const { return SharedThumbnailPool; }
 
diff --git a/Engine/Source/Editor/UnrealEd/Private/ThumbnailManager.cpp b/Engine/Source/Editor/UnrealEd/Private/ThumbnailManager.cpp
--- a/Engine/Source/Editor/UnrealEd/Private/ThumbnailManager.cpp
+++ b/Engine/Source/Editor/UnrealEd/Private/ThumbnailManager.cpp
@@ -74,7 +74,8 @@ void UThumbnailManager::Initialize(void)
 	{
 		InitializeRenderTypeArray(RenderableThumbnailTypes);
 
-		SharedThumbnailPool = MakeShared<FAssetThumbnailPool>(100);
+		// The size of the pool is a bit large to allow the asset views to use it (each asset view used 1024 by default)
+		SharedThumbnailPool = MakeShared<FAssetThumbnailPool>(4096);
 
 		FCoreUObjectDelegates::OnObjectPropertyChanged.AddUObject(this, &UThumbnailManager::OnObjectPropertyChanged);
 
