Fix content browser remove from favorites
#rnx

diff --git a/Engine/Source/Editor/ContentBrowser/Private/ContentBrowserUtils.cpp b/Engine/Source/Editor/ContentBrowser/Private/ContentBrowserUtils.cpp
--- a/Engine/Source/Editor/ContentBrowser/Private/ContentBrowserUtils.cpp
+++ b/Engine/Source/Editor/ContentBrowser/Private/ContentBrowserUtils.cpp
@@ -776,6 +776,8 @@ void ContentBrowserUtils::RemoveFavoriteFolder(const FContentBrowserItemPath& Fo
 
 	FString InvariantFolder = InvariantPath.ToString();
 
+	FContentBrowserSingleton::Get().FavoriteFolderPaths.Remove(InvariantFolder);
+
 	if (UContentBrowserConfig* EditorConfig = UContentBrowserConfig::Get())
 	{
 		EditorConfig->Favorites.Remove(InvariantFolder);
