Fixed an issue with the content browser not always showing the path filters. The SLATE_ARGUMENT to determine if it should be shown was not initialized.

diff --git a/Engine/Source/Editor/ContentBrowser/Private/SContentBrowser.cpp b/Engine/Source/Editor/ContentBrowser/Private/SContentBrowser.cpp
--- a/Engine/Source/Editor/ContentBrowser/Private/SContentBrowser.cpp
+++ b/Engine/Source/Editor/ContentBrowser/Private/SContentBrowser.cpp
@@ -274,6 +274,7 @@ void SContentBrowser::Construct( const FArguments& InArgs, const FName& InInstan
 		.AddMetaData<FTagMetaData>(FTagMetaData(TEXT("ContentBrowserAssets")))
 		.OwningContentBrowser(SharedThis(this))
 		.OnSearchOptionsChanged(this, &SContentBrowser::HandleAssetViewSearchOptionsChanged)
+		.bShowPathViewFilters(true)
 		.FillEmptySpaceInTileView(true)
 		.ShowDisallowedAssetClassAsUnsupportedItems(true);
 
diff --git a/Engine/Source/Editor/ContentBrowser/Public/SAssetView.h b/Engine/Source/Editor/ContentBrowser/Public/SAssetView.h
--- a/Engine/Source/Editor/ContentBrowser/Public/SAssetView.h
+++ b/Engine/Source/Editor/ContentBrowser/Public/SAssetView.h
@@ -97,6 +97,7 @@ public:
 		: _InitialCategoryFilter(EContentBrowserItemCategoryFilter::IncludeAll)
 		, _ThumbnailLabel( EThumbnailLabel::ClassName )
 		, _AllowThumbnailHintLabel(true)
+		, _bShowPathViewFilters(false)
 		, _InitialViewType(EAssetViewType::Tile)
 		, _InitialThumbnailSize(EThumbnailSize::Medium)
 		, _ShowBottomToolbar(true)
