Hide "Find in Explorer" menu option for cooked content, as there is never a source file available

#rnx

diff --git a/Engine/Source/Editor/ContentBrowser/Private/AssetContextMenu.cpp b/Engine/Source/Editor/ContentBrowser/Private/AssetContextMenu.cpp
--- a/Engine/Source/Editor/ContentBrowser/Private/AssetContextMenu.cpp
+++ b/Engine/Source/Editor/ContentBrowser/Private/AssetContextMenu.cpp
@@ -227,15 +227,17 @@ TSharedRef<SWidget> FAssetContextMenu::MakeContextMenu(TArrayView<const FContent
 			ContextObject->CommonClass = CommonClass;
 
 			ContextObject->bCanBeModified = true;
+			ContextObject->bHasCookedPackages = false;
 			for (const FAssetData& SelectedAsset : SelectedAssets)
 			{
-				if (WritableFolderPermission->HasFiltering() && !WritableFolderPermission->PassesStartsWithFilter(SelectedAsset.PackageName))
+				if (SelectedAsset.HasAnyPackageFlags(PKG_Cooked | PKG_FilterEditorOnly))
 				{
 					ContextObject->bCanBeModified = false;
+					ContextObject->bHasCookedPackages = true;
 					break;
 				}
 
-				if (SelectedAsset.HasAnyPackageFlags(PKG_Cooked | PKG_FilterEditorOnly))
+				if (WritableFolderPermission->HasFiltering() && !WritableFolderPermission->PassesStartsWithFilter(SelectedAsset.PackageName))
 				{
 					ContextObject->bCanBeModified = false;
 					break;
@@ -269,6 +271,7 @@ TSharedRef<SWidget> FAssetContextMenu::MakeContextMenu(TArrayView<const FContent
 			if (CommonDataSource)
 			{
 				ContextObject->bCanBeModified = true;
+				ContextObject->bHasCookedPackages = false;
 
 				if (WritableFolderPermission->HasFiltering())
 				{
@@ -287,6 +290,7 @@ TSharedRef<SWidget> FAssetContextMenu::MakeContextMenu(TArrayView<const FContent
 					if (SelectedAsset.HasAnyPackageFlags(PKG_Cooked | PKG_FilterEditorOnly))
 					{
 						ContextObject->bCanBeModified = false;
+						ContextObject->bHasCookedPackages = true;
 						break;
 					}
 
@@ -317,6 +321,7 @@ TSharedRef<SWidget> FAssetContextMenu::MakeContextMenu(TArrayView<const FContent
 		DataContextObject->SelectedItems = SelectedItems;
 		DataContextObject->SelectedCollections = SourcesData.Collections;
 		DataContextObject->bCanBeModified = ContextObject->bCanBeModified;
+		DataContextObject->bHasCookedPackages = ContextObject->bHasCookedPackages;
 		DataContextObject->bContainsUnsupportedAssets = ContextObject->bContainsUnsupportedAssets;
 		DataContextObject->ParentWidget = AssetView;
 		DataContextObject->OnShowInPathsView = OnShowInPathsViewRequested;
@@ -544,6 +549,8 @@ bool FAssetContextMenu::AddCommonMenuOptions(UToolMenu* Menu)
 
 void FAssetContextMenu::AddExploreMenuOptions(UToolMenu* Menu)
 {
+	UContentBrowserDataMenuContext_FileMenu* Context = Menu->FindContext<UContentBrowserDataMenuContext_FileMenu>();
+
 	FToolMenuSection& Section = Menu->AddSection("AssetContextExploreMenuOptions", LOCTEXT("AssetContextExploreMenuOptionsHeading", "Explore"));
 	{
 		// Find in Content Browser
@@ -553,17 +560,20 @@ void FAssetContextMenu::AddExploreMenuOptions(UToolMenu* Menu)
 			LOCTEXT("ShowInFolderViewTooltip", "Selects the folder that contains this asset in the Content Browser Sources Panel.")
 			);
 
-		// Find in Explorer
-		Section.AddMenuEntry(
-			"FindInExplorer",
-			ContentBrowserUtils::GetExploreFolderText(),
-			LOCTEXT("FindInExplorerTooltip", "Finds this asset on disk"),
-			FSlateIcon(FAppStyle::GetAppStyleSetName(), "ContentBrowser.ShowInExplorer"),
-			FUIAction(
-				FExecuteAction::CreateSP(this, &FAssetContextMenu::ExecuteFindInExplorer),
-				FCanExecuteAction::CreateSP(this, &FAssetContextMenu::CanExecuteFindInExplorer)
-			)
-		);
+		if (!Context->bHasCookedPackages)
+		{
+			// Find in Explorer
+			Section.AddMenuEntry(
+				"FindInExplorer",
+				ContentBrowserUtils::GetExploreFolderText(),
+				LOCTEXT("FindInExplorerTooltip", "Finds this asset on disk"),
+				FSlateIcon(FAppStyle::GetAppStyleSetName(), "ContentBrowser.ShowInExplorer"),
+				FUIAction(
+					FExecuteAction::CreateSP(this, &FAssetContextMenu::ExecuteFindInExplorer),
+					FCanExecuteAction::CreateSP(this, &FAssetContextMenu::CanExecuteFindInExplorer)
+				)
+			);
+		}
 	}
 }
 
@@ -649,7 +659,7 @@ bool FAssetContextMenu::AddReferenceMenuOptions(UToolMenu* Menu)
 			FUIAction( FExecuteAction::CreateSP( this, &FAssetContextMenu::ExecuteCopyReference ) )
 			);
 	
-		if (Context->bCanBeModified)
+		if (!Context->bHasCookedPackages)
 		{
 			Section.AddMenuEntry(
 				"CopyFilePath",
diff --git a/Engine/Source/Editor/ContentBrowser/Public/ContentBrowserMenuContexts.h b/Engine/Source/Editor/ContentBrowser/Public/ContentBrowserMenuContexts.h
--- a/Engine/Source/Editor/ContentBrowser/Public/ContentBrowserMenuContexts.h
+++ b/Engine/Source/Editor/ContentBrowser/Public/ContentBrowserMenuContexts.h
@@ -63,6 +63,9 @@ public:
 	UPROPERTY()
 	bool bCanBeModified;
 
+	UPROPERTY()
+	bool bHasCookedPackages;
+
 	UPROPERTY(BlueprintReadOnly, Category = ContentBrowser)
 	bool bContainsUnsupportedAssets = true;
 
diff --git a/Engine/Source/Editor/ContentBrowserData/Public/ContentBrowserDataMenuContexts.h b/Engine/Source/Editor/ContentBrowserData/Public/ContentBrowserDataMenuContexts.h
--- a/Engine/Source/Editor/ContentBrowserData/Public/ContentBrowserDataMenuContexts.h
+++ b/Engine/Source/Editor/ContentBrowserData/Public/ContentBrowserDataMenuContexts.h
@@ -77,6 +77,9 @@ public:
 	UPROPERTY(BlueprintReadOnly, Category = ContentBrowser)
 	bool bCanBeModified = true;
 
+	UPROPERTY(BlueprintReadOnly, Category = ContentBrowser)
+	bool bHasCookedPackages = false;
+
 	UPROPERTY(BlueprintReadOnly, Category = ContentBrowser)
 	bool bContainsUnsupportedAssets = true;
 
