Fixed "Folder" appearing on all folders in the Content Browser

This was a regression from CL# 21760852, which was intended to only remove the size restriction from the text visibility. This change re-adds and updates the visibility function so that it's still hidden for types that would have never shown the type regardless of view size.

#rnx

diff --git a/Engine/Source/Editor/ContentBrowser/Private/AssetViewWidgets.cpp b/Engine/Source/Editor/ContentBrowser/Private/AssetViewWidgets.cpp
--- a/Engine/Source/Editor/ContentBrowser/Private/AssetViewWidgets.cpp
+++ b/Engine/Source/Editor/ContentBrowser/Private/AssetViewWidgets.cpp
@@ -1948,6 +1948,7 @@ void SAssetTileItem::Construct( const FArguments& InArgs )
 								+SHorizontalBox::Slot()
 								[
 									SAssignNew(ClassTextWidget, STextBlock)
+									.Visibility(this, &SAssetTileItem::GetAssetClassLabelVisibility)
 									.TextStyle(FAppStyle::Get(), "ContentBrowser.ClassFont")
 									.OverflowPolicy(ETextOverflowPolicy::Ellipsis)
 									.Text(this, &SAssetTileItem::GetAssetClassText)
@@ -2074,16 +2075,9 @@ FOptionalSize SAssetTileItem::GetThumbnailBoxSize() const
 
 EVisibility SAssetTileItem::GetAssetClassLabelVisibility() const
 {
-	if(!IsFolder() && bShowType)
-	{
-		FOptionalSize ThumbSize = GetThumbnailBoxSize();
-		if (ThumbSize.IsSet())
-		{
-			float Size = ThumbSize.Get() * GetTickSpaceGeometry().Scale;
-			return Size < 100 ? EVisibility::Collapsed : EVisibility::Visible;
-		}
-	}
-	return EVisibility::Collapsed;
+	return (!IsFolder() && bShowType)
+		? EVisibility::Visible
+		: EVisibility::Collapsed;
 }
 
 FSlateColor SAssetTileItem::GetAssetClassLabelTextColor() const
diff --git a/Engine/Source/Editor/ContentBrowser/Public/AssetViewWidgets.h b/Engine/Source/Editor/ContentBrowser/Public/AssetViewWidgets.h
--- a/Engine/Source/Editor/ContentBrowser/Public/AssetViewWidgets.h
+++ b/Engine/Source/Editor/ContentBrowser/Public/AssetViewWidgets.h
@@ -713,7 +713,6 @@ protected:
 	FOptionalSize GetThumbnailBoxSize() const;
 
 	/** Gets the visibility of the asset class label in thumbnails */
-	UE_DEPRECATED(5.1, "GetAssetClassLabelVisibility is deprecated, as there is no longer a need to identify a visual difference in the label based on various states.")
 	EVisibility GetAssetClassLabelVisibility() const;
 
 	/** Gets the color of the asset class label in thumbnails */
