Re-add ability to hide collections that was removed by 23662444
#rnx

diff --git a/Engine/Source/Editor/ContentBrowser/Private/ContentBrowserConfig.h b/Engine/Source/Editor/ContentBrowser/Private/ContentBrowserConfig.h
--- a/Engine/Source/Editor/ContentBrowser/Private/ContentBrowserConfig.h
+++ b/Engine/Source/Editor/ContentBrowser/Private/ContentBrowserConfig.h
@@ -86,6 +86,8 @@ struct FContentBrowserInstanceConfig
 	UPROPERTY()
 	bool bShowCppFolders = false;
 
+	bool bCollectionsDocked = true;
+
 	UPROPERTY()
 	bool bSearchClasses = false;
 
diff --git a/Engine/Source/Editor/ContentBrowser/Private/SAssetView.cpp b/Engine/Source/Editor/ContentBrowser/Private/SAssetView.cpp
--- a/Engine/Source/Editor/ContentBrowser/Private/SAssetView.cpp
+++ b/Engine/Source/Editor/ContentBrowser/Private/SAssetView.cpp
@@ -256,6 +256,7 @@ void SAssetView::Construct( const FArguments& InArgs )
 	bCanShowDevelopersFolder = InArgs._CanShowDevelopersFolder;
 
 	bCanShowFavorites = InArgs._CanShowFavorites;
+	bCanDockCollections = InArgs._CanDockCollections;
 
 	SelectionMode = InArgs._SelectionMode;
 
@@ -3219,6 +3220,43 @@ bool SAssetView::IsShowingFavorites() const
 	return GetDefault<UContentBrowserSettings>()->GetDisplayFavorites();
 }
 
+void SAssetView::ToggleDockCollections()
+{
+	check(IsToggleDockCollectionsAllowed()); 
+
+	bool bNewState = !GetDefault<UContentBrowserSettings>()->GetDockCollections();
+
+	if (FContentBrowserInstanceConfig* Config = GetContentBrowserConfig())
+	{
+		bNewState = !Config->bCollectionsDocked;
+		Config->bCollectionsDocked = bNewState;
+		UContentBrowserConfig::Get()->SaveEditorConfig();
+	}
+
+	GetMutableDefault<UContentBrowserSettings>()->SetDockCollections(bNewState);
+	GetMutableDefault<UContentBrowserSettings>()->PostEditChange();
+}
+
+bool SAssetView::IsToggleDockCollectionsAllowed() const
+{
+	return bCanDockCollections;
+}
+
+bool SAssetView::HasDockedCollections() const
+{
+	if (!IsToggleDockCollectionsAllowed())
+	{
+		return false;
+	}
+
+	if (const FContentBrowserInstanceConfig* Config = GetContentBrowserConfig())
+	{
+		return Config->bCollectionsDocked;
+	}
+
+	return GetDefault<UContentBrowserSettings>()->GetDockCollections();
+}
+
 void SAssetView::ToggleShowCppContent()
 {
 	check(IsToggleShowCppContentAllowed());
diff --git a/Engine/Source/Editor/ContentBrowser/Private/SContentBrowser.cpp b/Engine/Source/Editor/ContentBrowser/Private/SContentBrowser.cpp
--- a/Engine/Source/Editor/ContentBrowser/Private/SContentBrowser.cpp
+++ b/Engine/Source/Editor/ContentBrowser/Private/SContentBrowser.cpp
@@ -156,6 +156,12 @@ struct FSlateBrush;
 
 const FString SContentBrowser::SettingsIniSection = TEXT("ContentBrowser");
 
+namespace ContentBrowserSourcesWidgetSwitcherIndex
+{
+	static const int32 PathView = 0;
+	static const int32 CollectionsView = 1;
+}
+
 SContentBrowser::~SContentBrowser()
 {
 	// Remove the listener for when view settings are changed
@@ -239,7 +245,7 @@ void SContentBrowser::Construct( const FArguments& InArgs, const FName& InInstan
 		.AddMetaData<FTagMetaData>(FTagMetaData(TEXT("ContentBrowserCollections")))
 		.AllowCollectionDrag(true)
 		.AllowQuickAssetManagement(true)
-		.IsDocked(true)
+		.IsDocked(this, &SContentBrowser::IsCollectionViewDocked)
 		.ExternalSearch(CollectionSearch);
 
 	static const FName DefaultForegroundName("DefaultForeground");
@@ -271,6 +277,7 @@ void SContentBrowser::Construct( const FArguments& InArgs, const FName& InInstan
 		.CanShowRealTimeThumbnails(Config != nullptr ? Config->bCanShowRealTimeThumbnails : true)
 		.CanShowDevelopersFolder(Config != nullptr ? Config->bCanShowDevelopersFolder : true)
 		.CanShowFavorites(true)
+		.CanDockCollections(true)
 		.AddMetaData<FTagMetaData>(FTagMetaData(TEXT("ContentBrowserAssets")))
 		.OwningContentBrowser(SharedThis(this))
 		.OnSearchOptionsChanged(this, &SContentBrowser::HandleAssetViewSearchOptionsChanged)
@@ -461,6 +468,13 @@ void SContentBrowser::Construct( const FArguments& InArgs, const FName& InInstan
 					SNew(SBorder)
 					.Padding(FMargin(0))
 					.BorderImage(FAppStyle::GetBrush("Brushes.Recessed"))
+
+					[
+						// Note: If adding more widgets here, fix ContentBrowserSourcesWidgetSwitcherIndex and the code that uses it!
+						SAssignNew(SourcesWidgetSwitcher, SWidgetSwitcher)
+
+						// Paths View
+						+SWidgetSwitcher::Slot()
 						[
 							SAssignNew(PathFavoriteSplitterPtr, SSplitter)
 							.Clipping(EWidgetClipping::ClipToBounds)
@@ -468,11 +482,10 @@ void SContentBrowser::Construct( const FArguments& InArgs, const FName& InInstan
 							.HitDetectionSplitterHandleSize(8.0f)
 							.Orientation(EOrientation::Orient_Vertical)
 							.MinimumSlotHeight(26.0f)
-						.Visibility(this, &SContentBrowser::GetSourcesViewVisibility)
-
-						+ SSplitter::Slot()
-						.SizeRule(this, &SContentBrowser::GetFavoritesAreaSizeRule)
-						.MinSize(this, &SContentBrowser::GetFavoritesAreaMinSize)
+							.Visibility( this, &SContentBrowser::GetSourcesViewVisibility )
+							+SSplitter::Slot()
+							.SizeRule(TAttribute<SSplitter::ESizeRule>(this, &SContentBrowser::GetFavoritesAreaSizeRule))
+							.MinSize(TAttribute<float>(this, &SContentBrowser::GetFavoritesAreaMinSize))
 							.Value(0.2f)
 							[
 								SNew(SBorder)
@@ -483,8 +496,8 @@ void SContentBrowser::Construct( const FArguments& InArgs, const FName& InInstan
 								]
 							]
 								
-						+ SSplitter::Slot()
-						.SizeRule(this, &SContentBrowser::GetPathAreaSizeRule)
+							+SSplitter::Slot()
+							.SizeRule(TAttribute<SSplitter::ESizeRule>(this, &SContentBrowser::GetPathAreaSizeRule))
 							.MinSize(29.0f)
 							.Value(0.8f)
 							[
@@ -496,9 +509,9 @@ void SContentBrowser::Construct( const FArguments& InArgs, const FName& InInstan
 								]
 							]
 
-						+ SSplitter::Slot()
-						.SizeRule(this, &SContentBrowser::GetCollectionsAreaSizeRule)
-						.MinSize(29.0f)
+							+SSplitter::Slot()
+							.SizeRule(TAttribute<SSplitter::ESizeRule>(this, &SContentBrowser::GetCollectionsAreaSizeRule))
+							.MinSize(TAttribute<float>(this, &SContentBrowser::GetCollectionsAreaMinSize))
 							.Value(0.4f)
 							[
 								SNew(SBorder)
@@ -509,8 +522,19 @@ void SContentBrowser::Construct( const FArguments& InArgs, const FName& InInstan
 								]
 							]
 						]
+
+						// Collections View
+						+SWidgetSwitcher::Slot()
+						[
+							SNew(SBox)
+							.Padding(FMargin(0.0f, 4.0f, 0.0f, 0.0f))
+							[
+								CollectionViewPtr.ToSharedRef()
+							]
+						]
 					]
 				]
+			]
 
 			// Asset View
 			+ SSplitter::Slot()
@@ -657,11 +681,26 @@ EVisibility SContentBrowser::GetFavoriteFolderVisibility() const
 	return GetDefault<UContentBrowserSettings>()->GetDisplayFavorites() ? EVisibility::Visible : EVisibility::Collapsed;
 }
 
+EVisibility SContentBrowser::GetDockedCollectionsVisibility() const
+{
+	return IsCollectionViewDocked() ? EVisibility::Visible : EVisibility::Collapsed;
+}
+
 EVisibility SContentBrowser::GetLockButtonVisibility() const
 {
 	return IsLocked() ? EVisibility::Visible : EVisibility::Collapsed;
 }
 
+bool SContentBrowser::IsCollectionViewDocked() const
+{
+	//if (const FContentBrowserInstanceConfig* Config = GetConstInstanceConfig())
+	//{
+	//	return Config->bCollectionsDocked;
+	//}
+
+	return GetDefault<UContentBrowserSettings>()->GetDockCollections();
+}
+
 void SContentBrowser::ToggleFolderFavorite(const TArray<FString>& FolderPaths)
 {
 	bool bAddedFavorite = false;
@@ -1124,6 +1163,7 @@ TSharedRef<SWidget> SContentBrowser::CreateDockedCollectionsView(const FContentB
 		.BodyBorderImage(FAppStyle::Get().GetBrush("Brushes.Recessed"))
 		.HeaderPadding(FMargin(4.0f, 0.0f))
 		.Padding(0.f)
+		.Visibility(this, &SContentBrowser::GetDockedCollectionsVisibility)
 		.AllowAnimatedTransition(false)
 		.OnAreaExpansionChanged_Lambda([this](bool bIsExpanded) { if (!bIsExpanded) CollectionSearchToggleButton->SetExpanded(false); })
 		.HeaderContent()
@@ -1381,7 +1421,7 @@ SSplitter::ESizeRule SContentBrowser::GetPathAreaSizeRule() const
 SSplitter::ESizeRule SContentBrowser::GetCollectionsAreaSizeRule() const
 {
 	// Make sure the area is expanded and visible 
-	return CollectionArea->IsExpanded() ? SSplitter::ESizeRule::FractionOfParent : SSplitter::ESizeRule::SizeToContent;
+	return CollectionArea->IsExpanded() && GetDockedCollectionsVisibility() == EVisibility::Visible ? SSplitter::ESizeRule::FractionOfParent : SSplitter::ESizeRule::SizeToContent;
 }
 
 float SContentBrowser::GetFavoritesAreaMinSize() const
@@ -1389,6 +1429,11 @@ float SContentBrowser::GetFavoritesAreaMinSize() const
 	return GetFavoriteFolderVisibility() == EVisibility::Visible ? 29.0f : 0.0f;
 }
 
+float SContentBrowser::GetCollectionsAreaMinSize() const
+{
+	return GetDockedCollectionsVisibility() == EVisibility::Visible ? 29.0f : 0.0f;
+}
+
 FText SContentBrowser::GetHighlightedText() const
 {
 	return TextFilter->GetRawFilterText();
@@ -3161,6 +3206,19 @@ FReply SContentBrowser::SourcesViewExpandClicked()
 	return FReply::Handled();
 }
 
+EVisibility SContentBrowser::GetSourcesSwitcherVisibility() const
+{
+	return IsCollectionViewDocked() ? EVisibility::Collapsed : EVisibility::Visible;
+}
+
+FReply SContentBrowser::OnSourcesSwitcherClicked()
+{
+	ActiveSourcesWidgetIndex = (ActiveSourcesWidgetIndex + 1) % SourcesWidgetSwitcher->GetNumWidgets();
+	SourcesWidgetSwitcher->SetActiveWidgetIndex(ActiveSourcesWidgetIndex);
+
+	return FReply::Handled();
+}
+
 void SContentBrowser::OnContentBrowserSettingsChanged(FName PropertyName)
 {
 	if (PropertyName.IsNone())
@@ -3556,8 +3614,12 @@ void SContentBrowser::UpdatePath()
 
 	PathBreadcrumbTrail->ClearCrumbs();
 
+	int32 NewSourcesWidgetIndex = ActiveSourcesWidgetIndex;
+
 	if ( SourcesData.HasVirtualPaths() )
 	{
+		NewSourcesWidgetIndex = ContentBrowserSourcesWidgetSwitcherIndex::PathView;
+
 		UContentBrowserDataSubsystem* ContentBrowserData = IContentBrowserDataModule::Get().GetSubsystem();
 
 		TArray<FString> Crumbs;
@@ -3576,6 +3638,8 @@ void SContentBrowser::UpdatePath()
 	}
 	else if ( SourcesData.HasCollections() )
 	{
+		NewSourcesWidgetIndex = IsCollectionViewDocked() ? ContentBrowserSourcesWidgetSwitcherIndex::PathView : ContentBrowserSourcesWidgetSwitcherIndex::CollectionsView;
+
 		FCollectionManagerModule& CollectionManagerModule = FCollectionManagerModule::GetModule();
 		TArray<FCollectionNameType> CollectionPathItems;
 
@@ -3605,6 +3669,12 @@ void SContentBrowser::UpdatePath()
 		PathBreadcrumbTrail->PushCrumb(LOCTEXT("AllAssets", "All Assets"), TEXT(""));
 	}
 
+	if (ActiveSourcesWidgetIndex != NewSourcesWidgetIndex)
+	{
+		ActiveSourcesWidgetIndex = NewSourcesWidgetIndex;
+		SourcesWidgetSwitcher->SetActiveWidgetIndex(ActiveSourcesWidgetIndex);
+	}
+
 	CachedCanWriteToCurrentPath.Reset();
 }
 
@@ -4362,6 +4432,7 @@ FContentBrowserInstanceConfig* SContentBrowser::CreateEditorConfigIfRequired()
 	InstanceConfig->bShowFolders = Settings->DisplayFolders;
 	InstanceConfig->bShowEmptyFolders = Settings->DisplayEmptyFolders;
 	InstanceConfig->bShowCppFolders = Settings->GetDisplayCppFolders();
+	InstanceConfig->bCollectionsDocked = Settings->GetDockCollections();
 	InstanceConfig->bFavoritesExpanded = Settings->GetDisplayFavorites();
 	InstanceConfig->bSearchAssetPaths = Settings->GetIncludeAssetPaths();
 	InstanceConfig->bSearchClasses = Settings->GetIncludeClassNames();
diff --git a/Engine/Source/Editor/ContentBrowser/Private/SContentBrowser.h b/Engine/Source/Editor/ContentBrowser/Private/SContentBrowser.h
--- a/Engine/Source/Editor/ContentBrowser/Private/SContentBrowser.h
+++ b/Engine/Source/Editor/ContentBrowser/Private/SContentBrowser.h
@@ -44,6 +44,7 @@ class SFilterList;
 class SPathView;
 class SSearchToggleButton;
 class SWidget;
+class SWidgetSwitcher;
 class UClass;
 class UContentBrowserToolbarMenuContext;
 class UFactory;
@@ -387,6 +388,18 @@ private:
 	/** Handler for clicking the tree expand/collapse button */
 	FReply SourcesViewExpandClicked();
 
+	/** Gets the visibility of the source switch button */
+	EVisibility GetSourcesSwitcherVisibility() const;
+
+	/** Gets the icon used on the source switch button */
+	const FSlateBrush* GetSourcesSwitcherIcon() const;
+
+	/** Gets the tooltip text used on the source switch button */
+	FText GetSourcesSwitcherToolTipText() const;
+
+	/** Handler for clicking the source switch button */
+	FReply OnSourcesSwitcherClicked();
+
 	/** Called to handle the Content Browser settings changing */
 	void OnContentBrowserSettingsChanged(FName PropertyName);
 
@@ -531,9 +544,15 @@ private:
 	/** Gets the visibility of the favorites view */
 	EVisibility GetFavoriteFolderVisibility() const;
 
+	/** Get the visibility of the docked collections view */
+	EVisibility GetDockedCollectionsVisibility() const;
+
 	/** Get the visibility of the lock button */
 	EVisibility GetLockButtonVisibility() const;
 
+	/** Whether or not the collections view is docked or exists in its own panel in the same area as the sources view */
+	bool IsCollectionViewDocked() const;
+
 	/** Toggles the favorite status of an array of folders*/
 	void ToggleFolderFavorite(const TArray<FString>& FolderPaths);
 
@@ -577,6 +596,7 @@ private:
 
 	/** Gets the min size for various areas. When areas are not visible the min size is 0, otherwise there is a minimum size to prevent overlap */
 	float GetFavoritesAreaMinSize() const;
+	float GetCollectionsAreaMinSize() const;
 
 	/** Called when the layout of the SFilterList is changing */
 	void OnFilterBarLayoutChanging(EFilterBarLayout NewLayout);
@@ -661,6 +681,9 @@ private:
 	/** Toggle button for showing/hiding collection search area */
 	TSharedPtr<SSearchToggleButton> CollectionSearchToggleButton;
 
+	/** Index of the active sources widget */
+	int32 ActiveSourcesWidgetIndex = 0;
+
 	/** The expanded state of the asset tree */
 	bool bSourcesViewExpanded = true;
 
@@ -697,6 +720,9 @@ private:
 	/** Delegate used to create a new folder */
 	FOnCreateNewFolder OnCreateNewFolder;
 
+	/** Switcher between the different sources views */
+	TSharedPtr<SWidgetSwitcher> SourcesWidgetSwitcher;
+
 	/** The splitter between the path & asset view */
 	TSharedPtr<SSplitter> PathAssetSplitterPtr;
 
diff --git a/Engine/Source/Editor/ContentBrowser/Public/SAssetView.h b/Engine/Source/Editor/ContentBrowser/Public/SAssetView.h
--- a/Engine/Source/Editor/ContentBrowser/Public/SAssetView.h
+++ b/Engine/Source/Editor/ContentBrowser/Public/SAssetView.h
@@ -110,6 +110,7 @@ public:
 		, _CanShowRealTimeThumbnails(false)
 		, _CanShowDevelopersFolder(false)
 		, _CanShowFavorites(false)
+		, _CanDockCollections(false)
 		, _SelectionMode( ESelectionMode::Multi )
 		, _AllowDragging(true)
 		, _AllowFocusOnSync(true)
@@ -227,7 +228,7 @@ public:
 		SLATE_ARGUMENT(bool, CanShowFavorites)
 
 		/** Indicates if the 'Dock Collections' option should be enabled or disabled */
-		SLATE_ARGUMENT_DEPRECATED(bool, CanDockCollections, 5.2, "CanDockCollections is no longer supported (now defaults to true).")
+		SLATE_ARGUMENT(bool, CanDockCollections)
 
 		/** The selection mode the asset view should use */
 		SLATE_ARGUMENT( ESelectionMode::Type, SelectionMode )
@@ -595,6 +596,15 @@ private:
 	/** @return true when we are showing favorites */
 	bool IsShowingFavorites() const;
 
+	/** Toggle whether the collections view should be docked under the paths view */
+	void ToggleDockCollections();
+
+	/** Whether or not it's possible to dock the collections view */
+	bool IsToggleDockCollectionsAllowed() const;
+
+	/** @return true when the collections view is docked */
+	bool HasDockedCollections() const;
+
 	/** Toggle whether C++ content should be shown or not */
 	void ToggleShowCppContent();
 
@@ -1061,6 +1071,9 @@ private:
 	/** Indicates if the 'Show Favorites' option should be enabled or disabled */
 	bool bCanShowFavorites : 1;
 
+	/** Indicates if the 'Dock Collections' option should be enabled or disabled */
+	bool bCanDockCollections : 1;
+
 	/** If true, it will show path column in the asset view */
 	bool bShowPathInColumnView : 1;
 
