Add an Icon extension point for the content browser path view that allows icons to be placed on the right of the item text.

diff --git a/Engine/Source/Editor/ContentBrowser/Private/SourcesViewWidgets.cpp b/Engine/Source/Editor/ContentBrowser/Private/SourcesViewWidgets.cpp
--- a/Engine/Source/Editor/ContentBrowser/Private/SourcesViewWidgets.cpp
+++ b/Engine/Source/Editor/ContentBrowser/Private/SourcesViewWidgets.cpp
@@ -10,6 +10,7 @@
 #include "ContentBrowserDataSource.h"
 #include "ContentBrowserItem.h"
 #include "ContentBrowserItemData.h"
+#include "ContentBrowserModule.h"
 #include "ContentBrowserPluginFilters.h"
 #include "ContentBrowserUtils.h"
 #include "DragAndDrop/AssetDragDropOp.h"
@@ -91,7 +92,6 @@ void SAssetTreeItem::Construct( const FArguments& InArgs )
 			]
 
 			+SHorizontalBox::Slot()
-			.AutoWidth()
 			.VAlign(VAlign_Center)
 			[
 				SAssignNew(InlineRenameWidget, SInlineEditableTextBlock)
@@ -104,6 +104,12 @@ void SAssetTreeItem::Construct( const FArguments& InArgs )
 					.IsSelected( InArgs._IsSelected )
 					.IsReadOnly( this, &SAssetTreeItem::IsReadOnly )
 			]
+			+SHorizontalBox::Slot()
+			.AutoWidth()
+			.HAlign(HAlign_Right)
+			[
+				GenerateStateIcons()
+			]
 		]
 	];
 
@@ -113,6 +119,46 @@ void SAssetTreeItem::Construct( const FArguments& InArgs )
 	}
 }
 
+TSharedRef<SWidget> SAssetTreeItem::GenerateStateIcons()
+{
+	TSharedRef<SBox> ContainingBox = SNew(SBox);
+	TSharedPtr<SHorizontalBox> HorizonalBox;
+	FContentBrowserModule& ContentBrowserModule = FModuleManager::GetModuleChecked<FContentBrowserModule>( TEXT("ContentBrowser") );
+	if (const TSharedPtr<FTreeItem> TreeItemPinned = TreeItem.Pin())
+	{
+		const FContentBrowserItem& ContentBrowserItem = TreeItemPinned->GetItem();
+		for(const FPathViewStateIconGenerator& Generator: ContentBrowserModule.GetAllPathViewStateIconGenerators())
+		{
+			if (Generator.IsBound())
+			{
+				if (TSharedPtr<SWidget> IconWidget = Generator.Execute(ContentBrowserItem))
+				{
+					if (!HorizonalBox)
+					{
+						HorizonalBox = SNew(SHorizontalBox);
+					}
+					
+					HorizonalBox->AddSlot()
+					.AutoWidth()
+					.VAlign(VAlign_Center)
+					.Padding(4, 0, 0, 0)
+					[
+						IconWidget.ToSharedRef()
+					];
+				}
+			}
+		}
+		// If we created any content, add it to the containg box
+		// and set padding
+		if (HorizonalBox)
+		{
+			ContainingBox->SetContent(HorizonalBox.ToSharedRef());
+			ContainingBox->SetPadding(FMargin(2, 0, 4, 0));
+		}
+	}
+	return ContainingBox;
+}
+
 SAssetTreeItem::~SAssetTreeItem()
 {
 	if( InlineRenameWidget.IsValid() )
diff --git a/Engine/Source/Editor/ContentBrowser/Private/SourcesViewWidgets.h b/Engine/Source/Editor/ContentBrowser/Private/SourcesViewWidgets.h
--- a/Engine/Source/Editor/ContentBrowser/Private/SourcesViewWidgets.h
+++ b/Engine/Source/Editor/ContentBrowser/Private/SourcesViewWidgets.h
@@ -99,6 +99,9 @@ private:
 	/** Returns the image for the border around this item. Used for drag/drop operations */
 	const FSlateBrush* GetBorderImage() const;
 
+	/** Returns a widget that combines the results of calling all registered PathViewStateIconGenerators **/
+	TSharedRef<SWidget> GenerateStateIcons();
+	
 private:
 	enum class EFolderType : uint8
 	{
diff --git a/Engine/Source/Editor/ContentBrowser/Public/ContentBrowserDelegates.h b/Engine/Source/Editor/ContentBrowser/Public/ContentBrowserDelegates.h
--- a/Engine/Source/Editor/ContentBrowser/Public/ContentBrowserDelegates.h
+++ b/Engine/Source/Editor/ContentBrowser/Public/ContentBrowserDelegates.h
@@ -110,6 +110,9 @@ DECLARE_DELEGATE_RetVal_TwoParams(FText, FOnGetCustomAssetColumnDisplayText, FAs
 /** Called to generate extra state information icons or tooltips on asset items. */
 DECLARE_DELEGATE_RetVal_OneParam(TSharedRef<SWidget>, FOnGenerateAssetViewExtraStateIndicators, const FAssetData& /*AssetData*/);
 
+/** Called to generate state information icons for path view items. */
+DECLARE_DELEGATE_RetVal_OneParam(TSharedPtr<SWidget>, FPathViewStateIconGenerator, const FContentBrowserItem& /*ContentBrowserItem*/);
+
 /** Called when an asset item visualizes its tooltip */
 DECLARE_DELEGATE_RetVal_TwoParams(bool, FOnVisualizeAssetToolTip, const TSharedPtr<SWidget>& /*ToolTipContent*/, FAssetData& /*AssetData*/);
 
diff --git a/Engine/Source/Editor/ContentBrowser/Public/ContentBrowserModule.h b/Engine/Source/Editor/ContentBrowser/Public/ContentBrowserModule.h
--- a/Engine/Source/Editor/ContentBrowser/Public/ContentBrowserModule.h
+++ b/Engine/Source/Editor/ContentBrowser/Public/ContentBrowserModule.h
@@ -163,6 +163,7 @@ public:
 	virtual TArray<FContentBrowserMenuExtender>& GetAllCollectionViewContextMenuExtenders() {return CollectionViewContextMenuExtenders;}
 	virtual TArray<FContentBrowserMenuExtender_SelectedAssets>& GetAllAssetViewContextMenuExtenders() {return AssetViewContextMenuExtenders;}
 	virtual TArray<FContentBrowserMenuExtender>& GetAllAssetViewViewMenuExtenders() {return AssetViewViewMenuExtenders;}
+	virtual TArray<FPathViewStateIconGenerator>& GetAllPathViewStateIconGenerators() {return PathViewStateIconGenerators;}
 
 	/** Delegates to call to extend the command/keybinds for content browser */
 	virtual TArray<FContentBrowserCommandExtender>& GetAllContentBrowserCommandExtenders() { return ContentBrowserCommandExtenders; }
@@ -232,6 +233,7 @@ private:
 	TArray<FContentBrowserMenuExtender> CollectionViewContextMenuExtenders;
 	TArray<FContentBrowserMenuExtender_SelectedAssets> AssetViewContextMenuExtenders;
 	TArray<FContentBrowserMenuExtender> AssetViewViewMenuExtenders;
+	TArray<FPathViewStateIconGenerator> PathViewStateIconGenerators;
 	TArray<FContentBrowserCommandExtender> ContentBrowserCommandExtenders;
 
 	/** All delegates generating extra state indicators */
