Content Browser: Immediately flush the config on adding/removing a folder from favorites to make sure it is saved if the editor crashes

diff --git a/Engine/Source/Editor/ContentBrowser/Private/SPathView.cpp b/Engine/Source/Editor/ContentBrowser/Private/SPathView.cpp
--- a/Engine/Source/Editor/ContentBrowser/Private/SPathView.cpp
+++ b/Engine/Source/Editor/ContentBrowser/Private/SPathView.cpp
@@ -2377,6 +2377,7 @@ void SFavoritePathView::SaveSettings(const FString& IniFilename, const FString&
 	}
 
 	GConfig->SetString(*IniSection, TEXT("FavoritePaths"), *FavoritePathsString, IniFilename);
+	GConfig->Flush(false, IniFilename);
 }
 
 void SFavoritePathView::LoadSettings(const FString& IniFilename, const FString& IniSection, const FString& SettingsString)
