Don't show the public status of non-assets

#rnx

diff --git a/Engine/Source/Editor/ContentBrowser/Private/AssetViewWidgets.cpp b/Engine/Source/Editor/ContentBrowser/Private/AssetViewWidgets.cpp
--- a/Engine/Source/Editor/ContentBrowser/Private/AssetViewWidgets.cpp
+++ b/Engine/Source/Editor/ContentBrowser/Private/AssetViewWidgets.cpp
@@ -770,7 +770,8 @@ TSharedRef<SWidget> SAssetViewItem::CreateToolTipWidget() const
 			// The tooltip contains the name, class, path, asset registry tags and source control status
 			const FText NameText = GetNameText();
 			const FText ClassText = FText::Format(LOCTEXT("ClassName", "({0})"), GetAssetClassText());
-			FText PublicStateText = LOCTEXT("PublicAssetState", "Public");
+
+			FText PublicStateText;
 			FName PublicStateTextBorder = "ContentBrowser.TileViewTooltip.PillBorder";
 
 			// Create a box to hold every line of info in the body of the tooltip
@@ -819,6 +820,10 @@ TSharedRef<SWidget> SAssetViewItem::CreateToolTipWidget() const
 				{
 					PublicStateText = LOCTEXT("PrivateAssetState", "Private");
 				}
+				else
+				{
+					PublicStateText = LOCTEXT("PublicAssetState", "Public");
+				}
 			}
 
 			if (!AssetItem->GetItem().CanEdit())
@@ -859,12 +864,7 @@ TSharedRef<SWidget> SAssetViewItem::CreateToolTipWidget() const
 			TSharedRef<SVerticalBox> OverallTooltipVBox = SNew(SVerticalBox);
 
 			static const auto PublicAssetUIEnabledCVar = IConsoleManager::Get().FindConsoleVariable(TEXT("ContentBrowser.PublicAsset.EnablePublicAssetFeature"));
-			bool bIsPublicAssetUIEnabled = false;
-
-			if (PublicAssetUIEnabledCVar)
-			{
-				bIsPublicAssetUIEnabled = PublicAssetUIEnabledCVar->GetBool();
-			}		 
+			const bool bIsPublicAssetUIEnabled = PublicAssetUIEnabledCVar && PublicAssetUIEnabledCVar->GetBool();
 
 			// Top section (asset name, type, is checked out)
 			OverallTooltipVBox->AddSlot()
@@ -909,7 +909,7 @@ TSharedRef<SWidget> SAssetViewItem::CreateToolTipWidget() const
 								SNew(SBorder)
 								.BorderImage(FAppStyle::GetBrush(PublicStateTextBorder))
 								.Padding(FMargin(12.0f, 2.0f, 12.0f, 2.0f))
-								.Visibility(bIsPublicAssetUIEnabled ? EVisibility::Visible : EVisibility::Hidden)
+								.Visibility(bIsPublicAssetUIEnabled && !PublicStateText.IsEmpty() ? EVisibility::Visible : EVisibility::Hidden)
 								[
 									SNew(STextBlock)
 									.Text(PublicStateText)
