Provide telemetry to SAssetView. Telemetry is not sent to any endpoint for licensees.

diff --git a/Engine/Source/Developer/ToolWidgets/Public/Filters/FilterBase.h b/Engine/Source/Developer/ToolWidgets/Public/Filters/FilterBase.h
--- a/Engine/Source/Developer/ToolWidgets/Public/Filters/FilterBase.h
+++ b/Engine/Source/Developer/ToolWidgets/Public/Filters/FilterBase.h
@@ -35,7 +35,7 @@ public:
 	FFilterBase(TSharedPtr<FFilterCategory> InCategory) : FilterCategory(InCategory) {}
 
 	/** Returns the system name for this filter */
-	virtual FString GetName() const = 0;
+	virtual FString GetName() const override = 0;
 
 	/** Returns the human readable name for this filter */
 	virtual FText GetDisplayName() const = 0;
diff --git a/Engine/Source/Editor/ContentBrowser/ContentBrowser.Build.cs b/Engine/Source/Editor/ContentBrowser/ContentBrowser.Build.cs
--- a/Engine/Source/Editor/ContentBrowser/ContentBrowser.Build.cs
+++ b/Engine/Source/Editor/ContentBrowser/ContentBrowser.Build.cs
@@ -43,6 +43,7 @@ public class ContentBrowser : ModuleRules
 				"ToolMenus",
 				"StatusBar",
 				"ToolWidgets",
+				"TelemetryUtils"
 			}
 		);
 
diff --git a/Engine/Source/Editor/ContentBrowser/Private/SAssetView.cpp b/Engine/Source/Editor/ContentBrowser/Private/SAssetView.cpp
--- a/Engine/Source/Editor/ContentBrowser/Private/SAssetView.cpp
+++ b/Engine/Source/Editor/ContentBrowser/Private/SAssetView.cpp
@@ -49,6 +49,7 @@
 #include "Settings/ContentBrowserSettings.h"
 #include "SlateOptMacros.h"
 #include "Styling/AppStyle.h"
+#include "TelemetryRouter.h"
 #include "Textures/SlateIcon.h"
 #include "ThumbnailRendering/ThumbnailManager.h"
 #include "ToolMenus.h"
@@ -190,6 +191,8 @@ SAssetView::~SAssetView()
 
 void SAssetView::Construct( const FArguments& InArgs )
 {
+	ViewCorrelationGuid = FGuid::NewGuid();	
+
 	InitialNumAmortizedTasks = 0;
 	TotalAmortizeTime = 0;
 	AmortizeStartTime = 0;
@@ -1044,9 +1047,20 @@ void SAssetView::Tick( const FGeometry& AllottedGeometry, const double InCurrent
 		{
 			AmortizeStartTime = FPlatformTime::Seconds();
 			InitialNumAmortizedTasks = ItemsPendingFrontendFilter.Num();
+			
+			CurrentFrontendFilterTelemetry = { ViewCorrelationGuid, FilterSessionCorrelationGuid };
+			CurrentFrontendFilterTelemetry.FrontendFilters = FrontendFilters;
+			CurrentFrontendFilterTelemetry.TotalItemsToFilter = ItemsPendingFrontendFilter.Num() + ItemsPendingPriorityFilter.Num();
+			CurrentFrontendFilterTelemetry.PriorityItemsToFilter = ItemsPendingPriorityFilter.Num();
 		}
 
+		int32 PreviousFilteredAssetItems = FilteredAssetItems.Num();
 		ProcessItemsPendingFilter(TickStartTime);
+		if (PreviousFilteredAssetItems == 0 && FilteredAssetItems.Num() != 0)
+		{
+			CurrentFrontendFilterTelemetry.ResultLatency = FPlatformTime::Seconds() - AmortizeStartTime;
+		}
+		CurrentFrontendFilterTelemetry.TotalResults = FilteredAssetItems.Num(); // Provide number of results even if filtering is interrupted
 
 		if (HasItemsPendingFilter())
 		{
@@ -1055,16 +1069,14 @@ void SAssetView::Tick( const FGeometry& AllottedGeometry, const double InCurrent
 				// Don't sync to selection if we are just going to do it below
 				SortList(!PendingSyncItems.Num());
 			}
+			
+			CurrentFrontendFilterTelemetry.WorkDuration += FPlatformTime::Seconds() - TickStartTime;
 
 			// Need to finish processing queried items before rest of function is safe
 			return;
 		}
 		else
 		{
-			TotalAmortizeTime += FPlatformTime::Seconds() - AmortizeStartTime;
-			AmortizeStartTime = 0;
-			InitialNumAmortizedTasks = 0;
-
 			// Update the columns in the column view now that we know the majority type
 			if (CurrentViewType == EAssetViewType::Column)
 			{
@@ -1087,6 +1099,15 @@ void SAssetView::Tick( const FGeometry& AllottedGeometry, const double InCurrent
 				// Don't sync to selection if we are just going to do it below
 				SortList(!PendingSyncItems.Num());
 			}
+
+			CurrentFrontendFilterTelemetry.WorkDuration += FPlatformTime::Seconds() - TickStartTime;
+
+			double AmortizeDuration = FPlatformTime::Seconds() - AmortizeStartTime;
+			TotalAmortizeTime += AmortizeDuration;
+			AmortizeStartTime = 0;
+			InitialNumAmortizedTasks = 0;
+			
+			OnCompleteFiltering(AmortizeDuration);
 		}
 	}
 
@@ -1928,8 +1949,13 @@ FContentBrowserDataFilter SAssetView::CreateBackendDataFilter(bool bInvalidateCa
 
 void SAssetView::RefreshSourceItems()
 {
+	TRACE_CPUPROFILER_EVENT_SCOPE("SAssetView::RefreshSourceItems");
 	const double RefreshSourceItemsStartTime = FPlatformTime::Seconds();
+	
+	OnInterruptFiltering();
 
+	FilterSessionCorrelationGuid = FGuid::NewGuid();
+	UE::Telemetry::ContentBrowser::FBackendFilterTelemetry Telemetry(ViewCorrelationGuid, FilterSessionCorrelationGuid);
 	FilteredAssetItems.Reset();
 	FilteredAssetItemTypeCounts.Reset();
 	VisibleItems.Reset();
@@ -1972,13 +1998,16 @@ void SAssetView::RefreshSourceItems()
 
 		if (SourcesData.OnEnumerateCustomSourceItemDatas.IsBound())
 		{
+			Telemetry.bHasCustomItemSources = true;
 			SourcesData.OnEnumerateCustomSourceItemDatas.Execute(AddNewItem);
 		}
 
+		FContentBrowserDataFilter DataFilter; // Must live long enough to provide telemetry
 		if (SourcesData.IsIncludingVirtualPaths() || SourcesData.HasCollections()) 
 		{
 			const bool bInvalidateFilterCache = true;
-			const FContentBrowserDataFilter DataFilter = CreateBackendDataFilter(bInvalidateFilterCache);
+			DataFilter = CreateBackendDataFilter(bInvalidateFilterCache);
+			Telemetry.DataFilter = &DataFilter;
 
 			bWereItemsRecursivelyFiltered = DataFilter.bRecursivePaths;
 
@@ -2029,6 +2058,10 @@ void SAssetView::RefreshSourceItems()
 				}
 			}
 		}
+
+		Telemetry.NumBackendItems = AvailableBackendItems.Num();
+		Telemetry.RefreshSourceItemsDurationSeconds = FPlatformTime::Seconds() - RefreshSourceItemsStartTime;
+		FTelemetryRouter::Get().ProvideTelemetry(Telemetry);
 	}
 
 	UE_LOG(LogContentBrowser, VeryVerbose, TEXT("AssetView - RefreshSourceItems completed in %0.4f seconds"), FPlatformTime::Seconds() - RefreshSourceItemsStartTime);
@@ -2117,6 +2150,8 @@ void SAssetView::RefreshFilteredItems()
 {
 	const double RefreshFilteredItemsStartTime = FPlatformTime::Seconds();
 
+	OnInterruptFiltering();
+	
 	ItemsPendingFrontendFilter.Reset();
 	FilteredAssetItems.Reset();
 	FilteredAssetItemTypeCounts.Reset();
@@ -3571,6 +3606,7 @@ void SAssetView::OnOpenAssetsOrFolders()
 {
 	if (OnItemsActivated.IsBound())
 	{
+		OnInteractDuringFiltering();
 		const TArray<FContentBrowserItem> SelectedItems = GetSelectedItems();
 		OnItemsActivated.Execute(SelectedItems, EAssetTypeActivationMethod::Opened);
 	}
@@ -3580,6 +3616,7 @@ void SAssetView::OnPreviewAssets()
 {
 	if (OnItemsActivated.IsBound())
 	{
+		OnInteractDuringFiltering();
 		const TArray<FContentBrowserItem> SelectedItems = GetSelectedItems();
 		OnItemsActivated.Execute(SelectedItems, EAssetTypeActivationMethod::Previewed);
 	}
@@ -3964,6 +4001,7 @@ TSharedPtr<SWidget> SAssetView::OnGetContextMenuContent()
 			RenamingAsset.Reset();
 		}
 
+		OnInteractDuringFiltering();
 		const TArray<FContentBrowserItem> SelectedItems = GetSelectedItems();
 		return OnGetItemContextMenu.Execute(SelectedItems);
 	}
@@ -4039,6 +4077,7 @@ void SAssetView::OnListMouseButtonDoubleClick(TSharedPtr<FAssetViewItem> AssetIt
 
 	if (OnItemsActivated.IsBound())
 	{
+		OnInteractDuringFiltering();
 		OnItemsActivated.Execute(MakeArrayView(&AssetItem->GetItem(), 1), EAssetTypeActivationMethod::DoubleClicked);
 	}
 }
@@ -4047,6 +4086,7 @@ FReply SAssetView::OnDraggingAssetItem( const FGeometry& MyGeometry, const FPoin
 {
 	if (bAllowDragging)
 	{
+		OnInteractDuringFiltering();
 		// Use the custom drag handler?
 		if (FEditorDelegates::OnAssetDragStarted.IsBound())
 		{
@@ -4105,6 +4145,8 @@ void SAssetView::AssetRenameBegin(const TSharedPtr<FAssetViewItem>& Item, const
 	check(!RenamingAsset.IsValid());
 	RenamingAsset = Item;
 
+	OnInteractDuringFiltering();
+
 	if (DeferredItemToCreate.IsValid())
 	{
 		UE_LOG(LogContentBrowser, Log, TEXT("Renaming the item being created (Deferred Item: %s)."), *Item->GetItem().GetItemName().ToString());
@@ -5005,6 +5047,33 @@ void SAssetView::SetFilterBar(TSharedPtr<SFilterList> InFilterBar)
 	FilterBar = InFilterBar;
 }
 
+void SAssetView::OnCompleteFiltering(double InAmortizeDuration)
+{
+	CurrentFrontendFilterTelemetry.AmortizeDuration = InAmortizeDuration;
+	CurrentFrontendFilterTelemetry.bCompleted = true;
+	FTelemetryRouter::Get().ProvideTelemetry(CurrentFrontendFilterTelemetry);
+	CurrentFrontendFilterTelemetry = {};
+}
+
+void SAssetView::OnInterruptFiltering()
+{
+	if (CurrentFrontendFilterTelemetry.FilterSessionCorrelationGuid.IsValid())
+	{
+		CurrentFrontendFilterTelemetry.AmortizeDuration = FPlatformTime::Seconds() - AmortizeStartTime;
+		CurrentFrontendFilterTelemetry.bCompleted = false;
+		FTelemetryRouter::Get().ProvideTelemetry(CurrentFrontendFilterTelemetry);
+		CurrentFrontendFilterTelemetry = {};
+	}
+}
+
+void SAssetView::OnInteractDuringFiltering()
+{
+	if (CurrentFrontendFilterTelemetry.FilterSessionCorrelationGuid.IsValid() && !CurrentFrontendFilterTelemetry.TimeUntilInteraction.IsSet())
+	{
+		CurrentFrontendFilterTelemetry.TimeUntilInteraction = FPlatformTime::Seconds() - AmortizeStartTime;
+	}
+}
+
 #undef checkAssetList
 #undef ASSET_VIEW_PARANOIA_LIST_CHECKS
 
diff --git a/Engine/Source/Editor/ContentBrowser/Public/ContentBrowserTelemetry.h b/Engine/Source/Editor/ContentBrowser/Public/ContentBrowserTelemetry.h
new file mode 100644
--- /dev/null
+++ b/Engine/Source/Editor/ContentBrowser/Public/ContentBrowserTelemetry.h
@@ -0,0 +1,76 @@
+// Copyright Epic Games, Inc. All Rights Reserved.
+
+#pragma once
+
+#include "CoreTypes.h"
+#include "IContentBrowserSingleton.h"
+#include "Misc/Guid.h"
+#include "Misc/Optional.h"
+
+struct CONTENTBROWSERDATA_API FContentBrowserDataFilter;
+
+namespace UE::Telemetry::ContentBrowser
+{
+	struct FBackendFilterTelemetry
+	{
+		static inline constexpr FGuid TelemetryID = FGuid(0x3a026bff,0x041444d0, 0x9457c4e5, 0x1378d03c);
+
+		FBackendFilterTelemetry(FGuid InViewCorrelationGuid, FGuid InFilterSessionCorrelationGuid)
+		: ViewCorrelationGuid(InViewCorrelationGuid)
+		, FilterSessionCorrelationGuid(InFilterSessionCorrelationGuid)
+		{
+		}
+
+		// Guid assigned to the originating instance of SAssetView
+		FGuid ViewCorrelationGuid;
+		// Guid assigned to this set of backend items for correlation with frontend filtering	
+		FGuid FilterSessionCorrelationGuid;
+		// Whether this view had a delegate bound for providing custom items
+		bool bHasCustomItemSources = false;
+		// Wall-clock duration of fetching all items from the backend
+		double RefreshSourceItemsDurationSeconds = 0.0;
+		// Number of items returned from the backend
+		int32 NumBackendItems = 0;
+		// Optional pointer to the filter passed to the backend - not to be stored past the lifetime of the telemetry handler.
+		const FContentBrowserDataFilter* DataFilter = nullptr;
+	};
+	
+	struct FFrontendFilterTelemetry
+	{
+		static inline constexpr FGuid TelemetryID = FGuid(0xee47ffd0, 0x86c74117, 0x86416a64, 0x85460b31);
+
+		FFrontendFilterTelemetry()
+		{
+		}
+
+		FFrontendFilterTelemetry(FGuid InViewCorrelationGuid, FGuid InFilterSessionCorrelationGuid)
+		: ViewCorrelationGuid(InViewCorrelationGuid)
+		, FilterSessionCorrelationGuid(InFilterSessionCorrelationGuid)
+		{
+		}
+
+		// Guid assigned to the originating instance of SAssetView
+		FGuid ViewCorrelationGuid = FGuid();
+		// Guid assigned to the set of backend items the frontend is filtering here	
+		FGuid FilterSessionCorrelationGuid = FGuid();
+		// The total number of items returned from the backend that needed filtering
+		int32 TotalItemsToFilter = 0;
+		// The number of items that were required to be filtered immediately
+		int32 PriorityItemsToFilter = 0;
+        // Total number of items that passed the filters 
+        int32 TotalResults = 0;
+		// Wall clock duration of filtering all items
+		double AmortizeDuration = 0.0;
+		// Actual cumulative tick time of filtering 
+		double WorkDuration = 0.0;
+        // Latency until first result is added (if any results were found)
+        TOptional<double> ResultLatency;
+		// Time until user interacted with filter results - e.g. opened an asset from the list, even if filtering was not complete
+		TOptional<double> TimeUntilInteraction;
+		// Whether this filtering session completed before being interrupted/cancelled in some way
+		bool bCompleted = false;
+		// 'Frontend' filter (amortized in asset view tick) 
+		// 'Backend' filter data is available in backend items telemetry - requires correlation 
+		TSharedPtr<FAssetFilterCollectionType> FrontendFilters;
+	};
+}
\ No newline at end of file
diff --git a/Engine/Source/Editor/ContentBrowser/Public/SAssetView.h b/Engine/Source/Editor/ContentBrowser/Public/SAssetView.h
--- a/Engine/Source/Editor/ContentBrowser/Public/SAssetView.h
+++ b/Engine/Source/Editor/ContentBrowser/Public/SAssetView.h
@@ -16,6 +16,7 @@
 #include "ContentBrowserDataLegacyBridge.h"
 #include "ContentBrowserDelegates.h"
 #include "ContentBrowserItem.h"
+#include "ContentBrowserTelemetry.h"
 #include "CoreMinimal.h"
 #include "Delegates/Delegate.h"
 #include "Framework/Views/ITypedTableView.h"
@@ -1194,6 +1195,24 @@ private:
 
 	/** An Id for the cache of the data sources for the filters compilation */
 	FContentBrowserDataFilterCacheIDOwner FilterCacheID;
+	
+	/*
+	 * Telemetry-related fields
+	 */ 
+	// Guid to correlate different telemetry events for this instance of an asset view over time 
+	FGuid ViewCorrelationGuid; 
+	// Guid to correlate different telemetry events for a 'session' of search/filtering, reset when new backend items are gathered for filtering
+	FGuid FilterSessionCorrelationGuid;
+	// In progress filter session telemetry - will be sent when filtering ends or user cancels or interacts with a partial filter set
+	UE::Telemetry::ContentBrowser::FFrontendFilterTelemetry CurrentFrontendFilterTelemetry;
+	
+	// If there is in-progress filter telemetry, mark it as complete and send pending data 
+	void OnCompleteFiltering(double InAmortizeDuration);
+	// If there is in-progress filter telemetry, mark it as incomplete and send pending data 
+	void OnInterruptFiltering();
+	// Update telemetry for user interaction during an incomplete filtering
+	void OnInteractDuringFiltering();
+
 public:
 	bool ShouldColumnGenerateWidget(const FString ColumnName) const;
 };
diff --git a/Engine/Source/Editor/ContentBrowserData/Private/ContentBrowserDataFilter.cpp b/Engine/Source/Editor/ContentBrowserData/Private/ContentBrowserDataFilter.cpp
--- a/Engine/Source/Editor/ContentBrowserData/Private/ContentBrowserDataFilter.cpp
+++ b/Engine/Source/Editor/ContentBrowserData/Private/ContentBrowserDataFilter.cpp
@@ -19,6 +19,17 @@ FContentBrowserDataFilterList& FContentBrowserDataFilterList::operator=(const FC
 	return *this;
 }
 
+TArray<const UScriptStruct*> FContentBrowserDataFilterList::GetFilterTypes() const
+{
+	TArray<const UScriptStruct*> Types;
+	Types.Reserve(TypedFilters.Num());
+	for (const FStructOnScope& Filter : TypedFilters)
+	{
+		Types.Add(CastChecked<const UScriptStruct>(Filter.GetStruct()));
+	}
+	return Types;
+}
+
 void* FContentBrowserDataFilterList::FindOrAddFilter(const UScriptStruct* InFilterType)
 {
 	check(InFilterType);
diff --git a/Engine/Source/Editor/ContentBrowserData/Public/ContentBrowserDataFilter.h b/Engine/Source/Editor/ContentBrowserData/Public/ContentBrowserDataFilter.h
--- a/Engine/Source/Editor/ContentBrowserData/Public/ContentBrowserDataFilter.h
+++ b/Engine/Source/Editor/ContentBrowserData/Public/ContentBrowserDataFilter.h
@@ -26,7 +26,7 @@ class UContentBrowserDataSource;
 class UContentBrowserDataSubsystem;
 
 /** Flags controlling which item types should be included */
-UENUM()
+UENUM(Flags)
 enum class EContentBrowserItemTypeFilter : uint8
 {
 	IncludeNone = 0,
@@ -37,7 +37,7 @@ enum class EContentBrowserItemTypeFilter : uint8
 ENUM_CLASS_FLAGS(EContentBrowserItemTypeFilter);
 
 /** Flags controlling which item categories should be included */
-UENUM()
+UENUM(Flags)
 enum class EContentBrowserItemCategoryFilter : uint8
 {
 	IncludeNone = 0,
@@ -50,7 +50,7 @@ enum class EContentBrowserItemCategoryFilter : uint8
 ENUM_CLASS_FLAGS(EContentBrowserItemCategoryFilter);
 
 /** Flags controlling which item attributes should be included */
-UENUM()
+UENUM(Flags)
 enum class EContentBrowserItemAttributeFilter : uint8
 {
 	IncludeNone = 0,
@@ -138,6 +138,8 @@ public:
 
 	/** Remove all filters in the list */
 	void ClearFilters();
+	
+	TArray<const UScriptStruct*> GetFilterTypes() const;
 
 private:
 	/** Set the contents of this list to be a deep copy of the contents of the other list */
diff --git a/Engine/Source/Runtime/Core/Public/Misc/IFilter.h b/Engine/Source/Runtime/Core/Public/Misc/IFilter.h
--- a/Engine/Source/Runtime/Core/Public/Misc/IFilter.h
+++ b/Engine/Source/Runtime/Core/Public/Misc/IFilter.h
@@ -15,6 +15,9 @@ public:
 	typedef TItemType ItemType;
 
 	virtual ~IFilter(){ }
+	
+	/** Convenience function for identifying types of filter in logging, telemetry, etc. */
+	virtual FString GetName() const { return FString{}; };
 
 	/** Returns whether the specified Item passes the Filter's restrictions */
 	virtual bool PassesFilter( TItemType InItem ) const = 0;
