Content Browser: Hide the scc icon in tiny thumbnail size and when there is no valid icon

diff --git a/Engine/Source/Editor/ContentBrowser/Private/AssetViewWidgets.cpp b/Engine/Source/Editor/ContentBrowser/Private/AssetViewWidgets.cpp
--- a/Engine/Source/Editor/ContentBrowser/Private/AssetViewWidgets.cpp
+++ b/Engine/Source/Editor/ContentBrowser/Private/AssetViewWidgets.cpp
@@ -650,7 +650,9 @@ void SAssetViewItem::HandleSourceControlStateChanged()
 			{
 				if (SCCStateWidget.IsValid())
 				{
-					SCCStateWidget->SetFromSlateIcon(SourceControlState->GetIcon());
+					FSlateIcon SCCIcon = SourceControlState->GetIcon();
+					bHasCCStateBrush = SCCIcon.GetIcon() != nullptr;
+					SCCStateWidget->SetFromSlateIcon(SCCIcon);
 				}
 			}
 		}
@@ -1965,6 +1967,7 @@ void SAssetTileItem::Construct( const FArguments& InArgs )
 									SNew(SBox)
 									.WidthOverride(16.0f)
 									.HeightOverride(16.0f)
+									.Visibility(this, &SAssetTileItem::GetSCCIconVisibility)
 									[
 										GenerateSourceControlIconWidget()
 									]
@@ -2204,6 +2207,13 @@ int32 SAssetTileItem::GetGenericThumbnailSize() const
 	return GenericThumbnailSizes[(int32)CurrentThumbnailSize.Get()];
 }
 
+EVisibility SAssetTileItem::GetSCCIconVisibility() const
+{
+	// Hide the scc state icon when there is no brush or in tiny size since there isn't enough space
+	return bHasCCStateBrush &&  CurrentThumbnailSize.Get() != EThumbnailSize::Tiny ? EVisibility::Visible : EVisibility::Collapsed;
+}
+
+
 void SAssetTileItem::InitializeAssetNameHeights()
 {
 	// The height of the asset name field for each thumbnail size
diff --git a/Engine/Source/Editor/ContentBrowser/Public/AssetViewWidgets.h b/Engine/Source/Editor/ContentBrowser/Public/AssetViewWidgets.h
--- a/Engine/Source/Editor/ContentBrowser/Public/AssetViewWidgets.h
+++ b/Engine/Source/Editor/ContentBrowser/Public/AssetViewWidgets.h
@@ -351,6 +351,9 @@ protected:
 	/** Widget for the source control state */
 	TSharedPtr<SLayeredImage> SCCStateWidget;
 
+	/** Whether the item has a valid scc state brush (not empty) */
+	bool bHasCCStateBrush = false;
+
 	/** Delegate handle for the HandleSourceControlStateChanged function callback */
 	FDelegateHandle SourceControlStateChangedDelegateHandle;
 };
@@ -734,6 +737,10 @@ protected:
 
 	int32 GetGenericThumbnailSize() const;
 
+	/** Gets the visibility of the SCC icons */
+	EVisibility GetSCCIconVisibility() const;
+
+
 private:
 	/** If false, the tooltip will not be displayed */
 	bool bShowType;
