Content Browser: Fix tile view weirdness by reverting changes done to AssetViewWidgets.cpp in CL 26326540

diff --git a/Engine/Source/Editor/ContentBrowser/Private/AssetViewWidgets.cpp b/Engine/Source/Editor/ContentBrowser/Private/AssetViewWidgets.cpp
--- a/Engine/Source/Editor/ContentBrowser/Private/AssetViewWidgets.cpp
+++ b/Engine/Source/Editor/ContentBrowser/Private/AssetViewWidgets.cpp
@@ -1913,7 +1913,6 @@ void SAssetTileItem::Construct( const FArguments& InArgs )
 					]
 
 					+SVerticalBox::Slot()
-					.AutoHeight()
 					[
 						SNew(SBorder)
 						.Padding(FMargin(2.0f, 3.0f))
@@ -1923,7 +1922,6 @@ void SAssetTileItem::Construct( const FArguments& InArgs )
 							+ SVerticalBox::Slot()
 							.Padding(2.0f,2.0f,0.0f,0.0f)
 							.VAlign(VAlign_Top)
-						    .AutoHeight()
 							.HAlign(IsFolder() ? HAlign_Center : HAlign_Left)
 							[
 								SNew(SBox)
