[ContentBrowser] Add cvar for debugging enumeration issues to shut off breadcrumb content browser data source queries each tick to.
Allows setting breakpoints and not having them hit each tick by breadcrumb.
#rnx

diff --git a/Engine/Source/Editor/ContentBrowser/Private/SContentBrowser.cpp b/Engine/Source/Editor/ContentBrowser/Private/SContentBrowser.cpp
--- a/Engine/Source/Editor/ContentBrowser/Private/SContentBrowser.cpp
+++ b/Engine/Source/Editor/ContentBrowser/Private/SContentBrowser.cpp
@@ -162,6 +162,13 @@ namespace ContentBrowserSourcesWidgetSwitcherIndex
 	static const int32 CollectionsView = 1;
 }
 
+namespace ContentBrowser
+{
+	// It is useful when debugging enumeration issues to shut this system off so that breakpoints set will not be triggered every tick
+	bool bCrumbsEnumerate = true;
+	FAutoConsoleVariableRef CrumbsEnumerateCVar(TEXT("ContentBrowser.Debug.CrumbsEnumerate"), bCrumbsEnumerate, TEXT("Enumerate crumbs"), ECVF_Default);
+}
+
 SContentBrowser::~SContentBrowser()
 {
 	// Remove the listener for when view settings are changed
@@ -2467,11 +2474,14 @@ bool SContentBrowser::OnHasCrumbDelimiterContent(const FString& CrumbData) const
 		SubItemsFilter.ItemAttributeFilter = PathViewPtr->GetContentBrowserItemAttributeFilter();
 
 		bool bHasSubItems = false;
-		ContentBrowserData->EnumerateItemsUnderPath(*CrumbData, SubItemsFilter, [&bHasSubItems](FContentBrowserItemData&& InSubItem)
+		if (ContentBrowser::bCrumbsEnumerate)
 		{
-			bHasSubItems = true;
-			return false;
-		});
+			ContentBrowserData->EnumerateItemsUnderPath(*CrumbData, SubItemsFilter, [&bHasSubItems](FContentBrowserItemData&& InSubItem)
+			{
+				bHasSubItems = true;
+				return false;
+			});
+		}
 
 		return bHasSubItems;
 	}
