AssetView: fix a bug where the name area is too small in Medium-sized tile views when we choose not to show the type of the asset.

diff --git a/Engine/Source/Editor/ContentBrowser/Private/AssetViewWidgets.cpp b/Engine/Source/Editor/ContentBrowser/Private/AssetViewWidgets.cpp
--- a/Engine/Source/Editor/ContentBrowser/Private/AssetViewWidgets.cpp
+++ b/Engine/Source/Editor/ContentBrowser/Private/AssetViewWidgets.cpp
@@ -1881,6 +1881,36 @@ void SAssetTileItem::Construct( const FArguments& InArgs )
 	FName ItemShadowBorderName;
 	TSharedRef<SWidget> ItemContents = FAssetViewItemHelper::CreateTileItemContents(this, Thumbnail.ToSharedRef(), ItemShadowBorderName);
 
+	if (bShowType)
+	{
+		SAssignNew(InlineRenameWidget, SInlineEditableTextBlock)
+			.Font(this, &SAssetTileItem::GetThumbnailFont)
+			.Text(GetNameText())
+			.OnBeginTextEdit(this, &SAssetTileItem::HandleBeginNameChange)
+			.OnTextCommitted(this, &SAssetTileItem::HandleNameCommitted)
+			.OnVerifyTextChanged(this, &SAssetTileItem::HandleVerifyNameChanged)
+			.HighlightText(InArgs._HighlightText)
+			.IsSelected(InArgs._IsSelectedExclusively)
+			.IsReadOnly(this, &SAssetTileItem::IsNameReadOnly)
+			.LineBreakPolicy(FBreakIterator::CreateCamelCaseBreakIterator())
+			.OverflowPolicy(ETextOverflowPolicy::Ellipsis)
+			.ColorAndOpacity(this, &SAssetTileItem::GetNameAreaTextColor);
+	}
+	else
+	{
+		SAssignNew(InlineRenameWidget, SInlineEditableTextBlock)
+			.Font(this, &SAssetTileItem::GetThumbnailFont)
+			.Text(GetNameText())
+			.OnBeginTextEdit(this, &SAssetTileItem::HandleBeginNameChange)
+			.OnTextCommitted(this, &SAssetTileItem::HandleNameCommitted)
+			.OnVerifyTextChanged(this, &SAssetTileItem::HandleVerifyNameChanged)
+			.HighlightText(InArgs._HighlightText)
+			.IsSelected(InArgs._IsSelectedExclusively)
+			.IsReadOnly(this, &SAssetTileItem::IsNameReadOnly)
+			.OverflowPolicy(ETextOverflowPolicy::Ellipsis)
+			.ColorAndOpacity(this, &SAssetTileItem::GetNameAreaTextColor);
+	}
+
 	ChildSlot
 	.Padding(FMargin(0.0f, 0.0f, 4.0f, 4.0f))
 	[				
@@ -1927,18 +1957,7 @@ void SAssetTileItem::Construct( const FArguments& InArgs )
 								SNew(SBox)
 								.MaxDesiredHeight(this, &SAssetTileItem::GetNameAreaMaxDesiredHeight)
 								[
-									SAssignNew(InlineRenameWidget, SInlineEditableTextBlock)
-									.Font(this, &SAssetTileItem::GetThumbnailFont)
-									.Text(GetNameText())
-									.OnBeginTextEdit(this, &SAssetTileItem::HandleBeginNameChange)
-									.OnTextCommitted(this, &SAssetTileItem::HandleNameCommitted)
-									.OnVerifyTextChanged(this, &SAssetTileItem::HandleVerifyNameChanged)
-									.HighlightText(InArgs._HighlightText)
-									.IsSelected(InArgs._IsSelectedExclusively)
-									.IsReadOnly(this, &SAssetTileItem::IsNameReadOnly)
-									.LineBreakPolicy(FBreakIterator::CreateCamelCaseBreakIterator())
-									.OverflowPolicy(ETextOverflowPolicy::Ellipsis)
-									.ColorAndOpacity(this, &SAssetTileItem::GetNameAreaTextColor)
+									InlineRenameWidget.ToSharedRef()
 								]
 							]
 							+ SVerticalBox::Slot()
diff --git a/Engine/Source/Editor/ContentBrowser/Private/SAssetView.cpp b/Engine/Source/Editor/ContentBrowser/Private/SAssetView.cpp
--- a/Engine/Source/Editor/ContentBrowser/Private/SAssetView.cpp
+++ b/Engine/Source/Editor/ContentBrowser/Private/SAssetView.cpp
@@ -4413,7 +4413,11 @@ float SAssetView::GetTileViewTypeNameHeight() const
 		{
 			TypeNameHeight = 25;
 		}
-		else if (ThumbnailSize > EThumbnailSize::Small)
+		else if (ThumbnailSize == EThumbnailSize::Medium)
+		{
+			TypeNameHeight = -5;
+		}
+		else if (ThumbnailSize > EThumbnailSize::Medium)
 		{
 			TypeNameHeight = -25;
 		}
