Update and optimize the collection manager callback that happens after the asset registry scan, this version is twice as fast
Add a direct accessor to Collection's asset set for fast iteration

diff --git a/Engine/Source/Developer/CollectionManager/Private/Collection.h b/Engine/Source/Developer/CollectionManager/Private/Collection.h
--- a/Engine/Source/Developer/CollectionManager/Private/Collection.h
+++ b/Engine/Source/Developer/CollectionManager/Private/Collection.h
@@ -124,6 +124,9 @@ public:
 	/** Get the source filename of this collection */
 	FORCEINLINE const FString& GetSourceFilename() const { return SourceFilename; }
 
+	/** Gets set of all objects for fast iteration */
+	FORCEINLINE const TSet<FSoftObjectPath>& GetObjectSet() const { return ObjectSet; }
+
 private:
 	/** Generates the header pairs for the collection file. */
 	void SaveHeaderPairs(TMap<FString,FString>& OutHeaderPairs) const;
diff --git a/Engine/Source/Developer/CollectionManager/Private/CollectionManager.cpp b/Engine/Source/Developer/CollectionManager/Private/CollectionManager.cpp
--- a/Engine/Source/Developer/CollectionManager/Private/CollectionManager.cpp
+++ b/Engine/Source/Developer/CollectionManager/Private/CollectionManager.cpp
@@ -81,8 +81,7 @@ const FCollectionObjectsMap& FCollectionManagerCache::GetCachedObjects() const
 			const FCollectionNameType& CollectionKey = AvailableCollection.Key;
 			const TSharedRef<FCollection>& Collection = AvailableCollection.Value;
 
-			TArray<FSoftObjectPath> ObjectsInCollection;
-			Collection->GetObjectsInCollection(ObjectsInCollection);
+			const TSet<FSoftObjectPath>& ObjectsInCollection = Collection->GetObjectSet();
 
 			if (ObjectsInCollection.Num() > 0)
 			{
diff --git a/Engine/Source/Editor/ContentBrowser/Private/CollectionAssetRegistryBridge.cpp b/Engine/Source/Editor/ContentBrowser/Private/CollectionAssetRegistryBridge.cpp
--- a/Engine/Source/Editor/ContentBrowser/Private/CollectionAssetRegistryBridge.cpp
+++ b/Engine/Source/Editor/ContentBrowser/Private/CollectionAssetRegistryBridge.cpp
@@ -3,7 +3,6 @@
 #include "CollectionAssetRegistryBridge.h"
 
 #include "AssetRegistry/AssetData.h"
-#include "AssetRegistry/AssetRegistryModule.h"
 #include "AssetRegistry/IAssetRegistry.h"
 #include "CollectionManagerModule.h"
 #include "CollectionManagerTypes.h"
@@ -33,125 +32,96 @@ class FCollectionRedirectorFollower : public ICollectionRedirectorFollower
 {
 public:
 	FCollectionRedirectorFollower()
-		: AssetRegistryModule(FModuleManager::LoadModuleChecked<FAssetRegistryModule>(TEXT("AssetRegistry")))
+		: AssetRegistry(IAssetRegistry::Get())
 	{
 	}
 
 	virtual bool FixupObject(const FSoftObjectPath& InObjectPath, FSoftObjectPath& OutNewObjectPath) override
 	{
-		OutNewObjectPath = FSoftObjectPath();
+		// Most of the time it will be in the asset registry so early return
+		FAssetData ObjectAssetData = AssetRegistry->GetAssetByObjectPath(InObjectPath, true);
 
-		FString InObjectPathString = InObjectPath.ToString();
-		if (!InObjectPathString.StartsWith(TEXT("/")))
+		if (ObjectAssetData.IsValid())
 		{
-			if (!GIsSavingPackage)
+			if (!ObjectAssetData.IsRedirector())
 			{
-				FTopLevelAssetPath FullPath = UClass::TryConvertShortTypeNameToPathName(UClass::StaticClass(), InObjectPathString);
-				if (FullPath.IsValid())
-				{
-					FString ClassPathStr = FullPath.ToString();
-					const FString NewClassName = FLinkerLoad::FindNewPathNameForClass(ClassPathStr, false);
-					if (!NewClassName.IsEmpty())
-					{
-						check(FPackageName::IsValidObjectPath(NewClassName));
-						// Our new class name might be lacking the path, so try and find it so we can use the full path in the collection
-						UClass* FoundClass = FindObject<UClass>(nullptr, *NewClassName);
-						if (FoundClass)
-						{
-							OutNewObjectPath = *FoundClass->GetPathName();
-						}
-					}
-				}
+				OutNewObjectPath = InObjectPath;
+				return false;
+			}
+			else
+			{
+				OutNewObjectPath = AssetRegistry->GetRedirectedObjectPath(InObjectPath);
 			}
 		}
-		else if (InObjectPathString.StartsWith(TEXT("/Script/"))) 
+		else
 		{
-			// We can't use FindObject while we're saving
-			if (!GIsSavingPackage)
+			FString InObjectPathString = InObjectPath.ToString();
+			if (!InObjectPathString.StartsWith(TEXT("/")))
 			{
-				const FString ClassPathStr = InObjectPath.ToString();
-				check(FPackageName::IsValidObjectPath(ClassPathStr));
-				UClass* FoundClass = FindObject<UClass>(nullptr, *ClassPathStr);
-				if (!FoundClass)
+				if (!GIsSavingPackage)
 				{
-					// Use the linker to search for class name redirects (from the loaded ActiveClassRedirects)
-					const FString NewClassName = FLinkerLoad::FindNewPathNameForClass(ClassPathStr, false);
-
-					if (!NewClassName.IsEmpty())
+					FTopLevelAssetPath FullPath = UClass::TryConvertShortTypeNameToPathName(UClass::StaticClass(), InObjectPathString);
+					if (FullPath.IsValid())
 					{
-						check(FPackageName::IsValidObjectPath(NewClassName));
-						// Our new class name might be lacking the path, so try and find it so we can use the full path in the collection
-						FoundClass = FindObject<UClass>(nullptr, *NewClassName);
-						if (FoundClass)
+						FString ClassPathStr = FullPath.ToString();
+						const FString NewClassName = FLinkerLoad::FindNewPathNameForClass(ClassPathStr, false);
+						if (!NewClassName.IsEmpty())
 						{
-							OutNewObjectPath = *FoundClass->GetPathName();
+							check(FPackageName::IsValidObjectPath(NewClassName));
+							// Our new class name might be lacking the path, so try and find it so we can use the full path in the collection
+							UClass* FoundClass = FindObject<UClass>(nullptr, *NewClassName);
+							if (FoundClass)
+							{
+								OutNewObjectPath = *FoundClass->GetPathName();
+							}
 						}
 					}
 				}
 			}
-		}
-		else
-		{
-			// Keep track of visted redirectors in case we loop.
-			TSet<FSoftObjectPath> VisitedRedirectors;
-
-			// Use the asset registry to avoid loading the object
-PRAGMA_DISABLE_DEPRECATION_WARNINGS
-			FAssetData ObjectAssetData = AssetRegistryModule.Get().GetAssetByObjectPath(InObjectPath, true);
-PRAGMA_ENABLE_DEPRECATION_WARNINGS
-			while (ObjectAssetData.IsValid() && ObjectAssetData.IsRedirector())
+			else if (InObjectPathString.StartsWith(TEXT("/Script/")))
 			{
-				// Check to see if we've already seen this path before, it's possible we might have found a redirector loop.
-				if ( VisitedRedirectors.Contains(ObjectAssetData.ToSoftObjectPath()) )
+				// We can't use FindObject while we're saving
+				if (!GIsSavingPackage)
 				{
-					UE_LOG(LogContentBrowser, Error, TEXT("Redirector Loop Found!"));
-					for ( const FSoftObjectPath& Redirector : VisitedRedirectors )
+					check(FPackageName::IsValidObjectPath(InObjectPathString));
+					UClass* FoundClass = FindObject<UClass>(nullptr, *InObjectPathString);
+					if (!FoundClass)
 					{
-						UE_LOG(LogContentBrowser, Error, TEXT("Redirector: %s"), *Redirector.ToString());
-					}
-
-					ObjectAssetData = FAssetData();
-					break;
-				}
-
-				VisitedRedirectors.Add(ObjectAssetData.ToSoftObjectPath());
+						// Use the linker to search for class name redirects (from the loaded ActiveClassRedirects)
+						const FString NewClassName = FLinkerLoad::FindNewPathNameForClass(InObjectPathString, false);
 
-				// Get the destination object from the meta-data rather than load the redirector object, as 
-				// loading a redirector will also load the object it points to, which could cause a large hitch
-				FString DestinationObjectPath;
-				if (ObjectAssetData.GetTagValue("DestinationObject", DestinationObjectPath))
-				{
-					ConstructorHelpers::StripObjectClass(DestinationObjectPath);
-					ObjectAssetData = AssetRegistryModule.Get().GetAssetByObjectPath(FSoftObjectPath(DestinationObjectPath));
-				}
-				else
-				{
-					ObjectAssetData = FAssetData();
+						if (!NewClassName.IsEmpty())
+						{
+							check(FPackageName::IsValidObjectPath(NewClassName));
+							// Our new class name might be lacking the path, so try and find it so we can use the full path in the collection
+							FoundClass = FindObject<UClass>(nullptr, *NewClassName);
+							if (FoundClass)
+							{
+								OutNewObjectPath = *FoundClass->GetPathName();
+							}
+						}
+					}
 				}
 			}
-
-PRAGMA_DISABLE_DEPRECATION_WARNINGS
-			OutNewObjectPath = ObjectAssetData.ObjectPath;
-PRAGMA_ENABLE_DEPRECATION_WARNINGS
 		}
 
 		return OutNewObjectPath.IsValid() && InObjectPath != OutNewObjectPath;
 	}
 
 private:
-	FAssetRegistryModule& AssetRegistryModule;
+	IAssetRegistry* AssetRegistry;
 };
 
 FCollectionAssetRegistryBridge::FCollectionAssetRegistryBridge()
 {
 	// Load the asset registry module to listen for updates
-	FAssetRegistryModule& AssetRegistryModule = FModuleManager::LoadModuleChecked<FAssetRegistryModule>(TEXT("AssetRegistry"));
-	AssetRegistryModule.Get().OnAssetRemoved().AddRaw(this, &FCollectionAssetRegistryBridge::OnAssetRemoved);
-	AssetRegistryModule.Get().OnAssetRenamed().AddRaw(this, &FCollectionAssetRegistryBridge::OnAssetRenamed);
+	IAssetRegistry& AssetRegistry = IAssetRegistry::GetChecked();
+	AssetRegistry.OnAssetRemoved().AddRaw(this, &FCollectionAssetRegistryBridge::OnAssetRemoved);
+	AssetRegistry.OnAssetRenamed().AddRaw(this, &FCollectionAssetRegistryBridge::OnAssetRenamed);
 	
-	if (AssetRegistryModule.Get().IsLoadingAssets())
+	if (AssetRegistry.IsLoadingAssets())
 	{
-		AssetRegistryModule.Get().OnFilesLoaded().AddRaw(this, &FCollectionAssetRegistryBridge::OnAssetRegistryLoadComplete);
+		AssetRegistry.OnFilesLoaded().AddRaw(this, &FCollectionAssetRegistryBridge::OnAssetRegistryLoadComplete);
 	}
 	else
 	{
@@ -162,12 +132,11 @@ FCollectionAssetRegistryBridge::FCollectionAssetRegistryBridge()
 FCollectionAssetRegistryBridge::~FCollectionAssetRegistryBridge()
 {
 	// Load the asset registry module to unregister delegates
-	if (FModuleManager::Get().IsModuleLoaded("AssetRegistry"))
+	if (IAssetRegistry* AssetRegistry = IAssetRegistry::Get())
 	{
-		FAssetRegistryModule& AssetRegistryModule = FModuleManager::GetModuleChecked<FAssetRegistryModule>(TEXT("AssetRegistry"));
-		AssetRegistryModule.Get().OnAssetRemoved().RemoveAll(this);
-		AssetRegistryModule.Get().OnAssetRenamed().RemoveAll(this);
-		AssetRegistryModule.Get().OnFilesLoaded().RemoveAll(this);
+		AssetRegistry->OnAssetRemoved().RemoveAll(this);
+		AssetRegistry->OnAssetRenamed().RemoveAll(this);
+		AssetRegistry->OnFilesLoaded().RemoveAll(this);
 	}
 }
 
