Fixed a crash when exiting the editor via the QuitEditor function.

The slate windows get close after the subsystems are deinitilaized so we need to be carefull with those in the content browser widget code. Some binding can still be queried while processing some defered windows event during the closing of the slate application.

diff --git a/Engine/Source/Editor/ContentBrowser/Private/SContentBrowser.cpp b/Engine/Source/Editor/ContentBrowser/Private/SContentBrowser.cpp
--- a/Engine/Source/Editor/ContentBrowser/Private/SContentBrowser.cpp
+++ b/Engine/Source/Editor/ContentBrowser/Private/SContentBrowser.cpp
@@ -2465,25 +2465,27 @@ bool SContentBrowser::OnHasCrumbDelimiterContent(const FString& CrumbData) const
 	}
 	else if (SourcesData.HasVirtualPaths())
 	{
-		UContentBrowserDataSubsystem* ContentBrowserData = IContentBrowserDataModule::Get().GetSubsystem();
-
-		FContentBrowserDataFilter SubItemsFilter;
-		SubItemsFilter.ItemTypeFilter = EContentBrowserItemTypeFilter::IncludeFolders;
-		SubItemsFilter.bRecursivePaths = false;
-		SubItemsFilter.ItemCategoryFilter = PathViewPtr->GetContentBrowserItemCategoryFilter();
-		SubItemsFilter.ItemAttributeFilter = PathViewPtr->GetContentBrowserItemAttributeFilter();
-
-		bool bHasSubItems = false;
-		if (ContentBrowser::bCrumbsEnumerate)
+		// The subsystems may get Deinitialize before the slate windows are close when exiting the editor
+		if (UContentBrowserDataSubsystem* ContentBrowserData = IContentBrowserDataModule::Get().GetSubsystem())
 		{
-			ContentBrowserData->EnumerateItemsUnderPath(*CrumbData, SubItemsFilter, [&bHasSubItems](FContentBrowserItemData&& InSubItem)
+			FContentBrowserDataFilter SubItemsFilter;
+			SubItemsFilter.ItemTypeFilter = EContentBrowserItemTypeFilter::IncludeFolders;
+			SubItemsFilter.bRecursivePaths = false;
+			SubItemsFilter.ItemCategoryFilter = PathViewPtr->GetContentBrowserItemCategoryFilter();
+			SubItemsFilter.ItemAttributeFilter = PathViewPtr->GetContentBrowserItemAttributeFilter();
+
+			bool bHasSubItems = false;
+			if (ContentBrowser::bCrumbsEnumerate)
 			{
-				bHasSubItems = true;
+				ContentBrowserData->EnumerateItemsUnderPath(*CrumbData, SubItemsFilter, [&bHasSubItems](FContentBrowserItemData&& InSubItem)
+					{
+						bHasSubItems = true;
 				return false;
-			});
-		}
+					});
+			}
 
-		return bHasSubItems;
+			return bHasSubItems;
+		}
 	}
 
 	return false;
