Collection workflow telemetry use TelemetryRouter infrastructure
* Use FTelemetryRouter framework to track collection workflow telemetry instead of bespoke delegates on ContentBrowserModule

diff --git a/Engine/Source/Editor/ContentBrowser/Private/AssetContextMenu.cpp b/Engine/Source/Editor/ContentBrowser/Private/AssetContextMenu.cpp
--- a/Engine/Source/Editor/ContentBrowser/Private/AssetContextMenu.cpp
+++ b/Engine/Source/Editor/ContentBrowser/Private/AssetContextMenu.cpp
@@ -63,6 +63,7 @@
 #include "UObject/WeakObjectPtr.h"
 #include "UObject/WeakObjectPtrTemplates.h"
 #include "AssetDefinitionRegistry.h"
+#include "TelemetryRouter.h"
 #include "AssetRegistry/AssetRegistryHelpers.h"
 #include "Engine/AssetManager.h"
 #include "Misc/WarnIfAssetsLoadedInScope.h"
@@ -810,30 +811,26 @@ bool FAssetContextMenu::AddCollectionMenuOptions(UToolMenu* Menu)
 			}
 
 			const double DurationSec = FPlatformTime::Seconds() - BeginTimeSec;
-
-			FContentBrowserModule& ContentBrowserModule = FModuleManager::GetModuleChecked<FContentBrowserModule>( TEXT("ContentBrowser") );
-			if (ContentBrowserModule.GetAssetCollectionTelemetryDelegate().IsBound())
+			
 			{
-				FCollectionTelemetryEvent Event;
 				if (RemoveFromCollection)
 				{
-					FAssetRemovedTelemetryEvent AssetRemoved;
+					FAssetRemovedFromCollectionTelemetryEvent AssetRemoved;
 					AssetRemoved.DurationSec = DurationSec;
 					AssetRemoved.NumRemoved = ObjectCount;
 					AssetRemoved.CollectionShareType = InCollectionKey.Type;
 					AssetRemoved.Workflow = ECollectionTelemetryAssetRemovedWorkflow::ContextMenu;
-					Event = MakeTelemetryEvent(AssetRemoved);
+					FTelemetryRouter::Get().ProvideTelemetry(AssetRemoved);
 				}
 				else
 				{
-					FAssetAddedTelemetryEvent AssetAdded;
+					FAssetAddedToCollectionTelemetryEvent AssetAdded;
 					AssetAdded.DurationSec = DurationSec;
 					AssetAdded.NumAdded = ObjectCount;
 					AssetAdded.CollectionShareType = InCollectionKey.Type;
 					AssetAdded.Workflow = ECollectionTelemetryAssetAddedWorkflow::ContextMenu;
-					Event = MakeTelemetryEvent(AssetAdded);
+					FTelemetryRouter::Get().ProvideTelemetry(AssetAdded);
 				}
-				ContentBrowserModule.GetAssetCollectionTelemetryDelegate().Execute(Event);
 			}
 		}
 	};
@@ -1217,8 +1214,17 @@ void FAssetContextMenu::ExecuteRemoveFromCollection()
 			FCollectionManagerModule& CollectionManagerModule = FCollectionManagerModule::GetModule();
 
 			const FCollectionNameType& Collection = SourcesData.Collections[0];
+			const double BeginTimeSec = FPlatformTime::Seconds();
 			CollectionManagerModule.Get().RemoveFromCollection(Collection.Name, Collection.Type, SelectedItemCollectionIds);
+			const double DurationSec = FPlatformTime::Seconds() - BeginTimeSec;
 			OnAssetViewRefreshRequested.ExecuteIfBound();
+
+			FAssetRemovedFromCollectionTelemetryEvent AssetRemoved;
+			AssetRemoved.DurationSec = DurationSec;
+			AssetRemoved.NumRemoved = SelectedItemCollectionIds.Num();
+			AssetRemoved.CollectionShareType = Collection.Type;
+			AssetRemoved.Workflow = ECollectionTelemetryAssetRemovedWorkflow::ContextMenu;
+			FTelemetryRouter::Get().ProvideTelemetry(AssetRemoved);
 		}
 	}
 }
diff --git a/Engine/Source/Editor/ContentBrowser/Private/CollectionContextMenu.cpp b/Engine/Source/Editor/ContentBrowser/Private/CollectionContextMenu.cpp
--- a/Engine/Source/Editor/ContentBrowser/Private/CollectionContextMenu.cpp
+++ b/Engine/Source/Editor/ContentBrowser/Private/CollectionContextMenu.cpp
@@ -33,6 +33,7 @@
 #include "Misc/Paths.h"
 #include "Modules/ModuleManager.h"
 #include "SlotBase.h"
+#include "TelemetryRouter.h"
 #include "Styling/AppStyle.h"
 #include "Textures/SlateIcon.h"
 #include "UObject/NameTypes.h"
@@ -457,14 +458,10 @@ void FCollectionContextMenu::ExecuteNewCollection(ECollectionShareType::Type Col
 
 	// Telemetry Event
 	{
-		FContentBrowserModule& ContentBrowserModule = FModuleManager::GetModuleChecked<FContentBrowserModule>( TEXT("ContentBrowser") );
-		if (ContentBrowserModule.GetAssetCollectionTelemetryDelegate().IsBound())
-		{
-			FCollectionCreatedTelemetryEvent AssetAdded;
-			AssetAdded.DurationSec = FPlatformTime::Seconds() - BeginTimeSec;
-			AssetAdded.CollectionShareType = CollectionType;
-			ContentBrowserModule.GetAssetCollectionTelemetryDelegate().Execute(MakeTelemetryEvent(AssetAdded));
-		}
+		FCollectionCreatedTelemetryEvent AssetAdded;
+		AssetAdded.DurationSec = FPlatformTime::Seconds() - BeginTimeSec;
+		AssetAdded.CollectionShareType = CollectionType;
+		FTelemetryRouter::Get().ProvideTelemetry(AssetAdded);
 	}
 }
 
@@ -606,13 +603,11 @@ FReply FCollectionContextMenu::ExecuteDestroyCollectionConfirmed(TArray<TSharedP
 	
 	CollectionView.Pin()->DeleteCollectionItems(CollectionList);
 
-	FContentBrowserModule& ContentBrowserModule = FModuleManager::GetModuleChecked<FContentBrowserModule>( TEXT("ContentBrowser") );
-	if (ContentBrowserModule.GetAssetCollectionTelemetryDelegate().IsBound())
 	{
 		FCollectionsDeletedTelemetryEvent CollectionDeleted;
 		CollectionDeleted.DurationSec = FPlatformTime::Seconds() - BeginEventSec;
 		CollectionDeleted.CollectionsDeleted = CollectionList.Num();
-		ContentBrowserModule.GetAssetCollectionTelemetryDelegate().Execute(MakeTelemetryEvent(CollectionDeleted));
+		FTelemetryRouter::Get().ProvideTelemetry(CollectionDeleted);
 	}
 	
 	return FReply::Handled();
diff --git a/Engine/Source/Editor/ContentBrowser/Private/SCollectionView.cpp b/Engine/Source/Editor/ContentBrowser/Private/SCollectionView.cpp
--- a/Engine/Source/Editor/ContentBrowser/Private/SCollectionView.cpp
+++ b/Engine/Source/Editor/ContentBrowser/Private/SCollectionView.cpp
@@ -49,6 +49,7 @@
 #include "SourcesData.h"
 #include "SourcesSearch.h"
 #include "SourcesViewWidgets.h"
+#include "TelemetryRouter.h"
 #include "Styling/AppStyle.h"
 #include "Styling/ISlateStyle.h"
 #include "Styling/SlateColor.h"
@@ -1368,15 +1369,13 @@ PRAGMA_ENABLE_DEPRECATION_WARNINGS
 
 			const double DurationSec = FPlatformTime::Seconds() - BeginTimeSec;
 
-			FContentBrowserModule& ContentBrowserModule = FModuleManager::GetModuleChecked<FContentBrowserModule>( TEXT("ContentBrowser") );
-			if (ContentBrowserModule.GetAssetCollectionTelemetryDelegate().IsBound())
 			{
-				FAssetAddedTelemetryEvent AssetAdded;
+				FAssetAddedToCollectionTelemetryEvent AssetAdded;
 				AssetAdded.DurationSec = DurationSec;
 				AssetAdded.NumAdded = NumAdded;
 				AssetAdded.CollectionShareType = CollectionItem->CollectionType;
 				AssetAdded.Workflow = ECollectionTelemetryAssetAddedWorkflow::DragAndDrop;
-				ContentBrowserModule.GetAssetCollectionTelemetryDelegate().Execute(MakeTelemetryEvent(AssetAdded));
+				FTelemetryRouter::Get().ProvideTelemetry(AssetAdded);
 			}
 		}
 		else
@@ -1545,14 +1544,12 @@ bool SCollectionView::CollectionNameChangeCommit( const TSharedPtr< FCollectionI
 			CollectionItem->OnCollectionCreatedEvent.Execute(FCollectionNameType(NewNameFinal, CollectionItem->CollectionType));
 			CollectionItem->OnCollectionCreatedEvent.Unbind();
 		}
-
-		FContentBrowserModule& ContentBrowserModule = FModuleManager::GetModuleChecked<FContentBrowserModule>( TEXT("ContentBrowser") );
-		if (ContentBrowserModule.GetAssetCollectionTelemetryDelegate().IsBound())
+		
 		{
 			FCollectionCreatedTelemetryEvent CollectionCreatedEvent;
 			CollectionCreatedEvent.DurationSec = FPlatformTime::Seconds() - BeginTimeSec;
 			CollectionCreatedEvent.CollectionShareType = CollectionItem->CollectionType;
-			ContentBrowserModule.GetAssetCollectionTelemetryDelegate().Execute(MakeTelemetryEvent(CollectionCreatedEvent));
+			FTelemetryRouter::Get().ProvideTelemetry(CollectionCreatedEvent);
 		}
 	}
 	else
diff --git a/Engine/Source/Editor/ContentBrowser/Public/ContentBrowserModule.h b/Engine/Source/Editor/ContentBrowser/Public/ContentBrowserModule.h
--- a/Engine/Source/Editor/ContentBrowser/Public/ContentBrowserModule.h
+++ b/Engine/Source/Editor/ContentBrowser/Public/ContentBrowserModule.h
@@ -48,17 +48,20 @@ private:
 	friend class FContentBrowserModule;
 };
 
-struct FCollectionTelemetryEventBase {};
 // Workflow event when a collection is created
-struct FCollectionCreatedTelemetryEvent : FCollectionTelemetryEventBase
+struct FCollectionCreatedTelemetryEvent
 {
+	static inline constexpr FGuid TelemetryID = FGuid(0x2F9C8896, 0xCB2C402B, 0xB8D6DCF1, 0xE3F22D41);
+	
 	double DurationSec = 0.0;
 	ECollectionShareType::Type CollectionShareType = ECollectionShareType::CST_All;
 };
 
 // Workflow event when a set of collections are deleted
-struct FCollectionsDeletedTelemetryEvent : FCollectionTelemetryEventBase
+struct FCollectionsDeletedTelemetryEvent
 {
+	static inline constexpr FGuid TelemetryID = FGuid(0x1362ABC8, 0x6CDD4FF9, 0xB4B0E06C, 0x2CA418DC);
+	
 	double DurationSec = 0.0;
 	int32 CollectionsDeleted = 0;
 };
@@ -70,8 +73,10 @@ enum class ECollectionTelemetryAssetAddedWorkflow : int32
 };
 
 // Workflow event when assets are added to a collection
-struct FAssetAddedTelemetryEvent : FCollectionTelemetryEventBase
+struct FAssetAddedToCollectionTelemetryEvent
 {
+	static inline constexpr FGuid TelemetryID = FGuid(0x3676C84E, 0xFEB74E21, 0x99FF6FB8, 0x622431D2);
+	
 	double DurationSec;
 	ECollectionShareType::Type CollectionShareType;
 	uint32 NumAdded;
@@ -84,23 +89,16 @@ enum class ECollectionTelemetryAssetRemovedWorkflow : int32
 };
 
 // Workflow event when assets are removed from a collection
-struct FAssetRemovedTelemetryEvent : FCollectionTelemetryEventBase
+struct FAssetRemovedFromCollectionTelemetryEvent
 {
+	static inline constexpr FGuid TelemetryID = FGuid(0xF7660FA4, 0x744F44F5, 0x9B647690, 0x3C5689BD);
+	
 	double DurationSec;
 	ECollectionShareType::Type CollectionShareType;
 	uint32 NumRemoved;
 	ECollectionTelemetryAssetRemovedWorkflow Workflow;
 };
 
-typedef TVariant<FCollectionCreatedTelemetryEvent, FCollectionsDeletedTelemetryEvent, FAssetAddedTelemetryEvent, FAssetRemovedTelemetryEvent> FCollectionTelemetryEvent;
-
-// Helper to simplify creation of FTelemetryEventTaggedUnion
-template<typename T>
-FCollectionTelemetryEvent MakeTelemetryEvent(const T& Event)
-{
-	return FCollectionTelemetryEvent(TInPlaceType<T>(), Event);
-}
-
 /**
  * Content browser module
  */
@@ -127,8 +125,6 @@ public:
 	DECLARE_DELEGATE_OneParam( FDefaultSelectedPathsDelegate, TArray<FName>& /*VirtualPaths*/ );
 	/** */
 	DECLARE_DELEGATE_OneParam( FDefaultPathsToExpandDelegate, TArray<FName>& /*VirtualPaths*/ );
-	/** */
-	DECLARE_DELEGATE_OneParam( FAssetCollectionTelemetryDelegate, const FCollectionTelemetryEvent&);
 	
 	/**
 	 * Called right after the plugin DLL has been loaded and the plugin object has been created
@@ -191,8 +187,6 @@ public:
 	/** Override list of paths to expand by default */
 	FDefaultPathsToExpandDelegate& GetDefaultPathsToExpandDelegate() { return DefaultPathsToExpandDelegate; }
 
-	FAssetCollectionTelemetryDelegate& GetAssetCollectionTelemetryDelegate() { return AssetCollectionTelemetryDelegate; }
-
 	FMainMRUFavoritesList* GetRecentlyOpenedAssets() const;
 
 	static const FName NumberOfRecentAssetsName;
@@ -254,6 +248,4 @@ private:
 
 	FDefaultSelectedPathsDelegate DefaultSelectedPathsDelegate;
 	FDefaultPathsToExpandDelegate DefaultPathsToExpandDelegate;
-	
-	FAssetCollectionTelemetryDelegate AssetCollectionTelemetryDelegate;
 };
