Fixed an issue where the customization point for the custom text filters would only use the first handler when multiple are registred.

diff --git a/Engine/Source/Editor/ContentBrowser/Private/TextFilterValueHandlers.cpp b/Engine/Source/Editor/ContentBrowser/Private/TextFilterValueHandlers.cpp
--- a/Engine/Source/Editor/ContentBrowser/Private/TextFilterValueHandlers.cpp
+++ b/Engine/Source/Editor/ContentBrowser/Private/TextFilterValueHandlers.cpp
@@ -14,7 +14,10 @@ bool UTextFilterValueHandlers::HandleTextFilterValue(const FContentBrowserItem&
 	{
 		if (UClass* TextFilterValueHandlerClass = HandlerClass.LoadSynchronous())
 		{
-			return GetDefault<UTextFilterValueHandler>(TextFilterValueHandlerClass)->HandleTextFilterValue(InContentBrowserItem, InValue, InTextComparisonMode, bOutIsMatch);
+			if (GetDefault<UTextFilterValueHandler>(TextFilterValueHandlerClass)->HandleTextFilterValue(InContentBrowserItem, InValue, InTextComparisonMode, bOutIsMatch))
+			{
+				return true;
+			}
 		}
 	}
 
