SPathPicker: Make sure DefaultPath is not a blocked path

#rnx

diff --git a/Engine/Source/Editor/ContentBrowser/Private/SPathPicker.cpp b/Engine/Source/Editor/ContentBrowser/Private/SPathPicker.cpp
--- a/Engine/Source/Editor/ContentBrowser/Private/SPathPicker.cpp
+++ b/Engine/Source/Editor/ContentBrowser/Private/SPathPicker.cpp
@@ -58,21 +58,16 @@ void SPathPicker::Construct( const FArguments& InArgs )
 	];
 
 	const FString& DefaultPath = InArgs._PathPickerConfig.DefaultPath;
-	if ( !DefaultPath.IsEmpty() )
+	if ( !DefaultPath.IsEmpty() && PathViewPtr->InternalPathPassesBlockLists(*DefaultPath))
 	{
 		const FName VirtualPath = IContentBrowserDataModule::Get().GetSubsystem()->ConvertInternalPathToVirtual(*DefaultPath);
-		if (InArgs._PathPickerConfig.bAddDefaultPath)
+		if (InArgs._PathPickerConfig.bAddDefaultPath && !PathViewPtr->FindTreeItem(VirtualPath))
 		{
-			if (!PathViewPtr->FindTreeItem(VirtualPath))
-			{
-				const FString DefaultPathLeafName = FPaths::GetPathLeaf(VirtualPath.ToString());
-				PathViewPtr->AddFolderItem(FContentBrowserItemData(nullptr, EContentBrowserItemFlags::Type_Folder, VirtualPath, *DefaultPathLeafName, FText(), nullptr), /*bUserNamed*/false);
-			}
+			const FString DefaultPathLeafName = FPaths::GetPathLeaf(VirtualPath.ToString());
+			PathViewPtr->AddFolderItem(FContentBrowserItemData(nullptr, EContentBrowserItemFlags::Type_Folder, VirtualPath, *DefaultPathLeafName, FText(), nullptr), /*bUserNamed*/false);
 		}
 
-		TArray<FString> SelectedPaths;
-		SelectedPaths.Add(VirtualPath.ToString());
-		PathViewPtr->SetSelectedPaths(SelectedPaths);
+		PathViewPtr->SetSelectedPaths({ VirtualPath.ToString() });
 	}
 }
 
