Updated the asset view filter to process every filter looking for an exclusion, rather than assuming the first inclusion means the rest of the filters can be skipped

#rnx

diff --git a/Engine/Source/Editor/ContentBrowser/Private/SAssetView.cpp b/Engine/Source/Editor/ContentBrowser/Private/SAssetView.cpp
--- a/Engine/Source/Editor/ContentBrowser/Private/SAssetView.cpp
+++ b/Engine/Source/Editor/ContentBrowser/Private/SAssetView.cpp
@@ -122,25 +122,23 @@ public:
 			return true;
 		}
 
-		if (AssetView->OnShouldFilterItem.IsBound()
-			&& !AssetView->OnShouldFilterItem.Execute(InItemToFilter->GetItem()))
+		if (AssetView->OnShouldFilterItem.IsBound() && AssetView->OnShouldFilterItem.Execute(InItemToFilter->GetItem()))
 		{
-			return true;
+			return false;
 		}
 
 		// If we have OnShouldFilterAsset then it is assumed that we really only want to see true assets and 
 		// nothing else so only include things that have asset data and also pass the query filter
-		FAssetData ItemAssetData;
-		if (AssetView->OnShouldFilterAsset.IsBound()
-			&& InItemToFilter->GetItem().Legacy_TryGetAssetData(ItemAssetData))
+		if (AssetView->OnShouldFilterAsset.IsBound())
 		{
-			if (!AssetView->OnShouldFilterAsset.Execute(ItemAssetData))
+			FAssetData ItemAssetData;
+			if (!InItemToFilter->GetItem().Legacy_TryGetAssetData(ItemAssetData) || AssetView->OnShouldFilterAsset.Execute(ItemAssetData))
 			{
-				return true;
+				return false;
 			}
 		}
 
-		return false;
+		return true;
 	}
 
 	bool DoesItemPassFrontendFilter(const TSharedPtr<FAssetViewItem>& InItemToFilter)
@@ -152,12 +150,12 @@ public:
 		}
 
 		// Run the item through the filters
-		if (!AssetView->IsFrontendFilterActive() || AssetView->PassesCurrentFrontendFilter(InItemToFilter->GetItem()))
+		if (AssetView->IsFrontendFilterActive() && !AssetView->PassesCurrentFrontendFilter(InItemToFilter->GetItem()))
 		{
-			return true;
+			return false;
 		}
 
-		return false;
+		return true;
 	}
 
 private:
