Removed an uneeded ensure as it was triggered in an expected scenario.

#rnx

diff --git a/Engine/Source/Editor/ContentBrowser/Private/AssetContextMenu.cpp b/Engine/Source/Editor/ContentBrowser/Private/AssetContextMenu.cpp
--- a/Engine/Source/Editor/ContentBrowser/Private/AssetContextMenu.cpp
+++ b/Engine/Source/Editor/ContentBrowser/Private/AssetContextMenu.cpp
@@ -254,7 +254,8 @@ TSharedRef<SWidget> FAssetContextMenu::MakeContextMenu(TArrayView<const FContent
 				}
 			}
 
-			if (ensure(CommonClass))
+			// We can have a null common class if an asset is from unloaded plugin or an missing class.
+			if (CommonClass)
 			{
 				MenuName = UToolMenus::JoinMenuPaths(BaseMenuName, CommonClass->GetFName());
 
