Asset View: Ensure there is enough space for the name of the asset when source control icons are enabled.
A constant height is added to the total asset tile height to make space for both the name and the icon.

diff --git a/Engine/Source/Editor/ContentBrowser/Private/AssetViewWidgets.cpp b/Engine/Source/Editor/ContentBrowser/Private/AssetViewWidgets.cpp
--- a/Engine/Source/Editor/ContentBrowser/Private/AssetViewWidgets.cpp
+++ b/Engine/Source/Editor/ContentBrowser/Private/AssetViewWidgets.cpp
@@ -2225,7 +2225,7 @@ int32 SAssetTileItem::GetGenericThumbnailSize() const
 EVisibility SAssetTileItem::GetSCCIconVisibility() const
 {
 	// Hide the scc state icon when there is no brush or in tiny size since there isn't enough space
-	return bHasCCStateBrush &&  CurrentThumbnailSize.Get() != EThumbnailSize::Tiny ? EVisibility::Visible : EVisibility::Collapsed;
+	return bHasCCStateBrush &&  CurrentThumbnailSize.Get() != EThumbnailSize::Tiny && ISourceControlModule::Get().IsEnabled() && ISourceControlModule::Get().GetProvider().IsAvailable() ? EVisibility::Visible : EVisibility::Collapsed;
 }
 
 
diff --git a/Engine/Source/Editor/ContentBrowser/Private/SAssetView.cpp b/Engine/Source/Editor/ContentBrowser/Private/SAssetView.cpp
--- a/Engine/Source/Editor/ContentBrowser/Private/SAssetView.cpp
+++ b/Engine/Source/Editor/ContentBrowser/Private/SAssetView.cpp
@@ -4427,6 +4427,11 @@ float SAssetView::GetTileViewTypeNameHeight() const
 	return TypeNameHeight;
 }
 
+float SAssetView::GetSourceControlIconHeight() const
+{
+	return (float)(ThumbnailSize != EThumbnailSize::Tiny && ISourceControlModule::Get().IsEnabled() && ISourceControlModule::Get().GetProvider().IsAvailable() && !bShowTypeInTileView ? 17.0 : 0.0);
+}
+
 float SAssetView::GetListViewItemHeight() const
 {
 	return (float)(ListViewThumbnailSize + ListViewThumbnailPadding * 2) * FMath::Lerp(MinThumbnailScale, MaxThumbnailScale, GetThumbnailScale());
@@ -4434,7 +4439,7 @@ float SAssetView::GetListViewItemHeight() const
 
 float SAssetView::GetTileViewItemHeight() const
 {
-	return (((float)TileViewNameHeight + GetTileViewTypeNameHeight()) * FMath::Lerp(MinThumbnailScale, MaxThumbnailScale, GetThumbnailScale())) + GetTileViewItemBaseHeight() * FillScale;
+	return (((float)TileViewNameHeight + GetTileViewTypeNameHeight()) * FMath::Lerp(MinThumbnailScale, MaxThumbnailScale, GetThumbnailScale())) + GetTileViewItemBaseHeight() * FillScale + GetSourceControlIconHeight();
 }
 
 float SAssetView::GetTileViewItemBaseHeight() const
diff --git a/Engine/Source/Editor/ContentBrowser/Public/SAssetView.h b/Engine/Source/Editor/ContentBrowser/Public/SAssetView.h
--- a/Engine/Source/Editor/ContentBrowser/Public/SAssetView.h
+++ b/Engine/Source/Editor/ContentBrowser/Public/SAssetView.h
@@ -773,6 +773,9 @@ private:
 	/** Gets the spacing dedicated to support type names */
 	float GetTileViewTypeNameHeight() const;
 
+	/** Gets the spacing needed to support source control icons. No scaling is required for this since the icon itself is not scaled by the thumbnail size. */
+	float GetSourceControlIconHeight() const;
+
 	/** Gets the scaled item height for the list view */
 	float GetListViewItemHeight() const;
 	
