Added the possibility to copy content browser selected paths to the clipboard (mapped to Ctrl+C).

diff --git a/Engine/Source/Editor/ContentBrowser/Private/AssetViewWidgets.cpp b/Engine/Source/Editor/ContentBrowser/Private/AssetViewWidgets.cpp
--- a/Engine/Source/Editor/ContentBrowser/Private/AssetViewWidgets.cpp
+++ b/Engine/Source/Editor/ContentBrowser/Private/AssetViewWidgets.cpp
@@ -27,6 +27,7 @@
 #include "Framework/Application/SlateApplication.h"
 #include "Framework/Text/TextLayout.h"
 #include "HAL/IConsoleManager.h"
+#include "HAL/PlatformApplicationMisc.h"
 #include "ICollectionManager.h"
 #include "ISourceControlModule.h"
 #include "ISourceControlProvider.h"
@@ -96,17 +97,39 @@ FReply FAssetViewModeUtils::OnViewModeKeyDown( const TSet< TSharedPtr<FAssetView
 		)
 	{
 		TArray<FContentBrowserItem> SelectedFiles;
+		TArray<FContentBrowserItem> SelectedFolders;
 		for (const TSharedPtr<FAssetViewItem>& SelectedItem : SelectedItems)
 		{
 			if (SelectedItem->GetItem().IsFile())
 			{
 				SelectedFiles.Add(SelectedItem->GetItem());
 			}
+			else if (SelectedItem->GetItem().IsFolder())
+			{
+				SelectedFolders.Add(SelectedItem->GetItem());
+			}
 		}
 
+		FString ClipboardText;
+
 		if (SelectedFiles.Num() > 0)
 		{
-			ContentBrowserUtils::CopyItemReferencesToClipboard(SelectedFiles);
+			ClipboardText += ContentBrowserUtils::GetItemReferencesText(SelectedFiles);
+		}
+
+		if (SelectedFolders.Num() > 0)
+		{
+			if (!ClipboardText.IsEmpty())
+			{
+				ClipboardText += LINE_TERMINATOR;
+			}
+
+			ClipboardText += ContentBrowserUtils::GetFolderReferencesText(SelectedFolders);
+		}
+
+		if (!ClipboardText.IsEmpty())
+		{
+			FPlatformApplicationMisc::ClipboardCopy(*ClipboardText);
 		}
 
 		return FReply::Handled();
diff --git a/Engine/Source/Editor/ContentBrowser/Private/ContentBrowserCommands.cpp b/Engine/Source/Editor/ContentBrowser/Private/ContentBrowserCommands.cpp
--- a/Engine/Source/Editor/ContentBrowser/Private/ContentBrowserCommands.cpp
+++ b/Engine/Source/Editor/ContentBrowser/Private/ContentBrowserCommands.cpp
@@ -19,6 +19,7 @@ void FContentBrowserCommands::RegisterCommands()
 	UI_COMMAND(SaveSelectedAsset, "Save Selected Asset", "Save the selected asset", EUserInterfaceActionType::Button, FInputChord(EModifierKey::Control, EKeys::S));
 	UI_COMMAND(SaveAllCurrentFolder, "Save All", "Save All in current folder", EUserInterfaceActionType::Button, FInputChord(EModifierKey::Control | EModifierKey::Shift, EKeys::S));
 	UI_COMMAND(ResaveAllCurrentFolder, "Resave All", "Resave all assets contained in the current folder", EUserInterfaceActionType::Button, FInputChord());
+	UI_COMMAND(CopySelectedAssetPath, "Copy Selected Asset Path", "Copy the selected asset path", EUserInterfaceActionType::Button, FInputChord(EModifierKey::Control, EKeys::C));
 }
 
 
diff --git a/Engine/Source/Editor/ContentBrowser/Private/ContentBrowserCommands.h b/Engine/Source/Editor/ContentBrowser/Private/ContentBrowserCommands.h
--- a/Engine/Source/Editor/ContentBrowser/Private/ContentBrowserCommands.h
+++ b/Engine/Source/Editor/ContentBrowser/Private/ContentBrowserCommands.h
@@ -36,4 +36,5 @@ public:
 	TSharedPtr<FUICommandInfo> SaveSelectedAsset;
 	TSharedPtr<FUICommandInfo> SaveAllCurrentFolder;
 	TSharedPtr<FUICommandInfo> ResaveAllCurrentFolder;
+	TSharedPtr<FUICommandInfo> CopySelectedAssetPath;
 };
diff --git a/Engine/Source/Editor/ContentBrowser/Private/ContentBrowserUtils.cpp b/Engine/Source/Editor/ContentBrowser/Private/ContentBrowserUtils.cpp
--- a/Engine/Source/Editor/ContentBrowser/Private/ContentBrowserUtils.cpp
+++ b/Engine/Source/Editor/ContentBrowser/Private/ContentBrowserUtils.cpp
@@ -312,21 +312,55 @@ void ContentBrowserUtils::DisplayConfirmationPopup(const FText& Message, const F
 	Popup->OpenPopup(ParentContent);
 }
 
-void ContentBrowserUtils::CopyItemReferencesToClipboard(const TArray<FContentBrowserItem>& ItemsToCopy)
+FString ContentBrowserUtils::GetItemReferencesText(const TArray<FContentBrowserItem>& Items)
 {
-	TArray<FContentBrowserItem> SortedItems = ItemsToCopy;
+	TArray<FContentBrowserItem> SortedItems = Items;
 	SortedItems.Sort([](const FContentBrowserItem& One, const FContentBrowserItem& Two)
 	{
 		return One.GetVirtualPath().Compare(Two.GetVirtualPath()) < 0;
 	});
 
-	FString ClipboardText;
+	FString Result;
 	for (const FContentBrowserItem& Item : SortedItems)
 	{
-		Item.AppendItemReference(ClipboardText);
+		if (ensure(!Item.IsFolder()))
+		{
+			Item.AppendItemReference(Result);
+		}
 	}
 
-	FPlatformApplicationMisc::ClipboardCopy(*ClipboardText);
+	return Result;
+}
+
+FString ContentBrowserUtils::GetFolderReferencesText(const TArray<FContentBrowserItem>& Folders)
+{
+	TArray<FContentBrowserItem> SortedItems = Folders;
+	SortedItems.Sort([](const FContentBrowserItem& One, const FContentBrowserItem& Two)
+	{
+		return One.GetVirtualPath().Compare(Two.GetVirtualPath()) < 0;
+	});
+
+	FString Result;
+	for (const FContentBrowserItem& Item : SortedItems)
+	{
+		if (ensure(Item.IsFolder()))
+		{
+			Result += Item.GetVirtualPath().ToString();
+			Result += LINE_TERMINATOR;
+		}
+	}
+
+	return Result;
+}
+
+void ContentBrowserUtils::CopyItemReferencesToClipboard(const TArray<FContentBrowserItem>& ItemsToCopy)
+{
+	FPlatformApplicationMisc::ClipboardCopy(*GetItemReferencesText(ItemsToCopy));
+}
+
+void ContentBrowserUtils::CopyFolderReferencesToClipboard(const TArray<FContentBrowserItem>& FoldersToCopy)
+{
+	FPlatformApplicationMisc::ClipboardCopy(*GetFolderReferencesText(FoldersToCopy));
 }
 
 void ContentBrowserUtils::CopyFilePathsToClipboard(const TArray<FContentBrowserItem>& ItemsToCopy)
diff --git a/Engine/Source/Editor/ContentBrowser/Private/ContentBrowserUtils.h b/Engine/Source/Editor/ContentBrowser/Private/ContentBrowserUtils.h
--- a/Engine/Source/Editor/ContentBrowser/Private/ContentBrowserUtils.h
+++ b/Engine/Source/Editor/ContentBrowser/Private/ContentBrowserUtils.h
@@ -38,9 +38,18 @@ namespace ContentBrowserUtils
 	/** Displays a modeless message asking yes or no type question */
 	void DisplayConfirmationPopup(const FText& Message, const FText& YesString, const FText& NoString, const TSharedRef<SWidget>& ParentContent, const FOnClicked& OnYesClicked, const FOnClicked& OnNoClicked = FOnClicked());
 
+	/** Returns references to the specified items */
+	FString GetItemReferencesText(const TArray<FContentBrowserItem>& Items);
+
+	/** Returns references to the specified folders */
+	FString GetFolderReferencesText(const TArray<FContentBrowserItem>& Folders);
+
 	/** Copies references to the specified items to the clipboard */
 	void CopyItemReferencesToClipboard(const TArray<FContentBrowserItem>& ItemsToCopy);
 
+	/** Copies references to the specified folders to the clipboard */
+	void CopyFolderReferencesToClipboard(const TArray<FContentBrowserItem>& FoldersToCopy);
+
 	/** Copies file paths on disk to the specified items to the clipboard */
 	void CopyFilePathsToClipboard(const TArray<FContentBrowserItem>& ItemsToCopy);
 
diff --git a/Engine/Source/Editor/ContentBrowser/Private/PathContextMenu.cpp b/Engine/Source/Editor/ContentBrowser/Private/PathContextMenu.cpp
--- a/Engine/Source/Editor/ContentBrowser/Private/PathContextMenu.cpp
+++ b/Engine/Source/Editor/ContentBrowser/Private/PathContextMenu.cpp
@@ -24,6 +24,7 @@
 #include "Framework/MultiBox/MultiBoxExtender.h"
 #include "Framework/SlateDelegates.h"
 #include "HAL/IConsoleManager.h"
+#include "HAL/PlatformApplicationMisc.h"
 #include "HAL/PlatformCrt.h"
 #include "IContentBrowserDataModule.h"
 #include "Math/Vector2D.h"
@@ -461,6 +462,11 @@ void FPathContextMenu::ExecuteResaveFolder()
 	SaveFilesWithinSelectedFolders(EContentBrowserItemSaveFlags::None);
 }
 
+void FPathContextMenu::CopySelectedFolder()
+{
+	CopySelectedFoldersToClipoard();
+}
+
 void FPathContextMenu::SaveFilesWithinSelectedFolders(EContentBrowserItemSaveFlags InSaveFlags)
 {
 	UContentBrowserDataSubsystem* ContentBrowserData = IContentBrowserDataModule::Get().GetSubsystem();
@@ -496,6 +502,19 @@ void FPathContextMenu::SaveFilesWithinSelectedFolders(EContentBrowserItemSaveFla
 	}
 }
 
+void FPathContextMenu::CopySelectedFoldersToClipoard()
+{
+	TStringBuilder<1024> StringBuilder;
+
+	for (const FContentBrowserItem& SelectedItem : SelectedFolders)
+	{
+		StringBuilder.Append(SelectedItem.GetVirtualPath().ToString());
+		StringBuilder.Append(LINE_TERMINATOR);
+	}
+
+	FPlatformApplicationMisc::ClipboardCopy(StringBuilder.ToString());
+}
+
 bool FPathContextMenu::CanExecuteDelete() const
 {
 	bool bCanDelete = false;
diff --git a/Engine/Source/Editor/ContentBrowser/Private/PathContextMenu.h b/Engine/Source/Editor/ContentBrowser/Private/PathContextMenu.h
--- a/Engine/Source/Editor/ContentBrowser/Private/PathContextMenu.h
+++ b/Engine/Source/Editor/ContentBrowser/Private/PathContextMenu.h
@@ -93,6 +93,9 @@ public:
 	/** Handler for when "Resave" is selected */
 	void ExecuteResaveFolder();
 
+	/** Handler for when "Copy AssetPath" is selected */
+	void CopySelectedFolder();
+
 	/** Handler for when "Delete" is selected and the delete was confirmed */
 	FReply ExecuteDeleteFolderConfirmed();
 
@@ -104,6 +107,8 @@ private:
 	FText GetDeleteToolTip() const;
 
 	void SaveFilesWithinSelectedFolders(EContentBrowserItemSaveFlags InSaveFlags);
+	
+	void CopySelectedFoldersToClipoard();
 
 	/** Checks to see if any of the selected paths use custom colors */
 	bool SelectedHasCustomColors() const;
diff --git a/Engine/Source/Editor/ContentBrowser/Private/SContentBrowser.cpp b/Engine/Source/Editor/ContentBrowser/Private/SContentBrowser.cpp
--- a/Engine/Source/Editor/ContentBrowser/Private/SContentBrowser.cpp
+++ b/Engine/Source/Editor/ContentBrowser/Private/SContentBrowser.cpp
@@ -663,6 +663,10 @@ void SContentBrowser::BindCommands()
 		FExecuteAction::CreateSP(this, &SContentBrowser::HandleResaveAllCurrentFolderCommand)
 	));
 
+	Commands->MapAction(FContentBrowserCommands::Get().CopySelectedAssetPath, FUIAction(
+		FExecuteAction::CreateSP(this, &SContentBrowser::CopySelectedAssetPathCommand)
+	));
+
 	// Allow extenders to add commands
 	FContentBrowserModule& ContentBrowserModule = FModuleManager::GetModuleChecked<FContentBrowserModule>(TEXT("ContentBrowser"));
 	TArray<FContentBrowserCommandExtender> CommmandExtenderDelegates = ContentBrowserModule.GetAllContentBrowserCommandExtenders();
@@ -3375,6 +3379,11 @@ void SContentBrowser::HandleResaveAllCurrentFolderCommand() const
 	PathContextMenu->ExecuteResaveFolder();
 }
 
+void SContentBrowser::CopySelectedAssetPathCommand() const
+{
+	PathContextMenu->CopySelectedFolder();
+}
+
 bool SContentBrowser::HandleDeleteCommandCanExecute() const
 {
 	// The order of these conditions are carefully crafted to match the logic of the context menu summoning, as this callback 
diff --git a/Engine/Source/Editor/ContentBrowser/Private/SContentBrowser.h b/Engine/Source/Editor/ContentBrowser/Private/SContentBrowser.h
--- a/Engine/Source/Editor/ContentBrowser/Private/SContentBrowser.h
+++ b/Engine/Source/Editor/ContentBrowser/Private/SContentBrowser.h
@@ -430,6 +430,9 @@ private:
 	/** Handler for Resave on a folder */
 	void HandleResaveAllCurrentFolderCommand() const;
 
+	/** Handler for Copy path on an asset*/
+	void CopySelectedAssetPathCommand() const;
+
 	/** Handler to check to see if a delete command is allowed */
 	bool HandleDeleteCommandCanExecute() const;
 
