Tag ContentBrowser's SAssetTileIItem with Asset's ItemName
* Enables automation driver to discover the widget using the path "<SContentBrowser>/<SAssetView>/MyAssetItemName"

diff --git a/Engine/Source/Editor/ContentBrowser/Private/SAssetView.cpp b/Engine/Source/Editor/ContentBrowser/Private/SAssetView.cpp
--- a/Engine/Source/Editor/ContentBrowser/Private/SAssetView.cpp
+++ b/Engine/Source/Editor/ContentBrowser/Private/SAssetView.cpp
@@ -3728,7 +3728,8 @@ TSharedRef<ITableRow> SAssetView::MakeTileViewWidget(TSharedPtr<FAssetViewItem>
 			.ShouldAllowToolTip(this, &SAssetView::ShouldAllowToolTips)
 			.HighlightText( HighlightedText )
 			.IsSelected(FIsSelected::CreateSP(TableRowWidget.Get(), &STableRow<TSharedPtr<FAssetViewItem>>::IsSelected))
-			.IsSelectedExclusively(FIsSelected::CreateSP(TableRowWidget.Get(), &STableRow<TSharedPtr<FAssetViewItem>>::IsSelectedExclusively));
+			.IsSelectedExclusively(FIsSelected::CreateSP(TableRowWidget.Get(), &STableRow<TSharedPtr<FAssetViewItem>>::IsSelectedExclusively))
+			.AddMetaData<FTagMetaData>(AssetItem->GetItem().GetItemName());
 
 		TableRowWidget->SetContent(Item);
 
@@ -3773,8 +3774,9 @@ TSharedRef<ITableRow> SAssetView::MakeTileViewWidget(TSharedPtr<FAssetViewItem>
 			.OnGetCustomAssetToolTip(OnGetCustomAssetToolTip)
 			.OnVisualizeAssetToolTip( OnVisualizeAssetToolTip )
 			.OnAssetToolTipClosing( OnAssetToolTipClosing )
-			.ShowType(bShowTypeInTileView);
-
+			.ShowType(bShowTypeInTileView)
+			.AddMetaData<FTagMetaData>(AssetItem->GetItem().GetItemName());
+		
 		TableRowWidget->SetContent(Item);
 
 		return TableRowWidget.ToSharedRef();
