Updated AssetPicker:
- Filter UI combo button is now created using helper method. This fixes the lack of the red dot overlay when filters are active & unifies it
- Filter UI combo button placement is now configurable. Maintains old behavior of being placed on the right to the search bar, but can also be placed on the left like the content browser now
- Other code was moved, but not changed

diff --git a/Engine/Source/Editor/ContentBrowser/Private/SAssetPicker.cpp b/Engine/Source/Editor/ContentBrowser/Private/SAssetPicker.cpp
--- a/Engine/Source/Editor/ContentBrowser/Private/SAssetPicker.cpp
+++ b/Engine/Source/Editor/ContentBrowser/Private/SAssetPicker.cpp
@@ -80,7 +80,12 @@ void SAssetPicker::Construct( const FArguments& InArgs )
 	OnFolderEnteredDelegate = InArgs._AssetPickerConfig.OnFolderEntered;
 	OnGetAssetContextMenu = InArgs._AssetPickerConfig.OnGetAssetContextMenu;
 	OnGetFolderContextMenu = InArgs._AssetPickerConfig.OnGetFolderContextMenu;
-
+	
+	// Break up the incoming filter into a sources data and backend filter.
+	CurrentSourcesData = FSourcesData(InArgs._AssetPickerConfig.Filter.PackagePaths, InArgs._AssetPickerConfig.Collections);
+	CurrentBackendFilter = InArgs._AssetPickerConfig.Filter;
+	CurrentBackendFilter.PackagePaths.Reset();
+	
 	FOnGetContentBrowserItemContextMenu OnGetItemContextMenu;
 	if (OnGetAssetContextMenu.IsBound() || OnGetFolderContextMenu.IsBound())
 	{
@@ -182,23 +187,32 @@ void SAssetPicker::Construct( const FArguments& InArgs )
 
 
 	if (InArgs._AssetPickerConfig.bAddFilterUI)
-	{
-		// Filter
-		HorizontalBox->AddSlot()
+	{		
+		// We create available classes here. These are used to hide away the type filters in the filter list that don't match this list of classes
+		TArray<UClass*> FilterClassList;
+		for(auto Iter = CurrentBackendFilter.ClassPaths.CreateIterator(); Iter; ++Iter)
+		{
+			FTopLevelAssetPath ClassName = (*Iter);
+			UClass* FilterClass = FindObject<UClass>(ClassName);
+			if(FilterClass)
+			{
+				FilterClassList.AddUnique(FilterClass);
+			}
+		}		
+		
+		SAssignNew(FilterListPtr, SFilterList)
+			.OnFilterChanged(this, &SAssetPicker::OnFilterChanged)
+			.FrontendFilters(FrontendFilters)
+			.InitialClassFilters(FilterClassList)
+			.FilterBarIdentifier(FName(SaveSettingsName))
+			.ExtraFrontendFilters(InArgs._AssetPickerConfig.ExtraFrontendFilters);
+		
+		FilterComboButtonPtr = StaticCastSharedRef<SComboButton>(SFilterList::MakeAddFilterButton(FilterListPtr.ToSharedRef()));
+		
+		HorizontalBox->InsertSlot(0)
 		.AutoWidth()
 		[
-			SAssignNew(FilterComboButtonPtr, SComboButton)
-			.ComboButtonStyle(&FAppStyle::Get().GetWidgetStyle<FComboButtonStyle>("SimpleComboButton"))
-			.ToolTipText( LOCTEXT( "AddFilterToolTip", "Add an asset filter." ) )
-			.OnGetMenuContent( this, &SAssetPicker::MakeAddFilterMenu )
-			.HasDownArrow( false )
-			.AddMetaData<FTagMetaData>(FTagMetaData(TEXT("ContentBrowserFiltersCombo")))
-			.ButtonContent()
-			[
-				SNew( SImage)
-				.ColorAndOpacity(FSlateColor::UseForeground())
-				.Image(FAppStyle::Get().GetBrush("Icons.Filter"))
-			]
+			FilterComboButtonPtr.ToSharedRef()
 		];
 	}
 		
@@ -242,35 +256,12 @@ void SAssetPicker::Construct( const FArguments& InArgs )
 	}
 
 	// Asset view
-	
-	// Break up the incoming filter into a sources data and backend filter.
-	CurrentSourcesData = FSourcesData(InArgs._AssetPickerConfig.Filter.PackagePaths, InArgs._AssetPickerConfig.Collections);
-	CurrentBackendFilter = InArgs._AssetPickerConfig.Filter;
-	CurrentBackendFilter.PackagePaths.Reset();
-
 	if (InArgs._AssetPickerConfig.bAddFilterUI)
 	{
-		// Filters
-		TArray<UClass*> FilterClassList;
-		for(auto Iter = CurrentBackendFilter.ClassPaths.CreateIterator(); Iter; ++Iter)
-		{
-			FTopLevelAssetPath ClassName = (*Iter);
-			UClass* FilterClass = FindObject<UClass>(ClassName);
-			if(FilterClass)
-			{
-				FilterClassList.AddUnique(FilterClass);
-			}
-		}
-
 		VerticalBox->AddSlot()
 		.AutoHeight()
 		[
-			SAssignNew(FilterListPtr, SFilterList)
-			.OnFilterChanged(this, &SAssetPicker::OnFilterChanged)
-			.FrontendFilters(FrontendFilters)
-			.InitialClassFilters(FilterClassList)
-			.FilterBarIdentifier(FName(SaveSettingsName))
-			.ExtraFrontendFilters(InArgs._AssetPickerConfig.ExtraFrontendFilters)
+			FilterListPtr.ToSharedRef()
 		];
 
 		// Use the 'other developer' filter from the filter list widget. 
