Add a parent window override parameter to asset dialog config
For context, this is needed for programs that can create asset dialogs without having a mainframe present.

diff --git a/Engine/Source/Editor/ContentBrowser/Private/ContentBrowserSingleton.cpp b/Engine/Source/Editor/ContentBrowser/Private/ContentBrowserSingleton.cpp
--- a/Engine/Source/Editor/ContentBrowser/Private/ContentBrowserSingleton.cpp
+++ b/Engine/Source/Editor/ContentBrowser/Private/ContentBrowserSingleton.cpp
@@ -663,17 +663,24 @@ void FContentBrowserSingleton::SharedCreateAssetDialogWindow(const TSharedRef<SA
 
 	DialogWindow->SetContent(AssetDialog);
 
-	IMainFrameModule& MainFrameModule = FModuleManager::LoadModuleChecked<IMainFrameModule>(TEXT("MainFrame"));
-	const TSharedPtr<SWindow>& MainFrameParentWindow = MainFrameModule.GetParentWindow();
-	if (MainFrameParentWindow.IsValid())
+	TSharedPtr<SWindow> ParentWindow = InConfig.WindowOverride;
+	if (!InConfig.WindowOverride)
+	{
+		IMainFrameModule& MainFrameModule = FModuleManager::LoadModuleChecked<IMainFrameModule>(TEXT("MainFrame"));
+		const TSharedPtr<SWindow>& MainFrameParentWindow = MainFrameModule.GetParentWindow();
+		ParentWindow = MainFrameModule.GetParentWindow();
+	}
+
+
+	if (ParentWindow.IsValid())
 	{
 		if (bModal)
 		{
-			FSlateApplication::Get().AddModalWindow(DialogWindow, MainFrameParentWindow.ToSharedRef());
+			FSlateApplication::Get().AddModalWindow(DialogWindow, ParentWindow.ToSharedRef());
 		}
 		else if (FGlobalTabmanager::Get()->GetRootWindow().IsValid())
 		{
-			FSlateApplication::Get().AddWindowAsNativeChild(DialogWindow, MainFrameParentWindow.ToSharedRef());
+			FSlateApplication::Get().AddWindowAsNativeChild(DialogWindow, ParentWindow.ToSharedRef());
 		}
 		else
 		{
diff --git a/Engine/Source/Editor/ContentBrowser/Public/IContentBrowserSingleton.h b/Engine/Source/Editor/ContentBrowser/Public/IContentBrowserSingleton.h
--- a/Engine/Source/Editor/ContentBrowser/Public/IContentBrowserSingleton.h
+++ b/Engine/Source/Editor/ContentBrowser/Public/IContentBrowserSingleton.h
@@ -421,6 +421,8 @@ struct FSharedAssetDialogConfig
 	TArray<FTopLevelAssetPath> AssetClassNames;
 	FVector2D WindowSizeOverride;
 	FOnPathSelected OnPathSelected;
+	/** When specified, this window will be used instead of the mainframe window. */
+	TSharedPtr<SWindow> WindowOverride;
 
 	virtual EAssetDialogType::Type GetDialogType() const = 0;
 
