Changed the save asset behavior in the content browser. It will now only save the modified assets and ignore the others.

Added an resave command and menu entry that save all the assets.

diff --git a/Engine/Plugins/Editor/ContentBrowser/ContentBrowserAssetDataSource/Source/ContentBrowserAssetDataSource/Private/ContentBrowserAssetDataCore.cpp b/Engine/Plugins/Editor/ContentBrowser/ContentBrowserAssetDataSource/Source/ContentBrowserAssetDataSource/Private/ContentBrowserAssetDataCore.cpp
--- a/Engine/Plugins/Editor/ContentBrowser/ContentBrowserAssetDataSource/Source/ContentBrowserAssetDataSource/Private/ContentBrowserAssetDataCore.cpp
+++ b/Engine/Plugins/Editor/ContentBrowser/ContentBrowserAssetDataSource/Source/ContentBrowserAssetDataSource/Private/ContentBrowserAssetDataCore.cpp
@@ -583,15 +583,24 @@ bool CanSaveAssetFileItem(IAssetTools* InAssetTools, const FContentBrowserAssetF
 		return false;
 	}
 
-	if (EnumHasAnyFlags(InSaveFlags, EContentBrowserItemSaveFlags::SaveOnlyIfLoaded))
+	if (EnumHasAnyFlags(InSaveFlags, EContentBrowserItemSaveFlags::SaveOnlyIfLoaded | EContentBrowserItemSaveFlags::SaveOnlyIfDirty))
 	{
 		// Can't save a package that hasn't been loaded
 		UPackage* Package = InAssetPayload.GetPackage(/*bTryRecacheIfNull*/true);
 		if (!Package)
 		{
-			ContentBrowserAssetData::SetOptionalErrorMessage(OutErrorMsg, LOCTEXT("Error_CannotSaveUnloadedAsset", "Cannot save unloaded asset"));
+			ContentBrowserAssetData::SetOptionalErrorMessage(OutErrorMsg, LOCTEXT("Error_CannotSaveUnloadedAsset", "Cannot save unloaded item"));
 			return false;
 		}
+		
+		if (EnumHasAnyFlags(InSaveFlags, EContentBrowserItemSaveFlags::SaveOnlyIfDirty))
+		{
+			if (!Package->IsDirty())
+			{
+				ContentBrowserAssetData::SetOptionalErrorMessage(OutErrorMsg, LOCTEXT("Error_CannotSaveNonDirtyAsset", "Cannot save an unmodified item"));
+				return false;
+			}
+		}
 	}
 
 	return true;
diff --git a/Engine/Source/Editor/ContentBrowser/Private/AssetContextMenu.cpp b/Engine/Source/Editor/ContentBrowser/Private/AssetContextMenu.cpp
--- a/Engine/Source/Editor/ContentBrowser/Private/AssetContextMenu.cpp
+++ b/Engine/Source/Editor/ContentBrowser/Private/AssetContextMenu.cpp
@@ -923,7 +923,7 @@ void FAssetContextMenu::ExecuteEditItems()
 
 void FAssetContextMenu::ExecuteSaveAsset()
 {
-	const EContentBrowserItemSaveFlags SaveFlags = EContentBrowserItemSaveFlags::SaveOnlyIfLoaded;
+	const EContentBrowserItemSaveFlags SaveFlags = EContentBrowserItemSaveFlags::None;
 
 	// Batch these by their data sources
 	TMap<UContentBrowserDataSource*, TArray<FContentBrowserItemData>> SourcesAndItems;
@@ -1281,7 +1281,7 @@ bool FAssetContextMenu::CanExecuteSaveAsset() const
 	bool bCanSave = false;
 	for (const FContentBrowserItem& SelectedItem : SelectedFiles)
 	{
-		bCanSave |= SelectedItem.CanSave(EContentBrowserItemSaveFlags::SaveOnlyIfLoaded);
+		bCanSave |= SelectedItem.CanSave(EContentBrowserItemSaveFlags::None);
 	}
 	return bCanSave;
 }
diff --git a/Engine/Source/Editor/ContentBrowser/Private/ContentBrowserCommands.cpp b/Engine/Source/Editor/ContentBrowser/Private/ContentBrowserCommands.cpp
--- a/Engine/Source/Editor/ContentBrowser/Private/ContentBrowserCommands.cpp
+++ b/Engine/Source/Editor/ContentBrowser/Private/ContentBrowserCommands.cpp
@@ -16,7 +16,7 @@ void FContentBrowserCommands::RegisterCommands()
 	UI_COMMAND(OpenAssetsOrFolders, "Open Assets or Folders", "Opens the selected assets or folders, depending on the selection", EUserInterfaceActionType::Button, FInputChord(EKeys::Enter));
 	UI_COMMAND(PreviewAssets, "Preview Assets", "Loads the selected assets and previews them if possible", EUserInterfaceActionType::Button, FInputChord(EKeys::SpaceBar));
 	UI_COMMAND(CreateNewFolder, "Create New Folder", "Creates new folder in selected path", EUserInterfaceActionType::Button, FInputChord(EModifierKey::Control | EModifierKey::Shift, EKeys::N));
-	UI_COMMAND(SaveSelectedAsset, "Save Selected Asset", "Save the selected asset", EUserInterfaceActionType::Button, FInputChord(EModifierKey::Control, EKeys::S));
+	UI_COMMAND(SaveSelectedAsset, "Save Selected Item", "Save the selected item", EUserInterfaceActionType::Button, FInputChord(EModifierKey::Control, EKeys::S));
 	UI_COMMAND(SaveAllCurrentFolder, "Save All", "Save All in current folder", EUserInterfaceActionType::Button, FInputChord(EModifierKey::Control | EModifierKey::Shift, EKeys::S));
 	UI_COMMAND(ResaveAllCurrentFolder, "Resave All", "Resave all assets contained in the current folder", EUserInterfaceActionType::Button, FInputChord());
 	UI_COMMAND(CopySelectedAssetPath, "Copy Selected Asset Path", "Copy the selected asset path", EUserInterfaceActionType::Button, FInputChord(EModifierKey::Control, EKeys::C));
diff --git a/Engine/Source/Editor/ContentBrowser/Private/SContentBrowser.h b/Engine/Source/Editor/ContentBrowser/Private/SContentBrowser.h
--- a/Engine/Source/Editor/ContentBrowser/Private/SContentBrowser.h
+++ b/Engine/Source/Editor/ContentBrowser/Private/SContentBrowser.h
@@ -415,7 +415,7 @@ private:
 	/** Handler to check to see if a save asset command is allowed */
 	bool HandleSaveAssetCommandCanExecute() const;
 
-	/** Handler for Rename */
+	/** Handler for save asset */
 	void HandleSaveAssetCommand();
 
 	/** Handler for SaveAll in folder */
