Reference Viewer: added option to browse to owning asset when browsing to external assets.

diff --git a/Engine/Source/Editor/ContentBrowser/Private/ContentBrowserSingleton.cpp b/Engine/Source/Editor/ContentBrowser/Private/ContentBrowserSingleton.cpp
--- a/Engine/Source/Editor/ContentBrowser/Private/ContentBrowserSingleton.cpp
+++ b/Engine/Source/Editor/ContentBrowser/Private/ContentBrowserSingleton.cpp
@@ -3,7 +3,9 @@
 
 #include "ContentBrowserSingleton.h"
 
+#include "Algo/Transform.h"
 #include "AssetRegistry/AssetData.h"
+#include "AssetRegistry/AssetRegistryModule.h"
 #include "AssetToolsModule.h"
 #include "AssetViewUtils.h"
 #include "CollectionAssetRegistryBridge.h"
@@ -468,7 +470,31 @@ void FContentBrowserSingleton::SyncBrowserToAssets(const TArray<FAssetData>& Ass
 		{
 			FocusContentBrowser(ContentBrowserToSync);
 		}
-		ContentBrowserToSync->SyncToAssets(AssetDataList);
+
+		TSet<FName> OuterPathNames;
+		TArray<FAssetData> AssetDataListToSync;
+		AssetDataListToSync.Reserve(AssetDataList.Num());
+
+		for (const FAssetData& AssetData : AssetDataList)
+		{
+			if (FName OuterPathName = AssetData.GetOptionalOuterPathName(); !OuterPathName.IsNone())
+			{
+				OuterPathNames.Add(*FSoftObjectPath(OuterPathName.ToString()).GetLongPackageName());
+			}
+			else
+			{
+				AssetDataListToSync.Add(AssetData);
+			}
+		}
+
+		if (OuterPathNames.Num())
+		{
+			TMap<FName, FAssetData> PackagesToAssetDataMap;
+			UE::AssetRegistry::GetAssetForPackages(OuterPathNames.Array(), PackagesToAssetDataMap);
+			Algo::Transform(PackagesToAssetDataMap, AssetDataListToSync, [](const TPair<FName, FAssetData>& Pair) { return Pair.Value; });
+		}
+
+		ContentBrowserToSync->SyncToAssets(AssetDataListToSync);
 	}
 }
 
