Enable PrimaryContentBrowser to be specified for use with IContentBrowserSingleton

* Exposed ability for test to set the Primary content browser so that setting changes performed imperatively will pertain to the test-created content browser that was spawned.
* This supports creation of automated editor workflow tests that use test-created content browser.  IContentBrowserSingleton is the interface the tests have to manipulate the CB settings.

diff --git a/Engine/Source/Editor/ContentBrowser/Private/ContentBrowserSingleton.cpp b/Engine/Source/Editor/ContentBrowser/Private/ContentBrowserSingleton.cpp
--- a/Engine/Source/Editor/ContentBrowser/Private/ContentBrowserSingleton.cpp
+++ b/Engine/Source/Editor/ContentBrowser/Private/ContentBrowserSingleton.cpp
@@ -288,6 +288,21 @@ bool FContentBrowserSingleton::HasPrimaryContentBrowser() const
 	}
 }
 
+bool FContentBrowserSingleton::SetPrimaryContentBrowser(FName InstanceName)
+{
+	for (int32 BrowserIdx = 0; BrowserIdx < AllContentBrowsers.Num(); ++BrowserIdx)
+	{
+		TSharedPtr<SContentBrowser> ContentBrowser = AllContentBrowsers[BrowserIdx].Pin();
+		if ( ContentBrowser && ContentBrowser->GetInstanceName() == InstanceName)
+		{
+			// There is at least one valid content browser
+			PrimaryContentBrowser = ContentBrowser;
+			return true;
+		}
+	}
+	return false;
+}
+
 void FContentBrowserSingleton::FocusPrimaryContentBrowser(bool bFocusSearch)
 {
 	// See if the primary content browser is still valid
diff --git a/Engine/Source/Editor/ContentBrowser/Private/ContentBrowserSingleton.h b/Engine/Source/Editor/ContentBrowser/Private/ContentBrowserSingleton.h
--- a/Engine/Source/Editor/ContentBrowser/Private/ContentBrowserSingleton.h
+++ b/Engine/Source/Editor/ContentBrowser/Private/ContentBrowserSingleton.h
@@ -85,6 +85,7 @@ public:
 	virtual void CreateSaveAssetDialog(const FSaveAssetDialogConfig& SaveAssetConfig, const FOnObjectPathChosenForSave& OnAssetNameChosenForSave, const FOnAssetDialogCancelled& OnAssetDialogCancelled) override;
 	virtual FString CreateModalSaveAssetDialog(const FSaveAssetDialogConfig& SaveAssetConfig) override;
 	virtual bool HasPrimaryContentBrowser() const override;
+	virtual bool SetPrimaryContentBrowser(FName InstanceName) override;
 	virtual void FocusPrimaryContentBrowser(bool bFocusSearch) override;
 	virtual void FocusContentBrowserSearchField(TSharedPtr<SWidget> ContentBrowserWidget) override;
 	virtual void CreateNewAsset(const FString& DefaultAssetName, const FString& PackagePath, UClass* AssetClass, UFactory* Factory) override;
diff --git a/Engine/Source/Editor/ContentBrowser/Public/IContentBrowserSingleton.h b/Engine/Source/Editor/ContentBrowser/Public/IContentBrowserSingleton.h
--- a/Engine/Source/Editor/ContentBrowser/Public/IContentBrowserSingleton.h
+++ b/Engine/Source/Editor/ContentBrowser/Public/IContentBrowserSingleton.h
@@ -569,6 +569,11 @@ public:
 	/** Returns true if there is at least one browser open that is eligible to be a primary content browser */
 	virtual bool HasPrimaryContentBrowser() const = 0;
 
+	/** Sets the primary content browser for subsequent state changes through this singleton
+	 * Returns true if content browser was changed sucessfully
+	 */
+	virtual bool SetPrimaryContentBrowser(FName InstanceName) = 0;
+
 	/** Brings the primary content browser to the front or opens one if it does not exist. */
 	virtual void FocusPrimaryContentBrowser(bool bFocusSearch) = 0;
 
