Replace IContentBrowserSingleton::CanChangeAssetPublicState by IAssetTools::CanAssetBePublic
#rnx

diff --git a/Engine/Source/Developer/AssetTools/Private/AssetTools.cpp b/Engine/Source/Developer/AssetTools/Private/AssetTools.cpp
--- a/Engine/Source/Developer/AssetTools/Private/AssetTools.cpp
+++ b/Engine/Source/Developer/AssetTools/Private/AssetTools.cpp
@@ -6198,5 +6198,30 @@ void UAssetToolsImpl::UnregisterCanMigrateAsset(const FName OwnerName)
 	CanMigrateAssetDelegates.Remove(OwnerName);
 }
 
+bool UAssetToolsImpl::CanAssetBePublic(FStringView AssetPath) const
+{
+	for (const TPair<FName, UE::AssetTools::FCanAssetBePublic>& Pair : CanAssetBePublicDelegates)
+	{
+		if (Pair.Value.IsBound())
+		{
+			if (!Pair.Value.Execute(AssetPath))
+			{
+				return false;
+			}
+		}
+	}
+	return true;
+}
+
+void UAssetToolsImpl::RegisterCanAssetBePublic(const FName OwnerName, UE::AssetTools::FCanAssetBePublic Delegate)
+{
+	CanAssetBePublicDelegates.Add(OwnerName, MoveTemp(Delegate));
+}
+
+void UAssetToolsImpl::UnregisterCanAssetBePublic(const FName OwnerName)
+{
+	CanAssetBePublicDelegates.Remove(OwnerName);
+}
+
 #undef LOCTEXT_NAMESPACE
 
diff --git a/Engine/Source/Developer/AssetTools/Private/AssetTools.h b/Engine/Source/Developer/AssetTools/Private/AssetTools.h
--- a/Engine/Source/Developer/AssetTools/Private/AssetTools.h
+++ b/Engine/Source/Developer/AssetTools/Private/AssetTools.h
@@ -145,9 +145,11 @@ public:
 	virtual bool IsNameAllowed(const FString& Name, FText* OutErrorMessage = nullptr) const override;
 	virtual void RegisterIsNameAllowedDelegate(const FName OwnerName, FIsNameAllowed Delegate) override;
 	virtual void UnregisterIsNameAllowedDelegate(const FName OwnerName) override;
-	
 	virtual void RegisterCanMigrateAsset(const FName OwnerName, UE::AssetTools::FCanMigrateAsset Delegate) override;
 	virtual void UnregisterCanMigrateAsset(const FName OwnerName) override;
+	virtual bool CanAssetBePublic(FStringView AssetPath) const override;
+	virtual void RegisterCanAssetBePublic(const FName OwnerName, UE::AssetTools::FCanAssetBePublic Delegate) override;
+	virtual void UnregisterCanAssetBePublic(const FName OwnerName) override;
 
 	virtual void SyncBrowserToAssets(const TArray<UObject*>& AssetsToSync) override;
 	virtual void SyncBrowserToAssets(const TArray<FAssetData>& AssetsToSync) override;
@@ -267,6 +269,8 @@ private:
 	UE::AssetTools::FOnPackageMigration OnPackageMigration;
 
 	TMap<FName, UE::AssetTools::FCanMigrateAsset> CanMigrateAssetDelegates;
+	
+	TMap<FName, UE::AssetTools::FCanAssetBePublic> CanAssetBePublicDelegates;
 };
 
 PRAGMA_ENABLE_DEPRECATION_WARNINGS
diff --git a/Engine/Source/Developer/AssetTools/Public/IAssetTools.h b/Engine/Source/Developer/AssetTools/Public/IAssetTools.h
--- a/Engine/Source/Developer/AssetTools/Public/IAssetTools.h
+++ b/Engine/Source/Developer/AssetTools/Public/IAssetTools.h
@@ -35,6 +35,8 @@ namespace UE::AssetTools
 	DECLARE_MULTICAST_DELEGATE_OneParam(FOnPackageMigration, FPackageMigrationContext&);
 
 	DECLARE_DELEGATE_RetVal_OneParam(bool, FCanMigrateAsset, FName);
+
+	DECLARE_DELEGATE_RetVal_OneParam(bool, FCanAssetBePublic, FStringView /*AssetPath*/);
 }
 
 UENUM()
@@ -618,6 +620,16 @@ public:
 	virtual void RegisterCanMigrateAsset(const FName OwnerName, UE::AssetTools::FCanMigrateAsset Delegate) = 0;
 	virtual void UnregisterCanMigrateAsset(const FName OwnerName) = 0;
 
+	/** Returns whether the specified asset can be public (referenceable from another mount point / plugin) */
+	virtual bool CanAssetBePublic(FStringView AssetPath) const = 0;
+
+	/**
+	 * Register/unregister delegates to specify whether an asset can be made public (referenceable from another mount point / plugin)
+	 * By default any asset can be public and if any delegate return false, the asset must be private
+	 */
+	virtual void RegisterCanAssetBePublic(const FName OwnerName, UE::AssetTools::FCanAssetBePublic Delegate) = 0;
+	virtual void UnregisterCanAssetBePublic(const FName OwnerName) = 0;
+
 	/** Syncs the primary content browser to the specified assets, whether or not it is locked. Most syncs that come from AssetTools -feel- like they came from the content browser, so this is okay. */
 	virtual void SyncBrowserToAssets(const TArray<UObject*>& AssetsToSync) = 0;
 	virtual void SyncBrowserToAssets(const TArray<FAssetData>& AssetsToSync) = 0;
diff --git a/Engine/Source/Editor/ContentBrowser/Private/AssetContextMenu.cpp b/Engine/Source/Editor/ContentBrowser/Private/AssetContextMenu.cpp
--- a/Engine/Source/Editor/ContentBrowser/Private/AssetContextMenu.cpp
+++ b/Engine/Source/Editor/ContentBrowser/Private/AssetContextMenu.cpp
@@ -1309,8 +1309,17 @@ void FAssetContextMenu::CacheCanExecuteVars()
 		}
 		if (bCanChangeAssetPublicState)
 		{
-			FNameBuilder ItemInternalPath(SelectedItem.GetInternalPath());
-			bCanChangeAssetPublicState = FContentBrowserSingleton::Get().CanChangeAssetPublicState(ItemInternalPath);
+			const FNameBuilder ItemInternalPath(SelectedItem.GetInternalPath());
+			const FStringView AssetPath(ItemInternalPath);
+
+			if (!IAssetTools::Get().CanAssetBePublic(AssetPath))
+			{
+				const FAssetData AssetData = IAssetRegistry::GetChecked().GetAssetByObjectPath(FSoftObjectPath(ItemInternalPath));
+				if (!AssetData.IsValid() || AssetData.HasAnyPackageFlags(PKG_NotExternallyReferenceable))
+				{
+					bCanChangeAssetPublicState = false;
+				}
+			}
 		}
 
 		if (bCanChangeAssetPublicState)
diff --git a/Engine/Source/Editor/ContentBrowser/Private/ContentBrowserSingleton.cpp b/Engine/Source/Editor/ContentBrowser/Private/ContentBrowserSingleton.cpp
--- a/Engine/Source/Editor/ContentBrowser/Private/ContentBrowserSingleton.cpp
+++ b/Engine/Source/Editor/ContentBrowser/Private/ContentBrowserSingleton.cpp
@@ -1176,15 +1176,6 @@ void FContentBrowserSingleton::RebuildPrivateContentStateCache()
 	}
 }
 
-bool FContentBrowserSingleton::CanChangeAssetPublicState(FStringView AssetPath)
-{
-	if (CanChangeAssetPublicStateDelegate.IsBound())
-	{
-		return CanChangeAssetPublicStateDelegate.Execute(AssetPath);
-	}
-	return true;
-}
-
 bool FContentBrowserSingleton::IsFolderShowPrivateContentToggleable(const FStringView VirtualFolderPath)
 {
 	if (IsFolderShowPrivateContentToggleableDelegate.IsBound())
@@ -1205,16 +1196,6 @@ void FContentBrowserSingleton::SetPrivateContentPermissionListDirty()
 	ShowPrivateContentState.CachedVirtualPaths.Reset();
 }
 
-void FContentBrowserSingleton::RegisterCanChangeAssetPublicStateDelegate(FCanChangeAssetPublicStateDelegate InCanChangeAssetPublicStateDelegate)
-{
-	CanChangeAssetPublicStateDelegate = MoveTemp(InCanChangeAssetPublicStateDelegate);
-}
-
-void FContentBrowserSingleton::UnregisterCanChangeAssetPublicStateDelegate()
-{
-	CanChangeAssetPublicStateDelegate = FCanChangeAssetPublicStateDelegate();
-}
-
 void FContentBrowserSingleton::RegisterIsFolderShowPrivateContentToggleableDelegate(FIsFolderShowPrivateContentToggleableDelegate InIsFolderShowPrivateContentToggleableDelegate)
 {
 	IsFolderShowPrivateContentToggleableDelegate = MoveTemp(InIsFolderShowPrivateContentToggleableDelegate);
diff --git a/Engine/Source/Editor/ContentBrowser/Private/ContentBrowserSingleton.h b/Engine/Source/Editor/ContentBrowser/Private/ContentBrowserSingleton.h
--- a/Engine/Source/Editor/ContentBrowser/Private/ContentBrowserSingleton.h
+++ b/Engine/Source/Editor/ContentBrowser/Private/ContentBrowserSingleton.h
@@ -109,13 +109,10 @@ PRAGMA_ENABLE_DEPRECATION_WARNINGS
 	virtual void ExecuteRename(TSharedPtr<SWidget> PickerWidget) override;
 	virtual void ExecuteAddFolder(TSharedPtr<SWidget> PathPickerWidget) override;
 	virtual void RefreshPathView(TSharedPtr<SWidget> PathPickerWidget) override;
-	virtual bool CanChangeAssetPublicState(FStringView AssetPath) override;
 	virtual bool IsShowingPrivateContent(const FStringView VirtualFolderPath) override;
 	virtual bool IsFolderShowPrivateContentToggleable(const FStringView VirtualFolderPath) override;
 	virtual const TSharedPtr<FPathPermissionList>& GetShowPrivateContentPermissionList() override;
 	virtual void SetPrivateContentPermissionListDirty() override;
-	virtual void RegisterCanChangeAssetPublicStateDelegate(FCanChangeAssetPublicStateDelegate InCanChangeAssetPublicStateDelegate) override;
-	virtual void UnregisterCanChangeAssetPublicStateDelegate() override;
 	virtual void RegisterIsFolderShowPrivateContentToggleableDelegate(FIsFolderShowPrivateContentToggleableDelegate InIsFolderShowPrivateContentToggleableDelegate) override;
 	virtual void UnregisterIsFolderShowPrivateContentToggleableDelegate() override;
 	virtual FDelegateHandle RegisterOnFavoritesChangedHandler(FSimpleDelegate OnFavoritesChanged) override;
@@ -194,8 +191,6 @@ public:
 
 private:
 
-	FCanChangeAssetPublicStateDelegate CanChangeAssetPublicStateDelegate;
-
 	FIsFolderShowPrivateContentToggleableDelegate IsFolderShowPrivateContentToggleableDelegate;
 
 	TArray<TWeakPtr<SContentBrowser>> AllContentBrowsers;
diff --git a/Engine/Source/Editor/ContentBrowser/Public/ContentBrowserDelegates.h b/Engine/Source/Editor/ContentBrowser/Public/ContentBrowserDelegates.h
--- a/Engine/Source/Editor/ContentBrowser/Public/ContentBrowserDelegates.h
+++ b/Engine/Source/Editor/ContentBrowser/Public/ContentBrowserDelegates.h
@@ -128,9 +128,6 @@ DECLARE_DELEGATE_OneParam(FOnAssetsChosenForOpen, const TArray<FAssetData>& /*Se
 /** Called from the Asset Dialog when an asset name is chosen in non-modal Save dialogs */
 DECLARE_DELEGATE_OneParam(FOnObjectPathChosenForSave, const FString& /*ObjectPath*/);
 
-/** Called to determine whether the public/private state of the specified asset can be changed */
-DECLARE_DELEGATE_RetVal_OneParam(bool, FCanChangeAssetPublicStateDelegate, const FStringView /*AssetPath*/);
-
 /** Called when custom behavior is needed for allowing folders to toggle their private content edit state */
 DECLARE_DELEGATE_RetVal_OneParam(bool, FIsFolderShowPrivateContentToggleableDelegate, const FStringView /*VirtualFolderPath*/);
 
diff --git a/Engine/Source/Editor/ContentBrowser/Public/IContentBrowserSingleton.h b/Engine/Source/Editor/ContentBrowser/Public/IContentBrowserSingleton.h
--- a/Engine/Source/Editor/ContentBrowser/Public/IContentBrowserSingleton.h
+++ b/Engine/Source/Editor/ContentBrowser/Public/IContentBrowserSingleton.h
@@ -707,9 +707,6 @@ public:
 	/** Returns InPath if can be written to, otherwise picks a default path that can be written to */
 	virtual FContentBrowserItemPath GetInitialPathToSaveAsset(const FContentBrowserItemPath& InPath) = 0;
 
-	/** Returns true if the public/private state of the specified asset can be changed */
-	virtual bool CanChangeAssetPublicState(FStringView AssetPath) = 0;
-
 	/** Returns true if FolderPath is a private content edit folder */
 	virtual bool IsShowingPrivateContent(const FStringView VirtualFolderPath) = 0;
 
@@ -722,12 +719,6 @@ public:
 	/** Declares the Private Content Permission List dirty */
 	virtual void SetPrivateContentPermissionListDirty() = 0;
 
-	/** Registers the delegate called to determine whether the public/private state of the specified asset can be changed */
-	virtual void RegisterCanChangeAssetPublicStateDelegate(FCanChangeAssetPublicStateDelegate InCanChangeAssetPublicStateDelegate) = 0;
-
-	/** Unregisters the delegate called to determine whether the public/private state of the specified asset can be changed */
-	virtual void UnregisterCanChangeAssetPublicStateDelegate() = 0;
-
 	/** Registers the delegate for custom handling of if a Folder allows private content edits */
 	virtual void RegisterIsFolderShowPrivateContentToggleableDelegate(FIsFolderShowPrivateContentToggleableDelegate InIsFolderShowPrivateContentToggleableDelegate) = 0;
 
