Fixed unreachable code warnings/errors.

diff --git a/Engine/Plugins/Media/MediaIOFramework/Source/MediaIOCore/Private/Player/MediaIOCorePlayerBase.cpp b/Engine/Plugins/Media/MediaIOFramework/Source/MediaIOCore/Private/Player/MediaIOCorePlayerBase.cpp
--- a/Engine/Plugins/Media/MediaIOFramework/Source/MediaIOCore/Private/Player/MediaIOCorePlayerBase.cpp
+++ b/Engine/Plugins/Media/MediaIOFramework/Source/MediaIOCore/Private/Player/MediaIOCorePlayerBase.cpp
@@ -465,7 +465,7 @@ FText FMediaIOCorePlayerBase::GetTrackDisplayName(EMediaTrackType TrackType, int
 		return LOCTEXT("DefaultVideoTrackName", "Video Track");
 
 	default:
-		return FText::GetEmpty();
+		break;
 	}
 }
 
diff --git a/Engine/Plugins/MovieScene/SequencerScripting/Source/SequencerScripting/Private/ExtensionLibraries/MovieSceneTrackExtensions.cpp b/Engine/Plugins/MovieScene/SequencerScripting/Source/SequencerScripting/Private/ExtensionLibraries/MovieSceneTrackExtensions.cpp
--- a/Engine/Plugins/MovieScene/SequencerScripting/Source/SequencerScripting/Private/ExtensionLibraries/MovieSceneTrackExtensions.cpp
+++ b/Engine/Plugins/MovieScene/SequencerScripting/Source/SequencerScripting/Private/ExtensionLibraries/MovieSceneTrackExtensions.cpp
@@ -31,8 +31,9 @@ FText UMovieSceneTrackExtensions::GetDisplayName(UMovieSceneTrack* Track)
 
 #if WITH_EDITORONLY_DATA
 	return Track->GetDisplayName();
-#endif
+#else
 	return FText::GetEmpty();
+#endif
 }
 
 void UMovieSceneTrackExtensions::SetTrackRowDisplayName(UMovieSceneTrack* Track, const FText& InName, int32 RowIndex)
@@ -61,8 +62,9 @@ FText UMovieSceneTrackExtensions::GetTrackRowDisplayName(UMovieSceneTrack* Track
 
 #if WITH_EDITORONLY_DATA
 	return Track->GetTrackRowDisplayName(RowIndex);
-#endif
+#else
 	return FText::GetEmpty();
+#endif
 }
 
 UMovieSceneSection* UMovieSceneTrackExtensions::AddSection(UMovieSceneTrack* Track)
diff --git a/Engine/Plugins/Runtime/StateTree/Source/StateTreeModule/Private/Conditions/StateTreeCommonConditions.cpp b/Engine/Plugins/Runtime/StateTree/Source/StateTreeModule/Private/Conditions/StateTreeCommonConditions.cpp
--- a/Engine/Plugins/Runtime/StateTree/Source/StateTreeModule/Private/Conditions/StateTreeCommonConditions.cpp
+++ b/Engine/Plugins/Runtime/StateTree/Source/StateTreeModule/Private/Conditions/StateTreeCommonConditions.cpp
@@ -36,10 +36,10 @@ FText GetOperatorText(const EGenericAICheck Operator)
 		return FText::FromString(TEXT("&gt;="));
 		break;
 	default:
-		return FText::FromString(TEXT("??"));
 		break;
 	}
-	return FText::GetEmpty();
+
+	return FText::FromString(TEXT("??"));
 }
 #endif // WITH_EDITOR
 
diff --git a/Engine/Source/Editor/ContentBrowser/Private/AssetViewWidgets.cpp b/Engine/Source/Editor/ContentBrowser/Private/AssetViewWidgets.cpp
--- a/Engine/Source/Editor/ContentBrowser/Private/AssetViewWidgets.cpp
+++ b/Engine/Source/Editor/ContentBrowser/Private/AssetViewWidgets.cpp
@@ -1330,29 +1330,20 @@ void SAssetViewItem::CacheDisplayTags()
 	
 					return FText::AsNumber(Num, &NumFormatOpts);
 				}
-				else
+
+				const bool bIsSigned = InNumberString.Len() > 0 && (InNumberString[0] == TEXT('-') || InNumberString[0] == TEXT('+'));
+				if (bIsSigned)
 				{
-					const bool bIsSigned = InNumberString.Len() > 0 && (InNumberString[0] == TEXT('-') || InNumberString[0] == TEXT('+'));
-	
-					if (bIsSigned)
-					{
-						// Convert the number as a signed int
-						int64 Num = 0;
-						LexFromString(Num, *InNumberString);
-	
-						return FText::AsNumber(Num);
-					}
-					else
-					{
-						// Convert the number as an unsigned int
-						uint64 Num = 0;
-						LexFromString(Num, *InNumberString);
-	
-						return FText::AsNumber(Num);
-					}
+					// Convert the number as a signed int
+					int64 Num = 0;
+					LexFromString(Num, *InNumberString);
+					return FText::AsNumber(Num);
 				}
-	
-				return FText::GetEmpty();
+
+				// Convert the number as an unsigned int
+				uint64 Num = 0;
+				LexFromString(Num, *InNumberString);
+				return FText::AsNumber(Num);
 			};
 	
 			bool bHasSetDisplayValue = false;
diff --git a/Engine/Source/Editor/GraphEditor/Private/SGraphNode.cpp b/Engine/Source/Editor/GraphEditor/Private/SGraphNode.cpp
--- a/Engine/Source/Editor/GraphEditor/Private/SGraphNode.cpp
+++ b/Engine/Source/Editor/GraphEditor/Private/SGraphNode.cpp
@@ -593,22 +593,14 @@ FVector2D SGraphNode::GetPosition() const
 
 FString SGraphNode::GetEditableNodeTitle() const
 {
-	if (GraphNode != NULL)
+	if (GraphNode != nullptr)
 	{
 		// Trying to catch a non-reproducible crash in this function
 		check(GraphNode->IsValidLowLevel());
-	}
-
-	if(GraphNode)
-	{
 		return GraphNode->GetNodeTitle(ENodeTitleType::EditableTitle).ToString();
 	}
-	return NSLOCTEXT("GraphEditor", "NullNode", "Null Node").ToString();
 
-	// Get the portion of the node that is actually editable text (may be a subsection of the title, or something else entirely)
-	return (GraphNode != NULL)
-		? GraphNode->GetNodeTitle(ENodeTitleType::EditableTitle).ToString()
-		: NSLOCTEXT("GraphEditor", "NullNode", "Null Node").ToString();
+	return NSLOCTEXT("GraphEditor", "NullNode", "Null Node").ToString();
 }
 
 FText SGraphNode::GetEditableNodeTitleAsText() const
diff --git a/Engine/Source/Editor/LandscapeEditor/Private/LandscapeEdModeMirrorTool.cpp b/Engine/Source/Editor/LandscapeEditor/Private/LandscapeEdModeMirrorTool.cpp
--- a/Engine/Source/Editor/LandscapeEditor/Private/LandscapeEdModeMirrorTool.cpp
+++ b/Engine/Source/Editor/LandscapeEditor/Private/LandscapeEdModeMirrorTool.cpp
@@ -224,10 +224,6 @@ public:
 				return EAxisList::None;
 			}
 		}
-		else
-		{
-			return EAxisList::None;
-		}
 
 		return EAxisList::None;
 	}
diff --git a/Engine/Source/Editor/LandscapeEditor/Private/LandscapeEditorDetailCustomization_ImportExport.cpp b/Engine/Source/Editor/LandscapeEditor/Private/LandscapeEditorDetailCustomization_ImportExport.cpp
--- a/Engine/Source/Editor/LandscapeEditor/Private/LandscapeEditorDetailCustomization_ImportExport.cpp
+++ b/Engine/Source/Editor/LandscapeEditor/Private/LandscapeEditorDetailCustomization_ImportExport.cpp
@@ -936,12 +936,8 @@ FText FLandscapeEditorDetailCustomization_ImportExport::GetImportLandscapeResolu
 		Args.Add(TEXT("Height"), Height);
 		return FText::Format(LOCTEXT("ImportResolution_Format", "{Width}\u00D7{Height}"), Args);
 	}
-	else
-	{
-		return LOCTEXT("ImportResolution_Invalid", "(invalid)");
-	}
-	
-	return FText::GetEmpty();
+
+	return LOCTEXT("ImportResolution_Invalid", "(invalid)");
 }
 
 
diff --git a/Engine/Source/Editor/MaterialEditor/Private/FindInMaterial.cpp b/Engine/Source/Editor/MaterialEditor/Private/FindInMaterial.cpp
--- a/Engine/Source/Editor/MaterialEditor/Private/FindInMaterial.cpp
+++ b/Engine/Source/Editor/MaterialEditor/Private/FindInMaterial.cpp
@@ -58,11 +58,8 @@ FText FFindInMaterialResult::GetCategory() const
 	{
 		return LOCTEXT("PinCategory", "Pin");
 	}
-	else
-	{
-		return LOCTEXT("NodeCategory", "Node");
-	}
-	return FText::GetEmpty();
+
+	return LOCTEXT("NodeCategory", "Node");
 }
 
 TSharedRef<SWidget> FFindInMaterialResult::CreateIcon() const
diff --git a/Engine/Source/Editor/MovieSceneTools/Private/TrackEditors/AttachTrackEditor.cpp b/Engine/Source/Editor/MovieSceneTools/Private/TrackEditors/AttachTrackEditor.cpp
--- a/Engine/Source/Editor/MovieSceneTools/Private/TrackEditors/AttachTrackEditor.cpp
+++ b/Engine/Source/Editor/MovieSceneTools/Private/TrackEditors/AttachTrackEditor.cpp
@@ -592,12 +592,8 @@ struct FLocalTransformEvaluator : ITransformEvaluator
 			UMovieScene3DTransformTrack* TransformTrack = TransformEval.GetSubtype<TTuple<UMovieScene3DTransformTrack*, UObject*>>().Get<0>();
 		 	return GetLocationAtTime(TransformTrack, InTime, *WeakAttachTrackEditor.Pin()->Interrogator.Get());
 		}
-		else
-		{
-		 	return TransformEval.GetSubtype<FTransform>();
-		}
 
-		return FTransform::Identity;
+		return TransformEval.GetSubtype<FTransform>();
 	}
 
 private:
