Plugin folder icon in content browser path tree
Gated behind ContentBrowser.UsePluginFolderIcon until proper icon is ready

diff --git a/Engine/Source/Editor/ContentBrowser/Private/SourcesViewWidgets.cpp b/Engine/Source/Editor/ContentBrowser/Private/SourcesViewWidgets.cpp
--- a/Engine/Source/Editor/ContentBrowser/Private/SourcesViewWidgets.cpp
+++ b/Engine/Source/Editor/ContentBrowser/Private/SourcesViewWidgets.cpp
@@ -39,6 +39,47 @@ struct FSlateBrush;
 
 #define LOCTEXT_NAMESPACE "ContentBrowser"
 
+static bool GUsePluginFolderIcon = 0;
+static FAutoConsoleVariableRef CVarUsePluginFolderIcon(
+	TEXT("ContentBrowser.UsePluginFolderIcon"),
+	GUsePluginFolderIcon,
+	TEXT("Temporary flag to control use of separate icon for plugin folders pending final icon"),
+	ECVF_Default
+);
+
+struct FAssetTreeItemBrushes
+{
+	/** Brushes for the different folder states */
+	const FSlateBrush* FolderOpenBrush;
+	const FSlateBrush* FolderClosedBrush;
+	const FSlateBrush* FolderOpenCodeBrush;
+	const FSlateBrush* FolderClosedCodeBrush;
+	const FSlateBrush* FolderOpenDeveloperBrush;
+	const FSlateBrush* FolderClosedDeveloperBrush;
+	const FSlateBrush* FolderOpenPluginBrush;
+	const FSlateBrush* FolderClosedPluginBrush;
+
+	FAssetTreeItemBrushes()
+	{
+		FolderOpenBrush = FAppStyle::GetBrush("ContentBrowser.AssetTreeFolderOpen");
+		FolderClosedBrush = FAppStyle::GetBrush("ContentBrowser.AssetTreeFolderClosed");
+		FolderOpenCodeBrush = FAppStyle::GetBrush("ContentBrowser.AssetTreeFolderOpenCode");
+		FolderClosedCodeBrush = FAppStyle::GetBrush("ContentBrowser.AssetTreeFolderClosedCode");
+		FolderOpenDeveloperBrush = FAppStyle::GetBrush("ContentBrowser.AssetTreeFolderOpenDeveloper");
+		FolderClosedDeveloperBrush = FAppStyle::GetBrush("ContentBrowser.AssetTreeFolderClosedDeveloper");
+		FolderOpenPluginBrush = FAppStyle::GetBrush("ContentBrowser.AssetTreeFolderOpenPlugin");
+		FolderClosedPluginBrush = FAppStyle::GetBrush("ContentBrowser.AssetTreeFolderClosedPlugin");
+	}
+	
+	static FAssetTreeItemBrushes& Get()
+	{
+		static FAssetTreeItemBrushes Instance;
+		return Instance;
+	}
+};
+
+
+
 //////////////////////////
 // SAssetTreeItem
 //////////////////////////
@@ -53,13 +94,6 @@ void SAssetTreeItem::Construct( const FArguments& InArgs )
 
 	IsSelected = InArgs._IsSelected;
 
-	FolderOpenBrush = FAppStyle::GetBrush("ContentBrowser.AssetTreeFolderOpen");
-	FolderClosedBrush = FAppStyle::GetBrush("ContentBrowser.AssetTreeFolderClosed");
-	FolderOpenCodeBrush = FAppStyle::GetBrush("ContentBrowser.AssetTreeFolderOpenCode");
-	FolderClosedCodeBrush = FAppStyle::GetBrush("ContentBrowser.AssetTreeFolderClosedCode");
-	FolderOpenDeveloperBrush = FAppStyle::GetBrush("ContentBrowser.AssetTreeFolderOpenDeveloper");
-	FolderClosedDeveloperBrush = FAppStyle::GetBrush("ContentBrowser.AssetTreeFolderClosedDeveloper");
-
 	FolderType = EFolderType::Normal;
 	if (ContentBrowserUtils::IsItemDeveloperContent(InArgs._TreeItem->GetItem()))
 	{
@@ -69,6 +103,10 @@ void SAssetTreeItem::Construct( const FArguments& InArgs )
 	{
 		FolderType = EFolderType::Code;
 	}
+	else if (EnumHasAnyFlags(InArgs._TreeItem->GetItem().GetItemCategory(), EContentBrowserItemFlags::Category_Plugin))
+	{
+		FolderType = EFolderType::Plugin;
+	}
 
 	bool bIsRoot = !InArgs._TreeItem->Parent.IsValid();
 
@@ -256,16 +294,24 @@ bool SAssetTreeItem::IsReadOnly() const
 
 const FSlateBrush* SAssetTreeItem::GetFolderIcon() const
 {
+	FAssetTreeItemBrushes& Brushes = FAssetTreeItemBrushes::Get();
 	switch( FolderType )
 	{
 	case EFolderType::Code:
-		return ( IsItemExpanded.Get() ) ? FolderOpenCodeBrush : FolderClosedCodeBrush;
+		return ( IsItemExpanded.Get() ) ? Brushes.FolderOpenCodeBrush : Brushes.FolderClosedCodeBrush;
 
 	case EFolderType::Developer:
-		return (IsItemExpanded.Get()) ? FolderOpenDeveloperBrush : FolderClosedDeveloperBrush;
+		return (IsItemExpanded.Get()) ? Brushes.FolderOpenDeveloperBrush : Brushes.FolderClosedDeveloperBrush;
+
+	case EFolderType::Plugin:
+		if (GUsePluginFolderIcon)
+		{
+			return (IsItemExpanded.Get()) ? Brushes.FolderOpenPluginBrush : Brushes.FolderClosedPluginBrush;
+		}
+		// Fall through to default 
 
 	default:
-		return ( IsItemExpanded.Get() ) ? FolderOpenBrush : FolderClosedBrush;
+		return ( IsItemExpanded.Get() ) ? Brushes.FolderOpenBrush : Brushes.FolderClosedBrush;
 	}
 }
 
@@ -296,7 +342,7 @@ FSlateColor SAssetTreeItem::GetFolderColor() const
 
 		return FoundColor;
 	}
-	
+
 	return ContentBrowserUtils::GetDefaultColor();
 }
 
diff --git a/Engine/Source/Editor/ContentBrowser/Private/SourcesViewWidgets.h b/Engine/Source/Editor/ContentBrowser/Private/SourcesViewWidgets.h
--- a/Engine/Source/Editor/ContentBrowser/Private/SourcesViewWidgets.h
+++ b/Engine/Source/Editor/ContentBrowser/Private/SourcesViewWidgets.h
@@ -108,6 +108,7 @@ private:
 		Normal,
 		Code,
 		Developer,
+		Plugin,
 	};
 
 	/** The data for this item */
@@ -125,14 +126,6 @@ private:
 	/** True when this item has children and is expanded */
 	TAttribute<bool> IsItemExpanded;
 
-	/** Brushes for the different folder states */
-	const FSlateBrush* FolderOpenBrush;
-	const FSlateBrush* FolderClosedBrush;
-	const FSlateBrush* FolderOpenCodeBrush;
-	const FSlateBrush* FolderClosedCodeBrush;
-	const FSlateBrush* FolderOpenDeveloperBrush;
-	const FSlateBrush* FolderClosedDeveloperBrush;
-
 	/** Delegate called to get the selection state of an asset path */
 	FIsSelected IsSelected;
 
diff --git a/Engine/Source/Editor/EditorStyle/Private/StarshipStyle.cpp b/Engine/Source/Editor/EditorStyle/Private/StarshipStyle.cpp
--- a/Engine/Source/Editor/EditorStyle/Private/StarshipStyle.cpp
+++ b/Engine/Source/Editor/EditorStyle/Private/StarshipStyle.cpp
@@ -6058,6 +6058,8 @@ void FStarshipEditorStyle::FStyle::SetupContentBrowserStyle()
 		Set( "ContentBrowser.AssetTreeFolderClosedDeveloper", new IMAGE_BRUSH_SVG( "Starship/ContentBrowser/folder-developer", Icon16x16));
 		Set( "ContentBrowser.AssetTreeFolderOpenCode", new IMAGE_BRUSH_SVG( "Starship/ContentBrowser/folder-code-open", Icon16x16) );
 		Set( "ContentBrowser.AssetTreeFolderClosedCode", new IMAGE_BRUSH_SVG( "Starship/ContentBrowser/folder-code", Icon16x16) );
+		Set( "ContentBrowser.AssetTreeFolderOpenPlugin", new IMAGE_BRUSH_SVG("Starship/ContentBrowser/folder-plugin-open", Icon16x16));
+		Set( "ContentBrowser.AssetTreeFolderClosedPlugin", new IMAGE_BRUSH_SVG( "Starship/ContentBrowser/folder-plugin", Icon16x16));
 
 		Set("ContentBrowser.DefaultFolderColor", FStyleColors::AccentFolder);
 
