Add plugin description to tooltips for content browser path tree items in plugins.

diff --git a/Engine/Source/Editor/ContentBrowser/Private/SourcesViewWidgets.cpp b/Engine/Source/Editor/ContentBrowser/Private/SourcesViewWidgets.cpp
--- a/Engine/Source/Editor/ContentBrowser/Private/SourcesViewWidgets.cpp
+++ b/Engine/Source/Editor/ContentBrowser/Private/SourcesViewWidgets.cpp
@@ -24,6 +24,7 @@
 #include "Layout/Margin.h"
 #include "Misc/EnumClassFlags.h"
 #include "Misc/Optional.h"
+#include "Misc/PathViews.h"
 #include "PathViewTypes.h"
 #include "SAssetTagItem.h"
 #include "SlotBase.h"
@@ -359,7 +360,26 @@ FText SAssetTreeItem::GetToolTipText() const
 {
 	if (TSharedPtr<FTreeItem> TreeItemPin = TreeItem.Pin())
 	{
-		return FText::FromName(TreeItemPin->GetItem().GetVirtualPath());
+		// If this item is a plugin folder, append the plugin description to the tooltip
+		const FContentBrowserItem& Item = TreeItemPin->GetItem();
+		FText PathText = FText::FromName(Item.GetVirtualPath());
+		if (Item.IsInPlugin())
+		{
+			FNameBuilder ItemPath{Item.GetInternalPath()};
+			FStringView PluginName = FPathViews::GetMountPointNameFromPath(ItemPath.ToView());
+			if (TSharedPtr<IPlugin> Plugin = IPluginManager::Get().FindPlugin(PluginName))
+			{
+				const FPluginDescriptor& Descriptor = Plugin->GetDescriptor();
+				if (!Descriptor.Description.IsEmpty())
+				{
+					return FText::Format(LOCTEXT("TwoLineTooltip", "{0}\n{1}"),
+						PathText,
+						FText::FromString(Descriptor.Description)
+						);
+				}
+			}
+		}
+		return PathText;
 	}
 	return FText();
 }
