Content browser folder copy updates.
Add copy path to context menu in path tree.
Change path copying for folders to use internal paths to more closely match how the same command copies internal path style references for assets.

diff --git a/Engine/Source/Editor/ContentBrowser/Private/ContentBrowserUtils.cpp b/Engine/Source/Editor/ContentBrowser/Private/ContentBrowserUtils.cpp
--- a/Engine/Source/Editor/ContentBrowser/Private/ContentBrowserUtils.cpp
+++ b/Engine/Source/Editor/ContentBrowser/Private/ContentBrowserUtils.cpp
@@ -340,27 +340,38 @@ FString ContentBrowserUtils::GetFolderReferencesText(const TArray<FContentBrowse
 		return One.GetVirtualPath().Compare(Two.GetVirtualPath()) < 0;
 	});
 
-	FString Result;
+	TStringBuilder<2048> Result;
 	for (const FContentBrowserItem& Item : SortedItems)
 	{
 		if (ensure(Item.IsFolder()))
 		{
-			Result += Item.GetVirtualPath().ToString();
-			Result += LINE_TERMINATOR;
+			FName InternalPath = Item.GetInternalPath();
+			if (!InternalPath.IsNone())
+			{
+				Result << InternalPath << LINE_TERMINATOR;
+			}
 		}
 	}
 
-	return Result;
+	return Result.ToString();
 }
 
 void ContentBrowserUtils::CopyItemReferencesToClipboard(const TArray<FContentBrowserItem>& ItemsToCopy)
 {
-	FPlatformApplicationMisc::ClipboardCopy(*GetItemReferencesText(ItemsToCopy));
+	FString Text = GetItemReferencesText(ItemsToCopy);
+	if (!Text.IsEmpty())
+	{
+		FPlatformApplicationMisc::ClipboardCopy(*Text);
+	}
 }
 
 void ContentBrowserUtils::CopyFolderReferencesToClipboard(const TArray<FContentBrowserItem>& FoldersToCopy)
 {
-	FPlatformApplicationMisc::ClipboardCopy(*GetFolderReferencesText(FoldersToCopy));
+	FString Text = GetFolderReferencesText(FoldersToCopy);
+	if (!Text.IsEmpty())
+	{
+		FPlatformApplicationMisc::ClipboardCopy(*Text);
+	}
 }
 
 void ContentBrowserUtils::CopyFilePathsToClipboard(const TArray<FContentBrowserItem>& ItemsToCopy)
diff --git a/Engine/Source/Editor/ContentBrowser/Private/PathContextMenu.cpp b/Engine/Source/Editor/ContentBrowser/Private/PathContextMenu.cpp
--- a/Engine/Source/Editor/ContentBrowser/Private/PathContextMenu.cpp
+++ b/Engine/Source/Editor/ContentBrowser/Private/PathContextMenu.cpp
@@ -161,6 +161,18 @@ void FPathContextMenu::MakePathViewContextMenu(UToolMenu* Menu)
 					);
 			}
 
+			// Assume paths with an on-disk representation also have an internal path to copy
+			if (!Context->bNoFolderOnDisk)
+			{
+				Section.AddMenuEntry(
+					"CopyPath",
+					LOCTEXT("CopyFolderPath", "Copy Path"),
+					LOCTEXT("CopyFolderTooltip", "Copy the paths of the selected folder(s)"),
+					FSlateIcon(FAppStyle::GetAppStyleSetName(), "GenericCommands.Copy"),
+					FExecuteAction::CreateSP(this, &FPathContextMenu::CopySelectedFolder)
+				);
+			}
+
 			if (Context->bCanBeModified)
 			{
 				Section.AddMenuEntry(FGenericCommands::Get().Rename,
@@ -499,15 +511,7 @@ void FPathContextMenu::SaveFilesWithinSelectedFolders(EContentBrowserItemSaveFla
 
 void FPathContextMenu::CopySelectedFoldersToClipoard()
 {
-	TStringBuilder<1024> StringBuilder;
-
-	for (const FContentBrowserItem& SelectedItem : SelectedFolders)
-	{
-		StringBuilder.Append(SelectedItem.GetVirtualPath().ToString());
-		StringBuilder.Append(LINE_TERMINATOR);
-	}
-
-	FPlatformApplicationMisc::ClipboardCopy(StringBuilder.ToString());
+	ContentBrowserUtils::CopyFolderReferencesToClipboard(SelectedFolders);
 }
 
 bool FPathContextMenu::CanExecuteDelete() const
