Separate icons for plugin roots and plugin subfolders

diff --git a/Engine/Source/Editor/ContentBrowser/Private/SourcesViewWidgets.cpp b/Engine/Source/Editor/ContentBrowser/Private/SourcesViewWidgets.cpp
--- a/Engine/Source/Editor/ContentBrowser/Private/SourcesViewWidgets.cpp
+++ b/Engine/Source/Editor/ContentBrowser/Private/SourcesViewWidgets.cpp
@@ -57,6 +57,8 @@ struct FAssetTreeItemBrushes
 	const FSlateBrush* FolderClosedCodeBrush;
 	const FSlateBrush* FolderOpenDeveloperBrush;
 	const FSlateBrush* FolderClosedDeveloperBrush;
+	const FSlateBrush* FolderOpenPluginRootBrush;
+	const FSlateBrush* FolderClosedPluginRootBrush;
 	const FSlateBrush* FolderOpenPluginBrush;
 	const FSlateBrush* FolderClosedPluginBrush;
 
@@ -68,6 +70,8 @@ struct FAssetTreeItemBrushes
 		FolderClosedCodeBrush = FAppStyle::GetBrush("ContentBrowser.AssetTreeFolderClosedCode");
 		FolderOpenDeveloperBrush = FAppStyle::GetBrush("ContentBrowser.AssetTreeFolderOpenDeveloper");
 		FolderClosedDeveloperBrush = FAppStyle::GetBrush("ContentBrowser.AssetTreeFolderClosedDeveloper");
+		FolderOpenPluginRootBrush = FAppStyle::GetBrush("ContentBrowser.AssetTreeFolderOpenPluginRoot");
+		FolderClosedPluginRootBrush = FAppStyle::GetBrush("ContentBrowser.AssetTreeFolderClosedPluginRoot");
 		FolderOpenPluginBrush = FAppStyle::GetBrush("ContentBrowser.AssetTreeFolderOpenPlugin");
 		FolderClosedPluginBrush = FAppStyle::GetBrush("ContentBrowser.AssetTreeFolderClosedPlugin");
 	}
@@ -104,9 +108,17 @@ void SAssetTreeItem::Construct( const FArguments& InArgs )
 	{
 		FolderType = EFolderType::Code;
 	}
-	else if (EnumHasAnyFlags(InArgs._TreeItem->GetItem().GetItemCategory(), EContentBrowserItemFlags::Category_Plugin))
+	else if (InArgs._TreeItem->GetItem().IsInPlugin())
 	{
-		FolderType = EFolderType::Plugin;
+		TSharedPtr<FTreeItem> Parent = InArgs._TreeItem->Parent.Pin();
+		if (Parent.IsValid() && Parent->GetItem().IsInPlugin())
+		{
+			FolderType = EFolderType::PluginSubfolder;
+		}
+		else
+		{
+			FolderType = EFolderType::PluginRoot;
+		}
 	}
 
 	bool bIsRoot = !InArgs._TreeItem->Parent.IsValid();
@@ -304,7 +316,13 @@ const FSlateBrush* SAssetTreeItem::GetFolderIcon() const
 	case EFolderType::Developer:
 		return (IsItemExpanded.Get()) ? Brushes.FolderOpenDeveloperBrush : Brushes.FolderClosedDeveloperBrush;
 
-	case EFolderType::Plugin:
+	case EFolderType::PluginRoot:
+		if (GUsePluginFolderIcon)
+		{
+			return (IsItemExpanded.Get()) ? Brushes.FolderOpenPluginRootBrush : Brushes.FolderClosedPluginRootBrush;
+		}
+		// Fall through to default 
+	case EFolderType::PluginSubfolder:
 		if (GUsePluginFolderIcon)
 		{
 			return (IsItemExpanded.Get()) ? Brushes.FolderOpenPluginBrush : Brushes.FolderClosedPluginBrush;
diff --git a/Engine/Source/Editor/ContentBrowser/Private/SourcesViewWidgets.h b/Engine/Source/Editor/ContentBrowser/Private/SourcesViewWidgets.h
--- a/Engine/Source/Editor/ContentBrowser/Private/SourcesViewWidgets.h
+++ b/Engine/Source/Editor/ContentBrowser/Private/SourcesViewWidgets.h
@@ -108,7 +108,8 @@ private:
 		Normal,
 		Code,
 		Developer,
-		Plugin,
+		PluginSubfolder,
+		PluginRoot,
 	};
 
 	/** The data for this item */
diff --git a/Engine/Source/Editor/EditorStyle/Private/StarshipStyle.cpp b/Engine/Source/Editor/EditorStyle/Private/StarshipStyle.cpp
--- a/Engine/Source/Editor/EditorStyle/Private/StarshipStyle.cpp
+++ b/Engine/Source/Editor/EditorStyle/Private/StarshipStyle.cpp
@@ -6059,6 +6059,8 @@ void FStarshipEditorStyle::FStyle::SetupContentBrowserStyle()
 		Set( "ContentBrowser.AssetTreeFolderClosedDeveloper", new IMAGE_BRUSH_SVG( "Starship/ContentBrowser/folder-developer", Icon16x16));
 		Set( "ContentBrowser.AssetTreeFolderOpenCode", new IMAGE_BRUSH_SVG( "Starship/ContentBrowser/folder-code-open", Icon16x16) );
 		Set( "ContentBrowser.AssetTreeFolderClosedCode", new IMAGE_BRUSH_SVG( "Starship/ContentBrowser/folder-code", Icon16x16) );
+		Set( "ContentBrowser.AssetTreeFolderOpenPluginRoot", new IMAGE_BRUSH_SVG("Starship/ContentBrowser/folder-pluginroot-open", Icon16x16));
+		Set( "ContentBrowser.AssetTreeFolderClosedPluginRoot", new IMAGE_BRUSH_SVG( "Starship/ContentBrowser/folder-pluginroot", Icon16x16));
 		Set( "ContentBrowser.AssetTreeFolderOpenPlugin", new IMAGE_BRUSH_SVG("Starship/ContentBrowser/folder-plugin-open", Icon16x16));
 		Set( "ContentBrowser.AssetTreeFolderClosedPlugin", new IMAGE_BRUSH_SVG( "Starship/ContentBrowser/folder-plugin", Icon16x16));
 
