Modify FCollectionAssetRegistryBridge to bind to OnAssetsRemoved to prepare for optimizations

diff --git a/Engine/Source/Editor/ContentBrowser/Private/CollectionAssetRegistryBridge.cpp b/Engine/Source/Editor/ContentBrowser/Private/CollectionAssetRegistryBridge.cpp
--- a/Engine/Source/Editor/ContentBrowser/Private/CollectionAssetRegistryBridge.cpp
+++ b/Engine/Source/Editor/ContentBrowser/Private/CollectionAssetRegistryBridge.cpp
@@ -116,7 +116,7 @@ FCollectionAssetRegistryBridge::FCollectionAssetRegistryBridge()
 {
 	// Load the asset registry module to listen for updates
 	IAssetRegistry& AssetRegistry = IAssetRegistry::GetChecked();
-	AssetRegistry.OnAssetRemoved().AddRaw(this, &FCollectionAssetRegistryBridge::OnAssetRemoved);
+	AssetRegistry.OnAssetsRemoved().AddRaw(this, &FCollectionAssetRegistryBridge::OnAssetsRemoved);
 	AssetRegistry.OnAssetRenamed().AddRaw(this, &FCollectionAssetRegistryBridge::OnAssetRenamed);
 	
 	if (AssetRegistry.IsLoadingAssets())
@@ -134,7 +134,7 @@ FCollectionAssetRegistryBridge::~FCollectionAssetRegistryBridge()
 	// Load the asset registry module to unregister delegates
 	if (IAssetRegistry* AssetRegistry = IAssetRegistry::Get())
 	{
-		AssetRegistry->OnAssetRemoved().RemoveAll(this);
+		AssetRegistry->OnAssetsRemoved().RemoveAll(this);
 		AssetRegistry->OnAssetRenamed().RemoveAll(this);
 		AssetRegistry->OnFilesLoaded().RemoveAll(this);
 	}
@@ -158,20 +158,23 @@ void FCollectionAssetRegistryBridge::OnAssetRenamed(const FAssetData& AssetData,
 	CollectionManagerModule.Get().HandleObjectRenamed(FSoftObjectPath(OldObjectPath), AssetData.GetSoftObjectPath());
 }
 
-void FCollectionAssetRegistryBridge::OnAssetRemoved(const FAssetData& AssetData)
+void FCollectionAssetRegistryBridge::OnAssetsRemoved(TConstArrayView<FAssetData> AssetDatas)
 {
 	FCollectionManagerModule& CollectionManagerModule = FCollectionManagerModule::GetModule();
 
-	if (AssetData.IsRedirector())
+	for (const FAssetData& AssetData : AssetDatas)
 	{
-		// Notify the collections manager that a redirector has been removed
-		// This will attempt to re-save any collections that still have a reference to this redirector in their on-disk collection data
-		CollectionManagerModule.Get().HandleRedirectorDeleted(AssetData.GetSoftObjectPath());
-	}
-	else
-	{
-		// Notify the collections manager that an asset has been removed
-		CollectionManagerModule.Get().HandleObjectDeleted(AssetData.GetSoftObjectPath());
+		if (AssetData.IsRedirector())
+		{
+			// Notify the collections manager that a redirector has been removed
+			// This will attempt to re-save any collections that still have a reference to this redirector in their on-disk collection data
+			CollectionManagerModule.Get().HandleRedirectorDeleted(AssetData.GetSoftObjectPath());
+		}
+		else
+		{
+			// Notify the collections manager that an asset has been removed
+			CollectionManagerModule.Get().HandleObjectDeleted(AssetData.GetSoftObjectPath());
+		}
 	}
 }
 
diff --git a/Engine/Source/Editor/ContentBrowser/Private/CollectionAssetRegistryBridge.h b/Engine/Source/Editor/ContentBrowser/Private/CollectionAssetRegistryBridge.h
--- a/Engine/Source/Editor/ContentBrowser/Private/CollectionAssetRegistryBridge.h
+++ b/Engine/Source/Editor/ContentBrowser/Private/CollectionAssetRegistryBridge.h
@@ -22,6 +22,6 @@ private:
 	/** Handler for when an asset was renamed in the asset registry */
 	void OnAssetRenamed(const FAssetData& AssetData, const FString& OldObjectPath);
 
-	/** Handler for when an asset was removed from the asset registry */
-	void OnAssetRemoved(const FAssetData& AssetData);
+	/** Handler for when assets were removed from the asset registry */
+	void OnAssetsRemoved(TConstArrayView<FAssetData> AssetDatas);
 };
