Fix the content browser loosing it path filter values each time it is opened.

diff --git a/Engine/Source/Editor/ContentBrowser/Private/SPathView.cpp b/Engine/Source/Editor/ContentBrowser/Private/SPathView.cpp
--- a/Engine/Source/Editor/ContentBrowser/Private/SPathView.cpp
+++ b/Engine/Source/Editor/ContentBrowser/Private/SPathView.cpp
@@ -183,11 +183,6 @@ void SPathView::Construct( const FArguments& InArgs )
 				Delegate.Execute(AllPluginPathFilters);
 			}
 		}
-
-		for (const TSharedRef<FContentBrowserPluginFilter>& Filter : AllPluginPathFilters)
-		{
-			SetPluginPathFilterActive(Filter, false);
-		}
 	}
 
 	if (!TreeViewPtr.IsValid())
