Plugin content folder icon for asset tile view.

diff --git a/Engine/Source/Editor/ContentBrowser/Private/AssetViewWidgets.cpp b/Engine/Source/Editor/ContentBrowser/Private/AssetViewWidgets.cpp
--- a/Engine/Source/Editor/ContentBrowser/Private/AssetViewWidgets.cpp
+++ b/Engine/Source/Editor/ContentBrowser/Private/AssetViewWidgets.cpp
@@ -191,6 +191,7 @@ TSharedRef<SWidget> FAssetViewItemHelper::CreateListTileItemContents(T* const In
 
 		const bool bDeveloperFolder = ContentBrowserUtils::IsItemDeveloperContent(InTileOrListItem->AssetItem->GetItem());
 		const bool bCodeFolder = EnumHasAnyFlags(InTileOrListItem->AssetItem->GetItem().GetItemCategory(), EContentBrowserItemFlags::Category_Class);
+		const bool bPluginFolder = EnumHasAnyFlags(InTileOrListItem->AssetItem->GetItem().GetItemCategory(), EContentBrowserItemFlags::Category_Plugin);
 
 		const bool bCollectionFolder = EnumHasAnyFlags(InTileOrListItem->AssetItem->GetItem().GetItemCategory(), EContentBrowserItemFlags::Category_Collection);
 		ECollectionShareType::Type CollectionFolderShareType = ECollectionShareType::CST_All;
@@ -199,11 +200,23 @@ TSharedRef<SWidget> FAssetViewItemHelper::CreateListTileItemContents(T* const In
 			ContentBrowserUtils::IsCollectionPath(InTileOrListItem->AssetItem->GetItem().GetVirtualPath().ToString(), nullptr, &CollectionFolderShareType);
 		}
 
-		const FSlateBrush* FolderBaseImage = bDeveloperFolder
-			? FAppStyle::GetBrush("ContentBrowser.ListViewDeveloperFolderIcon") 
-			: bCodeFolder
-				? FAppStyle::GetBrush("ContentBrowser.ListViewCodeFolderIcon")
-				: FAppStyle::GetBrush("ContentBrowser.ListViewFolderIcon");
+		const FSlateBrush* FolderBaseImage = nullptr;
+		if (bDeveloperFolder)
+		{
+			FolderBaseImage = FAppStyle::GetBrush("ContentBrowser.ListViewDeveloperFolderIcon");
+		}
+		else if (bCodeFolder)
+		{
+			FolderBaseImage = FAppStyle::GetBrush("ContentBrowser.ListViewCodeFolderIcon");
+		}
+		else if (bPluginFolder)
+		{
+			FolderBaseImage = FAppStyle::GetBrush("ContentBrowser.ListViewPluginFolderIcon");
+		}
+		else
+		{
+			FolderBaseImage = FAppStyle::GetBrush("ContentBrowser.ListViewFolderIcon");
+		}
 
 		// Folder base
 		ItemContentsOverlay->AddSlot()
diff --git a/Engine/Source/Editor/EditorStyle/Private/StarshipStyle.cpp b/Engine/Source/Editor/EditorStyle/Private/StarshipStyle.cpp
--- a/Engine/Source/Editor/EditorStyle/Private/StarshipStyle.cpp
+++ b/Engine/Source/Editor/EditorStyle/Private/StarshipStyle.cpp
@@ -6310,7 +6310,8 @@ void FStarshipEditorStyle::FStyle::SetupContentBrowserStyle()
 		Set( "ContentBrowser.ListViewFolderIcon", new IMAGE_BRUSH_SVG( "Starship/ContentBrowser/folder", FVector2D(64, 64) ) );
 		Set( "ContentBrowser.ListViewDeveloperFolderIcon", new IMAGE_BRUSH_SVG("Starship/ContentBrowser/FolderLargeDeveloper", FVector2D(64, 64)));
 		Set( "ContentBrowser.ListViewCodeFolderIcon", new IMAGE_BRUSH_SVG("Starship/ContentBrowser/FolderLargeCode", FVector2D(64, 64)));
-
+		Set( "ContentBrowser.ListViewPluginFolderIcon", new IMAGE_BRUSH_SVG("Starship/ContentBrowser/FolderLargePlugin", FVector2D(64, 64)));
+		
 		Set("ContentBrowser.AssetTileItem.FolderAreaHoveredBackground", new FSlateRoundedBoxBrush(FStyleColors::Secondary, 4.0f));
 		Set("ContentBrowser.AssetTileItem.FolderAreaSelectedBackground", new FSlateRoundedBoxBrush(FStyleColors::Secondary, 4.0f));
 		Set("ContentBrowser.AssetTileItem.FolderAreaSelectedHoverBackground", new FSlateRoundedBoxBrush(FStyleColors::Secondary, 4.0f));
