SAssetView: Change Content Browser column resize mode to fill so there is no cap on resizing columns

diff --git a/Engine/Source/Editor/ContentBrowser/Private/SAssetView.cpp b/Engine/Source/Editor/ContentBrowser/Private/SAssetView.cpp
--- a/Engine/Source/Editor/ContentBrowser/Private/SAssetView.cpp
+++ b/Engine/Source/Editor/ContentBrowser/Private/SAssetView.cpp
@@ -1764,7 +1764,7 @@ TSharedRef<SAssetColumnView> SAssetView::CreateColumnView()
 		.HeaderRow
 		(
 			SNew(SHeaderRow)
-			.ResizeMode(ESplitterResizeMode::FixedSize)
+			.ResizeMode(ESplitterResizeMode::Fill)
 			.CanSelectGeneratedColumn(true)
 			.OnHiddenColumnsListChanged(this, &SAssetView::OnHiddenColumnsChanged)
 
