SNavigation bar control for content browser path navigation.
Combo box with history dropdown.
Conversion to edit text box by clicking in blank space or F4 key.
Auto-complete for suffixes of currently typed path.

diff --git a/Engine/Source/Editor/ContentBrowser/Private/ContentBrowserCommands.cpp b/Engine/Source/Editor/ContentBrowser/Private/ContentBrowserCommands.cpp
--- a/Engine/Source/Editor/ContentBrowser/Private/ContentBrowserCommands.cpp
+++ b/Engine/Source/Editor/ContentBrowser/Private/ContentBrowserCommands.cpp
@@ -20,6 +20,7 @@ void FContentBrowserCommands::RegisterCommands()
 	UI_COMMAND(SaveAllCurrentFolder, "Save All", "Save All in current folder", EUserInterfaceActionType::Button, FInputChord(EModifierKey::Control | EModifierKey::Shift, EKeys::S));
 	UI_COMMAND(ResaveAllCurrentFolder, "Resave All", "Resave all assets contained in the current folder", EUserInterfaceActionType::Button, FInputChord());
 	UI_COMMAND(CopySelectedAssetPath, "Copy Selected Asset Path", "Copy the selected asset path", EUserInterfaceActionType::Button, FInputChord(EModifierKey::Control, EKeys::C));
+	UI_COMMAND(EditPath, "Edit Path", "Edit the current content browser path", EUserInterfaceActionType::Button, FInputChord(EKeys::F4));
 }
 
 
diff --git a/Engine/Source/Editor/ContentBrowser/Private/ContentBrowserCommands.h b/Engine/Source/Editor/ContentBrowser/Private/ContentBrowserCommands.h
--- a/Engine/Source/Editor/ContentBrowser/Private/ContentBrowserCommands.h
+++ b/Engine/Source/Editor/ContentBrowser/Private/ContentBrowserCommands.h
@@ -37,4 +37,5 @@ public:
 	TSharedPtr<FUICommandInfo> SaveAllCurrentFolder;
 	TSharedPtr<FUICommandInfo> ResaveAllCurrentFolder;
 	TSharedPtr<FUICommandInfo> CopySelectedAssetPath;
+	TSharedPtr<FUICommandInfo> EditPath;
 };
diff --git a/Engine/Source/Editor/ContentBrowser/Private/SContentBrowser.cpp b/Engine/Source/Editor/ContentBrowser/Private/SContentBrowser.cpp
--- a/Engine/Source/Editor/ContentBrowser/Private/SContentBrowser.cpp
+++ b/Engine/Source/Editor/ContentBrowser/Private/SContentBrowser.cpp
@@ -94,6 +94,7 @@
 #include "SAssetView.h"
 #include "SCollectionView.h"
 #include "SFilterList.h"
+#include "SNavigationBar.h"
 #include "SPathView.h"
 #include "SPositiveActionButton.h"
 #include "SSearchToggleButton.h"
@@ -104,6 +105,7 @@
 #include "SourcesData.h"
 #include "SourcesSearch.h"
 #include "StatusBarSubsystem.h"
+#include "String/Find.h"
 #include "Styling/AppStyle.h"
 #include "Styling/ISlateStyle.h"
 #include "Styling/SlateColor.h"
@@ -167,6 +169,116 @@ namespace ContentBrowser
 	// It is useful when debugging enumeration issues to shut this system off so that breakpoints set will not be triggered every tick
 	bool bCrumbsEnumerate = true;
 	FAutoConsoleVariableRef CrumbsEnumerateCVar(TEXT("ContentBrowser.Debug.CrumbsEnumerate"), bCrumbsEnumerate, TEXT("Enumerate crumbs"), ECVF_Default);
+	
+	const FContentBrowserInstanceConfig* GetConstInstanceConfig(FName ForInstance) 
+	{
+		if (ForInstance.IsNone())
+		{
+			return nullptr;
+		}
+
+		UContentBrowserConfig* Config = UContentBrowserConfig::Get();
+		if (Config == nullptr)
+		{
+			return nullptr;
+		}
+
+		const FContentBrowserInstanceConfig* InstanceConfig = Config->Instances.Find(ForInstance);
+		return InstanceConfig;
+	}
+
+	// Returns child items that can be navigated to from Path e.g. subfolders
+	TArray<FContentBrowserItem> GetChildItemsFromVirtualPath(
+		FName Path,
+		EContentBrowserItemCategoryFilter ItemCategoryFilter, 
+		EContentBrowserItemAttributeFilter ItemAttributeFilter,
+		FName ConfigInstanceName,
+		SPathView& PathViewForFiltering 
+		)
+	{
+		const UContentBrowserSettings* ContentBrowserSettings = GetDefault<UContentBrowserSettings>();
+		bool bDisplayEmpty = ContentBrowserSettings->DisplayEmptyFolders;
+		// check to see if we have an instance config that overrides the default in UContentBrowserSettings
+		if (const FContentBrowserInstanceConfig* EditorConfig = GetConstInstanceConfig(ConfigInstanceName))
+		{
+			bDisplayEmpty = EditorConfig->bShowEmptyFolders;
+		}
+
+		UContentBrowserDataSubsystem* ContentBrowserData = IContentBrowserDataModule::Get().GetSubsystem();
+
+		FContentBrowserDataFilter SubItemsFilter;
+		SubItemsFilter.ItemTypeFilter = EContentBrowserItemTypeFilter::IncludeFolders;
+		SubItemsFilter.bRecursivePaths = false;
+		SubItemsFilter.ItemCategoryFilter = ItemCategoryFilter;
+		SubItemsFilter.ItemAttributeFilter = ItemAttributeFilter;
+
+		TArray<FContentBrowserItem> SubItems = ContentBrowserData->GetItemsUnderPath(Path, SubItemsFilter);
+
+		for (auto It = SubItems.CreateIterator(); It; ++It)
+		{
+			const FContentBrowserItem& Item = *It;
+			if (!Item.GetInternalPath().IsNone())
+			{
+				if (!PathViewForFiltering.InternalPathPassesBlockLists(FNameBuilder(Item.GetInternalPath())))
+				{
+					It.RemoveCurrent();
+					continue;
+				}
+			}
+			else
+			{
+				// Test if any child internal paths pass for this fully virtual path
+				bool bPasses = false;
+				for (const FContentBrowserItemData& ItemData : Item.GetInternalItems())
+				{
+					UContentBrowserDataSource* ItemDataSource = ItemData.GetOwnerDataSource();
+					if (!ItemDataSource) 
+					{
+						continue;
+					}
+
+					auto OnSubPath = [&PathViewForFiltering, &bPasses, &SubItemsFilter](FName VirtualSubPath, FName InternalPath)
+					{
+						if (InternalPath.IsNone())
+						{
+							// Keep enumerating, this path only exist virtually e.g. /All/Plugins
+							return true;
+						}
+
+						const FNameBuilder InternalPathBuilder(InternalPath);
+						if (ContentBrowserDataUtils::PathPassesAttributeFilter(InternalPathBuilder, 0, SubItemsFilter.ItemAttributeFilter) 
+						&& PathViewForFiltering.InternalPathPassesBlockLists(InternalPathBuilder))
+						{
+							bPasses = true;
+							// Stop enumerating
+							return false;
+						}
+						return true;
+					}; 
+					ItemDataSource->GetRootPathVirtualTree().EnumerateSubPaths(Item.GetVirtualPath(), OnSubPath, /*bRecurse*/ true);
+
+					if (bPasses)
+					{
+						break;
+					}
+				}
+
+				if (!bPasses)
+				{
+					It.RemoveCurrent();
+					continue;
+				}
+			}
+
+			if (!ContentBrowserData->IsFolderVisible(Item.GetVirtualPath(), ContentBrowserUtils::GetIsFolderVisibleFlags(bDisplayEmpty)))
+			{
+				It.RemoveCurrent();
+				continue;
+			}
+		}
+
+		return SubItems;
+	}
 }
 
 SContentBrowser::~SContentBrowser()
@@ -205,6 +317,8 @@ void SContentBrowser::Construct( const FArguments& InArgs, const FName& InInstan
 {
 	InstanceName = InInstanceName;
 
+	JumpMRU.MaxItems = 30;
+
 	UContentBrowserConfig::Initialize();
 	UContentBrowserConfig::Get()->LoadEditorConfig();
 	const FContentBrowserInstanceConfig* EditorConfig = CreateEditorConfigIfRequired();
@@ -383,46 +497,25 @@ void SContentBrowser::Construct( const FArguments& InArgs, const FName& InInstan
 						.ColorAndOpacity(FSlateColor::UseForeground())
 					]
 				]
-				// Path picker
-				+ SHorizontalBox::Slot()
-				.Padding(2, 0, 0, 0)
-				.AutoWidth()
-				.VAlign(VAlign_Fill)
-				[
-					SAssignNew( PathPickerButton, SComboButton )
-					.Visibility( ( Config != nullptr ? Config->bUsePathPicker : true ) ? EVisibility::Visible : EVisibility::Collapsed )
-					.ButtonStyle(FAppStyle::Get(), "SimpleButton")
-					.ToolTipText( LOCTEXT( "PathPickerTooltip", "Choose a path" ) )
-					.OnGetMenuContent( this, &SContentBrowser::GetPathPickerContent )
-					.HasDownArrow( false )
-					.AddMetaData<FTagMetaData>(FTagMetaData(TEXT("ContentBrowserPathPicker")))
-					.ContentPadding(FMargin(1, 0))
-					.ButtonContent()
-					[
-						SNew(SImage)
-						.Image(FAppStyle::Get().GetBrush("Icons.FolderClosed"))
-						.ColorAndOpacity(FSlateColor::UseForeground())
-					]
-				]
-
+				
 				// Path
 				+ SHorizontalBox::Slot()
 				.VAlign(VAlign_Center)
-				.HAlign(HAlign_Left)
+				.HAlign(HAlign_Fill)
 				.FillWidth(1.0f)
 				.Padding(2, 0, 0, 0)
 				[
-					SAssignNew(PathBreadcrumbTrail, SBreadcrumbTrail<FString>)
-					.ButtonContentPadding(FMargin(2, 2))
-					.ButtonStyle(FAppStyle::Get(), "SimpleButton")
-					.DelimiterImage(FAppStyle::Get().GetBrush("Icons.ChevronRight"))
-					.TextStyle(FAppStyle::Get(), "NormalText")
-					.ShowLeadingDelimiter(false)
-					.OnCrumbClicked(this, &SContentBrowser::OnPathClicked)
-					.HasCrumbMenuContent(this, &SContentBrowser::OnHasCrumbDelimiterContent)
-					.GetCrumbMenuContent(this, &SContentBrowser::OnGetCrumbDelimiterContent)
+					SAssignNew(NavigationBar, SNavigationBar)
+					.OnPathClicked(this, &SContentBrowser::OnPathClicked)
+					.HasPathMenuContent(this, &SContentBrowser::OnHasCrumbDelimiterContent)
+					.GetPathMenuContent(this, &SContentBrowser::OnGetCrumbDelimiterContent)
+					.GetComboOptions(this, &SContentBrowser::GetRecentPaths)
+					.OnNavigateToPath(this, &SContentBrowser::OnNavigateToPath)
+					.OnCompletePrefix(this, &SContentBrowser::OnCompletePathPrefix)
+					.OnCanEditPathAsText(this, &SContentBrowser::OnCanEditPathAsText)
 					.AddMetaData<FTagMetaData>(FTagMetaData(TEXT("ContentBrowserPath")))
 				]
+
 				+ SHorizontalBox::Slot()
 				.AutoWidth()
 				.HAlign(HAlign_Right)
@@ -666,6 +759,10 @@ void SContentBrowser::BindCommands()
 	Commands->MapAction(FContentBrowserCommands::Get().CopySelectedAssetPath, FUIAction(
 		FExecuteAction::CreateSP(this, &SContentBrowser::CopySelectedAssetPathCommand)
 	));
+	
+	Commands->MapAction(FContentBrowserCommands::Get().EditPath, FUIAction(
+		FExecuteAction::CreateSP(this, &SContentBrowser::EditPathCommand)
+	));
 
 	// Allow extenders to add commands
 	FContentBrowserModule& ContentBrowserModule = FModuleManager::GetModuleChecked<FContentBrowserModule>(TEXT("ContentBrowser"));
@@ -1869,6 +1966,8 @@ void SContentBrowser::SaveSettings() const
 	FavoritePathViewPtr->SaveSettings(GEditorPerProjectIni, SettingsIniSection, SettingsString + TEXT(".Favorites"));
 	CollectionViewPtr->SaveSettings(GEditorPerProjectIni, SettingsIniSection, SettingsString);
 	AssetViewPtr->SaveSettings(GEditorPerProjectIni, SettingsIniSection, SettingsString);
+	
+	GConfig->SetArray(*SettingsIniSection, *(SettingsString + TEXT(".JumpMRU")), JumpMRU, GEditorPerProjectIni);
 }
 
 const FName SContentBrowser::GetInstanceName() const
@@ -2019,6 +2118,8 @@ void SContentBrowser::LoadSettings(const FName& InInstanceName)
 	FavoritePathViewPtr->LoadSettings(GEditorPerProjectIni, SettingsIniSection, SettingsString + TEXT(".Favorites"));
 	CollectionViewPtr->LoadSettings(GEditorPerProjectIni, SettingsIniSection, SettingsString);
 	AssetViewPtr->LoadSettings(GEditorPerProjectIni, SettingsIniSection, SettingsString);
+
+	GConfig->GetArray(*SettingsIniSection, *(SettingsString + TEXT(".JumpMRU")), JumpMRU, GEditorPerProjectIni);
 }
 
 void SContentBrowser::SourcesChanged(const TArray<FString>& SelectedPaths, const TArray<FCollectionNameType>& SelectedCollections)
@@ -2111,6 +2212,8 @@ void SContentBrowser::FolderEntered(const FString& FolderPath)
 
 void SContentBrowser::PathSelected(const FString& FolderPath)
 {
+	JumpMRU.AddUnique(FolderPath);
+
 	// You may not select both collections and paths
 	CollectionViewPtr->ClearSelection();
 
@@ -2134,6 +2237,8 @@ void SContentBrowser::PathSelected(const FString& FolderPath)
 
 void SContentBrowser::FavoritePathSelected(const FString& FolderPath)
 {
+	JumpMRU.AddUnique(FolderPath);
+	
 	// You may not select both collections and paths
 	CollectionViewPtr->ClearSelection();
 
@@ -2388,6 +2493,24 @@ void SContentBrowser::SaveSearchAsFilter()
 	FilterListPtr->CreateCustomFilterDialog(TextFilter->GetRawFilterText());
 }
 
+void SContentBrowser::EditPathCommand()
+{
+	NavigationBar->StartEditingPath();
+}
+
+void SContentBrowser::OnNavigateToPath(const FString& NewPath)
+{
+	// If NewPath is not a valid location, the content browser will stay in its current location.
+	TArray<FName> VirtualPaths;
+	FStringView PathView = NewPath;
+	if (PathView.EndsWith('/'))
+	{
+		PathView.LeftChopInline(1);
+	}
+	VirtualPaths.Add(FName(PathView)); 
+	SyncToVirtualPaths(VirtualPaths);
+}
+
 void SContentBrowser::OnPathClicked( const FString& CrumbData )
 {
 	FSourcesData SourcesData = AssetViewPtr->GetSourcesData();
@@ -2433,9 +2556,58 @@ void SContentBrowser::OnPathClicked( const FString& CrumbData )
 	}
 }
 
+TArray<FString> SContentBrowser::GetRecentPaths() const
+{
+	return JumpMRU;	
+}
+
 void SContentBrowser::OnPathMenuItemClicked(FString ClickedPath)
 {
-	OnPathClicked( ClickedPath );
+	OnPathClicked(ClickedPath);
+}
+
+bool SContentBrowser::OnCanEditPathAsText(const FString& Text) const
+{
+	FSourcesData SourcesData = AssetViewPtr->GetSourcesData();
+	if (SourcesData.HasCollections())
+	{
+		// Do not present collections as text because their names are not very user friendly right now.
+		return false;
+	}
+	return true;
+}
+
+TArray<FString> SContentBrowser::OnCompletePathPrefix(const FString& Prefix) const
+{
+	FStringView PrefixView = Prefix;
+	
+	// Strip to last path separator
+	FName Parent;
+	if (int32 Index = UE::String::FindLastChar(PrefixView, '/'); Index != INDEX_NONE)
+	{	
+		PrefixView.LeftInline(Index);
+		Parent = FName(PrefixView);
+	}
+
+	// Find PrefixView in the available tree of data sources, get its direct children, and filter them by SuffixView 
+	UContentBrowserDataSubsystem* ContentBrowserData = IContentBrowserDataModule::Get().GetSubsystem();
+	TArray<FContentBrowserItem> SubItems = ContentBrowser::GetChildItemsFromVirtualPath(
+		Parent,
+		PathViewPtr->GetContentBrowserItemCategoryFilter(),
+		PathViewPtr->GetContentBrowserItemAttributeFilter(),
+		InstanceName,
+		*PathViewPtr); 
+	TArray<FString> Results;
+	for (const FContentBrowserItem& Item : SubItems)
+	{
+		FName Path = Item.GetVirtualPath();
+		FNameBuilder PathBuilder(Path);
+		if (PathBuilder.ToView().StartsWith(Prefix))
+		{
+			Results.Add(Item.GetVirtualPath().ToString());
+		}
+	}
+	return Results;
 }
 
 bool SContentBrowser::OnHasCrumbDelimiterContent(const FString& CrumbData) const
@@ -2546,81 +2718,12 @@ TSharedRef<SWidget> SContentBrowser::OnGetCrumbDelimiterContent(const FString& C
 	}
 	else if( SourcesData.HasVirtualPaths() )
 	{
-		const UContentBrowserSettings* ContentBrowserSettings = GetDefault<UContentBrowserSettings>();
-		bool bDisplayEmpty = ContentBrowserSettings->DisplayEmptyFolders;
-		// check to see if we have an instance config that overrides the default in UContentBrowserSettings
-		if (const FContentBrowserInstanceConfig* EditorConfig = GetConstInstanceConfig())
-		{
-			bDisplayEmpty = EditorConfig->bShowEmptyFolders;
-		}
-
-		UContentBrowserDataSubsystem* ContentBrowserData = IContentBrowserDataModule::Get().GetSubsystem();
-
-		FContentBrowserDataFilter SubItemsFilter;
-		SubItemsFilter.ItemTypeFilter = EContentBrowserItemTypeFilter::IncludeFolders;
-		SubItemsFilter.bRecursivePaths = false;
-		SubItemsFilter.ItemCategoryFilter = PathViewPtr->GetContentBrowserItemCategoryFilter();
-		SubItemsFilter.ItemAttributeFilter = PathViewPtr->GetContentBrowserItemAttributeFilter();
-
-		TArray<FContentBrowserItem> SubItems = ContentBrowserData->GetItemsUnderPath(*CrumbData, SubItemsFilter);
-
-		for (auto It = SubItems.CreateIterator(); It; ++It)
-		{
-			const FContentBrowserItem& Item = *It;
-			if (!Item.GetInternalPath().IsNone())
-			{
-				if (!PathViewPtr->InternalPathPassesBlockLists(FNameBuilder(Item.GetInternalPath())))
-				{
-					It.RemoveCurrent();
-					continue;
-				}
-			}
-			else
-			{
-				// Test if any child internal paths pass for this fully virtual path
-				bool bPasses = false;
-				for (const FContentBrowserItemData& ItemData : Item.GetInternalItems())
-				{
-					if (UContentBrowserDataSource* ItemDataSource = ItemData.GetOwnerDataSource())
-					{
-						ItemDataSource->GetRootPathVirtualTree().EnumerateSubPaths(Item.GetVirtualPath(), [this, &bPasses, &SubItemsFilter](FName VirtualSubPath, FName InternalPath)
-						{
-							if (!InternalPath.IsNone())
-							{
-								const FNameBuilder InternalPathBuilder(InternalPath);
-								if (ContentBrowserDataUtils::PathPassesAttributeFilter(InternalPathBuilder, 0, SubItemsFilter.ItemAttributeFilter) &&
-									PathViewPtr->InternalPathPassesBlockLists(InternalPathBuilder))
-								{
-									bPasses = true;
-									// Stop enumerating
-									return false;
-								}
-							}
-							// Keep enumerating
-							return true;
-						}, /*bRecurse*/ true);
-
-						if (bPasses)
-						{
-							break;
-						}
-					}
-				}
-
-				if (!bPasses)
-				{
-					It.RemoveCurrent();
-					continue;
-				}
-			}
-
-			if (!ContentBrowserData->IsFolderVisible(Item.GetVirtualPath(), ContentBrowserUtils::GetIsFolderVisibleFlags(bDisplayEmpty)))
-			{
-				It.RemoveCurrent();
-				continue;
-			}
-		}
-
+		TArray<FContentBrowserItem> SubItems = ContentBrowser::GetChildItemsFromVirtualPath(
+			*CrumbData,
+			PathViewPtr->GetContentBrowserItemCategoryFilter(),
+			PathViewPtr->GetContentBrowserItemAttributeFilter(),
+			InstanceName,
+			*PathViewPtr); 
 		SubItems.Sort([](const FContentBrowserItem& ItemOne, const FContentBrowserItem& ItemTwo)
 		{
 			return ItemOne.GetDisplayName().CompareTo(ItemTwo.GetDisplayName()) < 0;
@@ -2659,49 +2762,6 @@ TSharedRef<SWidget> SContentBrowser::OnGetCrumbDelimiterContent(const FString& C
 	return Widget.ToSharedRef();
 }
 
-TSharedRef<SWidget> SContentBrowser::GetPathPickerContent()
-{
-	FPathPickerConfig PathPickerConfig;
-
-	FSourcesData SourcesData = AssetViewPtr->GetSourcesData();
-	if ( SourcesData.HasVirtualPaths() )
-	{
-		PathPickerConfig.DefaultPath = SourcesData.VirtualPaths[0].ToString();
-	}
-	
-	// TODO: This needs to be able to pick any content folder, so needs to use the new item-based API
-	PathPickerConfig.OnPathSelected = FOnPathSelected::CreateSP(this, &SContentBrowser::PathPickerPathSelected);
-	PathPickerConfig.bAllowContextMenu = false;
-	PathPickerConfig.bAllowClassesFolder = true;
-	PathPickerConfig.bOnPathSelectedPassesVirtualPaths = true;
-
-	return SNew(SBox)
-		.WidthOverride(300.f)
-		.HeightOverride(500.f)
-		.Padding(4.f)
-		[
-			SNew(SVerticalBox)
-
-			// Path Picker
-			+SVerticalBox::Slot()
-			.FillHeight(1.f)
-			[
-				FContentBrowserSingleton::Get().CreatePathPicker(PathPickerConfig)
-			]
-
-			// Collection View
-			+SVerticalBox::Slot()
-			.AutoHeight()
-			.Padding(0.f, 6.f, 0.f, 0.f)
-			[
-				SNew(SCollectionView)
-				.AllowCollectionButtons(false)
-				.OnCollectionSelected(this, &SContentBrowser::PathPickerCollectionSelected)
-				.AllowContextMenu(false)
-			]
-		];
-}
-
 FString SContentBrowser::GetCurrentPath(const EContentBrowserPathType PathType) const
 {
 	FString CurrentPath;
@@ -3624,7 +3684,7 @@ void SContentBrowser::UpdatePath()
 {
 	FSourcesData SourcesData = AssetViewPtr->GetSourcesData();
 
-	PathBreadcrumbTrail->ClearCrumbs();
+	NavigationBar->ClearPaths();
 
 	int32 NewSourcesWidgetIndex = ActiveSourcesWidgetIndex;
 
@@ -3643,7 +3703,7 @@ void SContentBrowser::UpdatePath()
 			CrumbPath += Crumb;
 
 			const FContentBrowserItem CrumbFolderItem = ContentBrowserData->GetItemAtPath(*CrumbPath, EContentBrowserItemTypeFilter::IncludeFolders);
-			PathBreadcrumbTrail->PushCrumb(CrumbFolderItem.IsValid() ? CrumbFolderItem.GetDisplayName() : FText::FromString(Crumb), CrumbPath);
+			NavigationBar->PushPath(CrumbFolderItem.IsValid() ? CrumbFolderItem.GetDisplayName() : FText::FromString(Crumb), CrumbPath);
 
 			CrumbPath += TEXT("/");
 		}
@@ -3673,12 +3733,12 @@ void SContentBrowser::UpdatePath()
 			Args.Add(TEXT("CollectionName"), FText::FromName(CollectionPathItem.Name));
 			const FText DisplayName = FText::Format(LOCTEXT("CollectionPathIndicator", "{CollectionName} (Collection)"), Args);
 
-			PathBreadcrumbTrail->PushCrumb(DisplayName, CrumbData);
+			NavigationBar->PushPath(DisplayName, CrumbData);
 		}
 	}
 	else
 	{
-		PathBreadcrumbTrail->PushCrumb(LOCTEXT("AllAssets", "All Assets"), TEXT(""));
+		NavigationBar->PushPath(LOCTEXT("AllAssets", "All Assets"), TEXT(""));
 	}
 
 	if (ActiveSourcesWidgetIndex != NewSourcesWidgetIndex)
@@ -4381,19 +4441,7 @@ void SContentBrowser::OpenNewContentBrowser()
  
 const FContentBrowserInstanceConfig* SContentBrowser::GetConstInstanceConfig() const
 {
-	if (InstanceName.IsNone())
-	{
-		return nullptr;
-	}
-
-	UContentBrowserConfig* Config = UContentBrowserConfig::Get();
-	if (Config == nullptr)
-	{
-		return nullptr;
-	}
-
-	const FContentBrowserInstanceConfig* InstanceConfig = Config->Instances.Find(InstanceName);
-	return InstanceConfig;
+	return ContentBrowser::GetConstInstanceConfig(InstanceName);
 }
 
 FContentBrowserInstanceConfig* SContentBrowser::GetMutableInstanceConfig()
diff --git a/Engine/Source/Editor/ContentBrowser/Private/SContentBrowser.h b/Engine/Source/Editor/ContentBrowser/Private/SContentBrowser.h
--- a/Engine/Source/Editor/ContentBrowser/Private/SContentBrowser.h
+++ b/Engine/Source/Editor/ContentBrowser/Private/SContentBrowser.h
@@ -4,6 +4,7 @@
 
 #include "Containers/Array.h"
 #include "Containers/ArrayView.h"
+#include "Containers/MRUArray.h"
 #include "Containers/UnrealString.h"
 #include "ContentBrowserDataMenuContexts.h"
 #include "ContentBrowserDataSubsystem.h"
@@ -41,6 +42,7 @@ class SComboButton;
 class SExpandableArea;
 class SFavoritePathView;
 class SFilterList;
+class SNavigationBar;
 class SPathView;
 class SSearchToggleButton;
 class SWidget;
@@ -307,13 +309,16 @@ private:
 
 	/** Save the current search as a filter pill */
 	void SaveSearchAsFilter();
+	
+	/** Binding for begining to edit text in the location bar */
+	void EditPathCommand();
 
 	/** Called when a crumb in the path breadcrumb trail or menu is clicked */
-	void OnPathClicked(const FString& CrumbData);
-
+	void OnPathClicked(const FString& VirtualPath);
+	
 	/** Called when item in the path delimiter arrow menu is clicked */
 	void OnPathMenuItemClicked(FString ClickedPath);
-
+	
 	/** Called to query whether the crumb menu would contain any content (see also OnGetCrumbDelimiterContent) */
 	bool OnHasCrumbDelimiterContent(const FString& CrumbData) const;
 
@@ -325,8 +330,20 @@ private:
 	 */
 	TSharedRef<SWidget> OnGetCrumbDelimiterContent(const FString& CrumbData) const;
 
-	/** Gets the content for the path picker combo button */
-	TSharedRef<SWidget> GetPathPickerContent();
+	/** Return a list of recently visited locations which can be used with OnNavigateToPath to return to those locations */
+	TArray<FString> GetRecentPaths() const;	
+	
+	/** Navigate to a location by text identifier. E.g. a virtual path such as /All/Game/MyFolder */
+	void OnNavigateToPath(const FString& Path);
+
+	/** 
+	  * Returns whether the currently shown location should be presented to the user as a string when beginning to edit the path, 
+	  * Otherwise the text box will be initially empty.
+	  */
+	bool OnCanEditPathAsText(const FString& Text) const;
+	
+	/** Returns a list of valid paths starting with Prefix which can be navigated to with OnNavigateToPath */
+	TArray<FString> OnCompletePathPrefix(const FString& Prefix) const;
 
 	/** Register the context objects needed for the "Add New" menu */
 	void AppendNewMenuContextObjects(const EContentBrowserDataMenuContext_AddNewMenuDomain InDomain, const TArray<FName>& InSelectedPaths, FToolMenuContext& InOutMenuContext, UContentBrowserToolbarMenuContext* CommonContext, bool bCanBeModified);
@@ -615,6 +632,12 @@ private:
 	/** The manager that keeps track of history data for this browser */
 	FHistoryManager HistoryManager;
 
+	/** A list of locations "jumped" to for populating a dropdown of such locations.
+	 *  As a general rule, simple up/down navigation should not populate this list and only direct entries such as 
+	 *  "find in content browser" or typing in a path should populate it 
+	 */
+	TMRUArray<FString> JumpMRU;
+	
 	/** A helper class to manage asset context menu options */
 	TSharedPtr<class FAssetContextMenu> AssetContextMenu;
 
@@ -641,9 +664,9 @@ private:
 
 	/** The asset view widget */
 	TSharedPtr<SAssetView> AssetViewPtr;
-
-	/** The breadcrumb trail representing the current path */
-	TSharedPtr< SBreadcrumbTrail<FString> > PathBreadcrumbTrail;
+	
+	/** Combine breadcrumb/text-box widget for showing & changing location */
+	TSharedPtr<SNavigationBar> NavigationBar;
 
 	/** The text box used to search for assets */
 	TSharedPtr<SAssetSearchBox> SearchBoxPtr;
diff --git a/Engine/Source/Editor/ContentBrowser/Private/SNavigationBar.cpp b/Engine/Source/Editor/ContentBrowser/Private/SNavigationBar.cpp
new file mode 100644
--- /dev/null
+++ b/Engine/Source/Editor/ContentBrowser/Private/SNavigationBar.cpp
@@ -0,0 +1,526 @@
+// Copyright Epic Games, Inc. All Rights Reserved.
+
+#include "SNavigationBar.h"
+
+#include "Framework/Application/SlateApplication.h"
+#include "Framework/Application/SlateUser.h"
+#include "Framework/MultiBox/MultiBoxBuilder.h"
+#include "Widgets/Input/SEditableText.h"
+#include "Widgets/Input/SSuggestionTextBox.h"
+#include "Widgets/Layout/SSeparator.h"
+#include "Widgets/Navigation/SBreadcrumbTrail.h"
+#include "Widgets/SOverlay.h"
+#include "Widgets/Views/SListView.h"
+
+#define LOCTEXT_NAMESPACE "ContentBrowser"
+
+// Category of list view item to choose an icon
+enum class ELocationSource
+{
+	History,
+	Suggestion
+};
+
+// Type for list view control
+struct FLocationItem
+{
+	FString VirtualPath;
+	ELocationSource Source;
+};
+
+class SLocationListView : public SListView<TSharedPtr<FLocationItem>>
+{
+public:
+	SLATE_BEGIN_ARGS(SLocationListView)
+		{}
+		SLATE_EVENT(FOnGenerateRow, OnGenerateRow)
+		SLATE_EVENT(FOnSelectionChanged, OnSelectionChanged)
+		SLATE_EVENT(FOnKeyDown, OnKeyDownHandler)
+		SLATE_EVENT(FOnKeyChar, OnKeyCharHandler)
+	SLATE_END_ARGS()
+
+	void Construct(const FArguments& InArgs)
+	{
+		OnKeyCharHandler = InArgs._OnKeyCharHandler;
+		using SSuper = SListView<TSharedPtr<FLocationItem>>;
+		SSuper::Construct(SSuper::FArguments()
+			.ListItemsSource(&Items)
+			.ItemHeight(18.0f)
+			.SelectionMode(ESelectionMode::Single)
+			.OnGenerateRow(InArgs._OnGenerateRow)
+			.OnSelectionChanged(InArgs._OnSelectionChanged)
+			.OnKeyDownHandler(InArgs._OnKeyDownHandler)
+		);
+	}
+
+	// Forward typing events to the navigation bar so that focus can be moved back to the editable text box 
+	FReply OnKeyChar(const FGeometry& MyGeometry, const FCharacterEvent& InCharacterEvent)
+	{
+		if (OnKeyCharHandler.IsBound())
+		{
+			return OnKeyCharHandler.Execute(MyGeometry, InCharacterEvent);
+		}
+		return FReply::Unhandled();
+	}
+	
+	void SetItems(TArray<FString> InItems, ELocationSource InSource)
+	{
+		RequestListRefresh();
+		ClearSelection();
+		Items.Reset(InItems.Num());
+		for (FString& Item : InItems)
+		{
+			Items.Emplace(MakeShared<FLocationItem>(MoveTemp(Item), InSource));
+		}
+	}
+	
+	FReply NavigateToFirst()
+	{
+		if (Items.Num() > 0)
+		{
+			SetSelection(Items[0], ESelectInfo::OnNavigation);
+			return FReply::Handled().SetUserFocus(AsShared(), EFocusCause::SetDirectly);
+		}
+		return FReply::Handled();
+	}
+	FReply NavigateToLast()
+	{
+		if (Items.Num() > 0)
+		{
+			SetSelection(Items.Last(), ESelectInfo::OnNavigation);
+			return FReply::Handled().SetUserFocus(AsShared(), EFocusCause::SetDirectly);
+		}
+		return FReply::Handled();
+	}
+	
+protected:
+	FOnKeyChar OnKeyCharHandler;
+	TArray<TSharedPtr<FLocationItem>> Items;
+};
+
+void SNavigationBar::Construct(const FArguments& InArgs)
+{
+	OnGetComboOptions = InArgs._GetComboOptions;
+	OnCompletePrefix = InArgs._OnCompletePrefix;
+	OnNavigateToPath = InArgs._OnNavigateToPath;
+	OnCanEditPathAsText = InArgs._OnCanEditPathAsText;
+
+	ComboBoxStyle = InArgs._ComboBoxStyle;
+	TextBoxStyle = InArgs._TextBoxStyle;
+	TableRowStyle = InArgs._ItemStyle;
+	bShowMenuBackground = false;
+	
+	const ISlateStyle& StyleSet = FAppStyle::Get();
+	SuggestionIcon = StyleSet.GetBrush("Icons.Search");
+	HistoryIcon = StyleSet.GetBrush("Icons.Recent");
+
+	SComboButton::Construct(SComboButton::FArguments()
+		.ComboButtonStyle(&ComboBoxStyle->ComboButtonStyle)
+		.ButtonStyle(&ComboBoxStyle->ComboButtonStyle.ButtonStyle)
+		.ContentPadding(ComboBoxStyle->ContentPadding)
+		.ButtonContent()
+		[
+			SNew(SHorizontalBox)
+			+ SHorizontalBox::Slot()
+			.HAlign(HAlign_Fill)
+			[
+				SNew(SOverlay)
+				// Invisible button under visible controls to handle clicking in blank space
+				+ SOverlay::Slot()
+				.HAlign(HAlign_Fill)
+				.VAlign(VAlign_Fill)
+				[
+					SNew(SButton)
+					.OnClicked(this, &SNavigationBar::HandleBlankSpaceClicked)	
+					.ButtonStyle( FAppStyle::Get(), "NoBorder" )
+				]
+				// Breadcrumb trail aligned to left in combo box
+				+ SOverlay::Slot()
+				.HAlign(HAlign_Left)
+				.VAlign(VAlign_Fill)
+				[
+					SAssignNew(BreadcrumbBar, SBreadcrumbTrail<FString>)
+					.Visibility(this, &SNavigationBar::GetNonEditVisibility)
+					.TextStyle(InArgs._BreadcrumbTextStyle)
+					.ButtonStyle(InArgs._BreadcrumbButtonStyle)
+					.ButtonContentPadding(InArgs._BreadcrumbButtonContentPadding)
+					.DelimiterImage(InArgs._BreadcrumbDelimiterImage)
+					.OnCrumbClicked(InArgs._OnPathClicked)
+					.HasCrumbMenuContent(InArgs._HasPathMenuContent)
+					.GetCrumbMenuContent(InArgs._GetPathMenuContent)
+				]
+				// Editable text box taking up all space
+				+ SOverlay::Slot()
+				.HAlign(HAlign_Fill)
+				.VAlign(VAlign_Center)
+				[
+					SAssignNew( EditableText, SEditableText )
+					.Visibility(this, &SNavigationBar::GetEditTextVisibility)
+					.OnTextCommitted(this, &SNavigationBar::HandleTextCommitted)
+					.OnTextChanged(this, &SNavigationBar::HandleTextChanged)
+					.SelectAllTextWhenFocused(false)
+					.ClearKeyboardFocusOnCommit(true)
+					.OnKeyDownHandler(this, &SNavigationBar::HandleEditableTextKeyDown)
+				]
+			]
+		]
+		.MenuContent()
+		[
+			SNew(SVerticalBox)
+			+ SVerticalBox::Slot()
+			.Padding(StyleSet.GetMargin("Menu.Heading.Padding"))
+			.AutoHeight()
+			[
+				// This matches multibox heading control, it could be refactored to use that (or a shared control)
+				SNew(SHorizontalBox)
+				+ SHorizontalBox::Slot()
+				.VAlign(VAlign_Center)
+				.AutoWidth()
+				[
+					SNew(STextBlock)
+					.Text(this, &SNavigationBar::GetPopupHeading)
+					.TextStyle(&StyleSet, "Menu.Heading")
+				]
+				+ SHorizontalBox::Slot()
+				.Padding(FMargin(14.f, 0.f, 0.f, 0.f))
+				.VAlign(VAlign_Center)
+				.HAlign(HAlign_Fill)
+				[
+					SNew(SSeparator)
+					.Orientation(Orient_Horizontal)
+					.Thickness(1.0f)
+					.SeparatorImage(StyleSet.GetBrush("Menu.Separator"))
+				]
+			]
+			+ SVerticalBox::Slot()
+			.AutoHeight()
+			.MaxHeight(300.0f)
+			.Padding(2.0f,2.0f)
+			[
+				SAssignNew(ComboListView, SLocationListView)
+				.OnGenerateRow(this, &SNavigationBar::HandleGenerateComboRow)
+				.OnSelectionChanged(this, &SNavigationBar::HandleComboSelectionChanged)
+				.OnKeyDownHandler(this, &SNavigationBar::HandleComboKeyDown)
+				.OnKeyCharHandler(this, &SNavigationBar::HandleComboKeyChar)
+			]
+		]
+	);
+
+	ComboListView->SetBackgroundBrush(FStyleDefaults::GetNoBrush());
+	EditableText->SetTextBlockStyle(&TextBoxStyle->TextStyle);
+	SetMenuContentWidgetToFocus(ComboListView);
+}
+
+void SNavigationBar::ClearPaths()
+{
+	BreadcrumbBar->ClearCrumbs();
+}
+
+void SNavigationBar::PushPath(const FText& ElementText, const FString& FullPath)
+{
+	BreadcrumbBar->PushCrumb(ElementText, FullPath);	
+}
+
+bool SNavigationBar::SupportsKeyboardFocus() const 
+{
+	return bIsFocusable;
+}
+
+void SNavigationBar::OnFocusChanging(const FWeakWidgetPath& PreviousFocusPath, const FWidgetPath& NewWidgetPath, const FFocusEvent& InFocusEvent) 
+{
+	// Stop editing and hide dropdown when focusing away from either the navigation bar or the popup
+	TSharedPtr<SWindow> Popup = PopupWindowPtr.Pin();
+	if(PreviousFocusPath.ContainsWidget(this) || PreviousFocusPath.ContainsWidget(Popup.Get()))
+	{
+		if (!NewWidgetPath.ContainsWidget(this) && !NewWidgetPath.ContainsWidget(Popup.Get()))
+		{
+			if (CompletionTimerHandle.IsValid())
+			{
+				UnRegisterActiveTimer(CompletionTimerHandle.ToSharedRef());	
+				CompletionTimerHandle.Reset();
+			}
+
+			EditTextVisibility = EVisibility::Hidden;
+			SetIsOpen(false, false);
+		}
+	}
+}
+
+FReply SNavigationBar::OnButtonClicked()
+{
+	GenerateHistoryOptions();
+	return SComboButton::OnButtonClicked();
+}
+
+void SNavigationBar::GenerateHistoryOptions()
+{
+	PopupHeading = LOCTEXT("NavigationBar.HistoryHeader", "HISTORY");
+	TArray<FString> Options;
+	if (OnGetComboOptions.IsBound())
+	{
+		Options = OnGetComboOptions.Execute();
+	}
+	ComboListView->SetItems(MoveTemp(Options), ELocationSource::History);
+}
+
+void SNavigationBar::GenerateCompletionOptions(const FString& Prefix)
+{
+	PopupHeading = LOCTEXT("NavigationBar.SuggestionsHeader", "SUGGESTIONS"); 
+	TArray<FString> Options;
+	if (OnCompletePrefix.IsBound())
+	{
+		Options = OnCompletePrefix.Execute(Prefix);
+	}
+	ComboListView->SetItems(MoveTemp(Options), ELocationSource::Suggestion);
+}
+
+FReply SNavigationBar::HandleBlankSpaceClicked()
+{
+	StartEditingPath();
+	return FReply::Handled();
+}
+
+FReply SNavigationBar::HandleComboKeyChar(const FGeometry& MyGeometry, const FCharacterEvent& InCharEvent)
+{
+	// Add the typed character to the currently suggested suggestion and return focus to the editable text box
+	TArray<TSharedPtr<FLocationItem>> SelectedItems = ComboListView->GetSelectedItems();
+	if (SelectedItems.Num() != 0)
+	{
+		FString NewText = SelectedItems[0]->VirtualPath;
+		NewText += InCharEvent.GetCharacter();
+		EditableText->SetText(FText::FromString(NewText));
+
+		return FReply::Handled().SetUserFocus(EditableText.ToSharedRef());
+	}
+	return FReply::Unhandled();	
+}
+
+FReply SNavigationBar::HandleComboKeyDown(const FGeometry& MyGeometry, const FKeyEvent& InKeyEvent)
+{
+	if (InKeyEvent.GetKey() == EKeys::Enter)
+	{
+		// Select (and navigate to) the first selected item on hitting enter
+		TArray<TSharedPtr<FLocationItem>> SelectedItems = ComboListView->GetSelectedItems();
+		if (SelectedItems.Num() > 0)
+		{
+			HandleComboSelectionChanged(SelectedItems[0], ESelectInfo::OnKeyPress);
+
+			// Set focus back to navigation box 
+			FSlateApplication::Get().ForEachUser([this](FSlateUser& User) 
+			{
+				TSharedRef<SWidget> ThisRef = this->AsShared();
+				if (User.IsWidgetInFocusPath(this->ComboListView))
+				{
+					User.SetFocus(ThisRef);
+				}
+			});
+
+			return FReply::Handled();
+		}
+	}
+	else if (InKeyEvent.GetKey() == EKeys::Escape)
+	{
+		// Close text box and popup when pressing escape
+		EditTextVisibility = EVisibility::Hidden;
+		SetIsOpen(false, false);
+		return FReply::Handled();
+	}
+	else if (InKeyEvent.GetKey() == EKeys::Tab)
+	{
+		// Put current selected item into editable text and move focus back there
+		TArray<TSharedPtr<FLocationItem>> SelectedItems = ComboListView->GetSelectedItems();
+		if (SelectedItems.Num() > 0)
+		{
+			FString NewText = SelectedItems[0]->VirtualPath + TEXT("/"); // TODO: Ideally only add slash if this prefix has children itself 
+			EditableText->SetText(FText::FromString(NewText));
+			return FReply::Handled().SetUserFocus(EditableText.ToSharedRef(), EFocusCause::SetDirectly);
+		}
+		return FReply::Handled(); // Swallow event to avoid confusion with items being selected or not 
+	}
+
+	return FReply::Unhandled();
+}
+
+FReply SNavigationBar::HandleEditableTextKeyDown(const FGeometry& MyGeometry, const FKeyEvent& KeyEvent)
+{
+	if (EditableText->HasKeyboardFocus() && KeyEvent.GetKey() == EKeys::Escape)
+	{
+		// Stop editing and discard when user presses escape 
+		return FReply::Handled().SetUserFocus(AsShared(), EFocusCause::Cleared);
+	}
+
+	if (EditableText->HasKeyboardFocus() && KeyEvent.GetKey() == EKeys::Tab)
+	{
+		// Swallow tab key to prevent moving focus away because Tab is used by the popup if it's focused
+		return FReply::Handled();
+	}
+
+	bool bUp = KeyEvent.GetKey() == EKeys::Up;
+	bool bDown = KeyEvent.GetKey() == EKeys::Down;
+	
+	if (bUp || bDown)
+	{
+		if (!IsOpen())
+		{
+			// Open popup if completions haven't been generated from timer after text edit yet
+			GenerateCompletionOptions(EditableText->GetText().ToString());
+			SetIsOpen(true, false);
+			return FReply::Handled();
+		}
+		else 
+		{
+			// Switch focus from text box to first or last element of popup to start navigation
+			if (bUp)
+			{
+				return ComboListView->NavigateToLast();
+			}
+			else 
+			{
+				return ComboListView->NavigateToFirst();
+			}
+		}
+	}
+
+	return FReply::Unhandled();
+}
+
+void SNavigationBar::StartEditingPath()
+{
+	EditTextVisibility = EVisibility::Visible;
+	FText Text = FText::GetEmpty();
+	if (BreadcrumbBar->HasCrumbs())
+	{
+		FString Path = BreadcrumbBar->PeekCrumb();
+		if (OnCanEditPathAsText.IsBound() && OnCanEditPathAsText.Execute(Path))
+		{
+			Text = FText::FromString(Path);
+		}
+	}
+	EditableText->SetText(Text);
+	FSlateApplication::Get().SetKeyboardFocus(EditableText.ToSharedRef(), EFocusCause::SetDirectly);
+}
+
+void SNavigationBar::HandleTextChanged(const FText& NewText)
+{
+	if (EditableText->GetVisibility().IsVisible() && !bNoSuggestionsFromTextChange)
+	{
+		if (CompletionTimerHandle.IsValid())
+		{
+			UnRegisterActiveTimer(CompletionTimerHandle.ToSharedRef());	
+		}
+		// Generate completion suggestions shortly after user stops typing
+		CompletionTimerHandle = RegisterActiveTimer(0.1f, FWidgetActiveTimerDelegate::CreateSP(this, &SNavigationBar::HandleUpdateCompletionOptions));
+	}
+}
+
+	
+void SNavigationBar::HandleTextCommitted(const FText& InText, ETextCommit::Type CommitType)
+{
+	switch (CommitType)
+	{
+		case ETextCommit::Default:
+		case ETextCommit::OnCleared:
+		case ETextCommit::OnUserMovedFocus:
+			return; // Discard changes and don't navigate
+	}
+
+	// Stop editing and navigate to new path
+	EditTextVisibility = EVisibility::Hidden;
+	SetIsOpen(false, false);
+	OnNavigateToPath.ExecuteIfBound(InText.ToString());
+}
+
+EActiveTimerReturnType SNavigationBar::HandleUpdateCompletionOptions(double InCurrentTime, float InDeltaTime)
+{
+	GenerateCompletionOptions(EditableText->GetText().ToString());
+	SetIsOpen(true, false);
+	CompletionTimerHandle = nullptr;
+	return EActiveTimerReturnType::Stop; // Never run more than once unless text changes again 
+}
+
+TSharedRef<ITableRow> SNavigationBar::HandleGenerateComboRow(TSharedPtr<FLocationItem> ForItem, const TSharedRef<STableViewBase>& OwnerTable)
+{
+	return SNew(STableRow<TSharedPtr<FLocationItem>>, OwnerTable)
+		.Padding(ComboBoxStyle->MenuRowPadding)
+		.Style(TableRowStyle)
+		[
+			SNew(SHorizontalBox)
+			+ SHorizontalBox::Slot()
+			.HAlign(HAlign_Left)
+			.VAlign(VAlign_Center)
+			.Padding(2.0f, 0.0f)
+			.AutoWidth()
+			[
+				SNew(SImage)
+				.Image(GetImageForItem(*ForItem))
+			]
+			+ SHorizontalBox::Slot()
+			.HAlign(HAlign_Left)
+			.Padding(2.0f, 0.0f)
+			[
+				SNew(STextBlock)
+				.Text(FText::FromString(ForItem->VirtualPath))
+			]
+		];
+}
+
+void SNavigationBar::HandleComboSelectionChanged(TSharedPtr<FLocationItem> SelectedItem, ESelectInfo::Type SelectInfo)
+{
+	if (!SelectedItem.IsValid())
+	{
+		return;
+	}
+
+	switch(SelectInfo)
+	{
+		case ESelectInfo::Direct:
+		case ESelectInfo::OnKeyPress:
+		case ESelectInfo::OnMouseClick:
+			// Stop editing/hide popup and navigate to chosen path
+			EditTextVisibility = EVisibility::Hidden;
+			SetIsOpen(false);
+			OnNavigateToPath.ExecuteIfBound(SelectedItem->VirtualPath);
+			break;
+		case ESelectInfo::OnNavigation:
+			// Mirror chosen item into text box without moving focus back or triggering new suggestions
+			TGuardValue<bool> BlockSuggestions{bNoSuggestionsFromTextChange, true};
+			EditableText->SetText(FText::FromString(SelectedItem->VirtualPath));
+			break;
+	}
+}
+
+FText SNavigationBar::GetPopupHeading() const 
+{
+	return PopupHeading;
+}
+
+const FSlateBrush* SNavigationBar::GetImageForItem(const FLocationItem& ForItem) const
+{
+	switch (ForItem.Source)
+	{
+		case ELocationSource::History:
+			return HistoryIcon;
+		case ELocationSource::Suggestion:
+		default:
+			return SuggestionIcon;
+	}
+}
+
+EVisibility SNavigationBar::GetEditTextVisibility() const
+{
+	return EditTextVisibility;
+}
+
+EVisibility SNavigationBar::GetNonEditVisibility() const
+{
+	if (EditTextVisibility.IsVisible())
+	{
+		return EVisibility::Hidden;
+	}
+	else
+	{
+		return EVisibility::Visible;
+	}
+}
+
+#undef LOCTEXT_NAMESPACE
\ No newline at end of file
diff --git a/Engine/Source/Editor/ContentBrowser/Private/SNavigationBar.h b/Engine/Source/Editor/ContentBrowser/Private/SNavigationBar.h
new file mode 100644
--- /dev/null
+++ b/Engine/Source/Editor/ContentBrowser/Private/SNavigationBar.h
@@ -0,0 +1,154 @@
+// Copyright Epic Games, Inc. All Rights Reserved.
+
+#pragma once
+
+#include "Delegates/Delegate.h"
+#include "Styling/CoreStyle.h"
+#include "Styling/ToolBarStyle.h"
+#include "Styling/SlateTypes.h"
+#include "Styling/SlateWidgetStyleAsset.h"
+#include "Widgets/Input/SComboButton.h"
+#include "Widgets/Layout/SBorder.h"
+
+class FText;
+class ITableRow;
+template<typename ItemType>
+class SBreadcrumbTrail;
+class SComboButton;
+class SEditableText;
+class SLocationListView;
+class STableViewBase;
+class USlateWidgetStyleAsset;
+
+using FOnPathClicked = TDelegate<void(const FString&)>;
+using FHasPathMenuContent = TDelegate<bool(const FString&)>;
+using FGetPathMenuContent = TDelegate<TSharedRef<SWidget>(const FString&)>;
+using FGetComboOptions = TDelegate<TArray<FString>(void)>;
+using FOnCompletePrefix = TDelegate<TArray<FString>(const FString&)>;
+using FOnNavigateToPath = TDelegate<void(const FString&)>;
+using FCanEditPathAsText = TDelegate<bool(const FString&)>;
+
+// Private class for internal list view
+struct FLocationItem;
+
+class SNavigationBar : public SComboButton
+{
+public:
+    SLATE_BEGIN_ARGS(SNavigationBar) 
+	: _ComboBoxStyle(&FAppStyle::Get().GetWidgetStyle< FComboBoxStyle >("ComboBox"))
+	, _ItemStyle(&FAppStyle::Get().GetWidgetStyle< FTableRowStyle >("ComboBox.Row"))
+	, _TextBoxStyle(&FAppStyle::Get().GetWidgetStyle<FEditableTextBoxStyle>("NormalEditableTextBox"))
+	, _BreadcrumbButtonStyle(&FAppStyle::Get().GetWidgetStyle<FButtonStyle>("SimpleButton"))
+	, _BreadcrumbTextStyle(&FAppStyle::Get().GetWidgetStyle<FTextBlockStyle>("NormalText"))
+	, _BreadcrumbButtonContentPadding(FMargin(2,2))
+	, _BreadcrumbDelimiterImage(FAppStyle::Get().GetBrush("Icons.ChevronRight"))
+    {}
+		/** The styling of the combo box which wraps the other navigation controls */
+		SLATE_STYLE_ARGUMENT( FComboBoxStyle, ComboBoxStyle )
+
+		/** Styling for items in the drop down history/suggestion box */
+		SLATE_STYLE_ARGUMENT(FTableRowStyle, ItemStyle)
+
+		/** The styling of the suggestions/history popup */
+		SLATE_STYLE_ARGUMENT(FTableViewStyle, SuggestionsStyle)
+
+		/** The styling of the textbox when in edit mode */
+		SLATE_STYLE_ARGUMENT(FEditableTextBoxStyle, TextBoxStyle)
+		
+		/** The name of the style to use for the crumb buttons */
+		SLATE_STYLE_ARGUMENT(FButtonStyle, BreadcrumbButtonStyle)
+
+		/** The name of the style to use for the crumb button text */
+		SLATE_STYLE_ARGUMENT(FTextBlockStyle, BreadcrumbTextStyle)
+
+		/** The padding for the content in crumb buttons */
+		SLATE_ATTRIBUTE(FMargin, BreadcrumbButtonContentPadding)
+	
+		/** The image to use between crumb trail buttons */
+		SLATE_ATTRIBUTE(const FSlateBrush*, BreadcrumbDelimiterImage)
+		
+		/** Callback for when a new location is committed via text entry */
+        SLATE_EVENT(FOnNavigateToPath , OnNavigateToPath)
+
+		/** Called when an invididual button is clicked, after the later crumbs were popped */
+		SLATE_EVENT(FOnPathClicked , OnPathClicked)
+
+		/** Called to check whether there are locations after a particular crumb to populate a dropdown*/
+		SLATE_EVENT(FHasPathMenuContent, HasPathMenuContent)
+
+		/** Called to get dropdown options for a crumb separator's menu */
+		SLATE_EVENT(FGetPathMenuContent, GetPathMenuContent)
+		
+		/** Called to get a list of options when clicking the combo botton on the side of the box */
+		SLATE_EVENT(FGetComboOptions, GetComboOptions)
+		
+		/** Called to get suggested text completion values for the currently edit text */
+		SLATE_EVENT(FOnCompletePrefix, OnCompletePrefix)
+		
+		/** 
+		 * Called to check if a path (e.g. the last breadcrumb) can be edited as text 
+		 * If not, when editing this location the text box will be emptied
+		 */
+		SLATE_EVENT(FCanEditPathAsText, OnCanEditPathAsText)
+
+    SLATE_END_ARGS()
+    
+    void Construct(const FArguments& InArgs);
+    
+	// Remove all paths from the breadcrumb bar
+	void ClearPaths();
+	// Add a new segment to the breadcrumb bar
+	void PushPath(const FText& SegmentDisplayText, const FString& FullLocation);
+	// Replace the breadcrumb bar with an editable text box and focus it
+	void StartEditingPath();
+
+private:
+	// SWidget overrides
+	virtual void OnFocusChanging(const FWeakWidgetPath& PreviousFocusPath, const FWidgetPath& NewWidgetPath, const FFocusEvent& InFocusEvent) override;
+	// SComboButton overrides
+	virtual FReply OnButtonClicked() override;
+
+	// Handle clicking on part of the navigation bar which is not the combo button or the breadcrumb bar, to switch to edit mode
+	FReply HandleBlankSpaceClicked();
+	void HandleTextChanged(const FText& NewText);
+	void HandleTextCommitted(const FText& InText, ETextCommit::Type CommitType);
+	void HandleComboSelectionChanged(TSharedPtr<FLocationItem> SelectedItem, ESelectInfo::Type SelectInfo);
+	TSharedRef<ITableRow> HandleGenerateComboRow(TSharedPtr<FLocationItem> ForItem, const TSharedRef<STableViewBase>& OwnerTable);
+	FReply HandleEditableTextKeyDown(const FGeometry& MyGeometry, const FKeyEvent& KeyEvent);
+	FReply HandleComboKeyDown(const FGeometry& MyGeometry, const FKeyEvent& InKeyEvent);
+	FReply HandleComboKeyChar(const FGeometry& MyGeometry, const FCharacterEvent& InCharEvent);
+	
+	EActiveTimerReturnType HandleUpdateCompletionOptions(double InCurrentTime, float InDeltaTime);
+
+	void GenerateHistoryOptions();
+	void GenerateCompletionOptions(const FString& Prefix);
+	
+	EVisibility GetEditTextVisibility() const;
+	EVisibility GetNonEditVisibility() const;
+	FText GetPopupHeading() const;
+	const FSlateBrush* GetImageForItem(const FLocationItem& ForItem) const;
+
+	virtual bool SupportsKeyboardFocus() const override;
+
+	const FComboBoxStyle* ComboBoxStyle;
+	const FTableRowStyle* TableRowStyle;
+	const FEditableTextBoxStyle* TextBoxStyle;
+	const FSlateBrush* SuggestionIcon;
+	const FSlateBrush* HistoryIcon;
+	
+	EVisibility EditTextVisibility = EVisibility::Hidden;
+	TSharedPtr<FActiveTimerHandle> CompletionTimerHandle;
+
+	TSharedPtr<SEditableText> EditableText;
+    TSharedPtr<SBreadcrumbTrail<FString>> BreadcrumbBar;
+	TSharedPtr<SLocationListView> ComboListView;
+	
+	FText PopupHeading;
+	
+	bool bNoSuggestionsFromTextChange = false;
+	
+	FOnNavigateToPath OnNavigateToPath;
+	FGetComboOptions OnGetComboOptions;
+	FOnCompletePrefix OnCompletePrefix;	
+	FCanEditPathAsText OnCanEditPathAsText;
+};
\ No newline at end of file
