Fix construction of FLocationItem to explicitly use aggregate initialization

diff --git a/Engine/Source/Editor/ContentBrowser/Private/SNavigationBar.cpp b/Engine/Source/Editor/ContentBrowser/Private/SNavigationBar.cpp
--- a/Engine/Source/Editor/ContentBrowser/Private/SNavigationBar.cpp
+++ b/Engine/Source/Editor/ContentBrowser/Private/SNavigationBar.cpp
@@ -70,7 +70,7 @@ public:
 		Items.Reset(InItems.Num());
 		for (FString& Item : InItems)
 		{
-			Items.Emplace(MakeShared<FLocationItem>(MoveTemp(Item), InSource));
+			Items.Emplace(MakeShared<FLocationItem>(FLocationItem{MoveTemp(Item), InSource}));
 		}
 	}
 	
