Add additional default binding of Ctrl+L (matches Windows Explorer and web browsers) for editing current path in content browser.

diff --git a/Engine/Source/Editor/ContentBrowser/Private/ContentBrowserCommands.cpp b/Engine/Source/Editor/ContentBrowser/Private/ContentBrowserCommands.cpp
--- a/Engine/Source/Editor/ContentBrowser/Private/ContentBrowserCommands.cpp
+++ b/Engine/Source/Editor/ContentBrowser/Private/ContentBrowserCommands.cpp
@@ -20,7 +20,7 @@ void FContentBrowserCommands::RegisterCommands()
 	UI_COMMAND(SaveAllCurrentFolder, "Save All", "Save All in current folder", EUserInterfaceActionType::Button, FInputChord(EModifierKey::Control | EModifierKey::Shift, EKeys::S));
 	UI_COMMAND(ResaveAllCurrentFolder, "Resave All", "Resave all assets contained in the current folder", EUserInterfaceActionType::Button, FInputChord());
 	UI_COMMAND(CopySelectedAssetPath, "Copy Selected Asset Path", "Copy the selected asset path", EUserInterfaceActionType::Button, FInputChord(EModifierKey::Control, EKeys::C));
-	UI_COMMAND(EditPath, "Edit Path", "Edit the current content browser path", EUserInterfaceActionType::Button, FInputChord(EKeys::F4));
+	UI_COMMAND(EditPath, "Edit Path", "Edit the current content browser path", EUserInterfaceActionType::Button, FInputChord(EKeys::F4), FInputChord(EModifierKey::Control, EKeys::L));
 }
 
 
