Remove code for different icons in content browser for plugins.

diff --git a/Engine/Source/Editor/ContentBrowser/Private/AssetViewWidgets.cpp b/Engine/Source/Editor/ContentBrowser/Private/AssetViewWidgets.cpp
--- a/Engine/Source/Editor/ContentBrowser/Private/AssetViewWidgets.cpp
+++ b/Engine/Source/Editor/ContentBrowser/Private/AssetViewWidgets.cpp
@@ -191,7 +191,6 @@ TSharedRef<SWidget> FAssetViewItemHelper::CreateListTileItemContents(T* const In
 
 		const bool bDeveloperFolder = ContentBrowserUtils::IsItemDeveloperContent(InTileOrListItem->AssetItem->GetItem());
 		const bool bCodeFolder = EnumHasAnyFlags(InTileOrListItem->AssetItem->GetItem().GetItemCategory(), EContentBrowserItemFlags::Category_Class);
-		const bool bPluginFolder = EnumHasAnyFlags(InTileOrListItem->AssetItem->GetItem().GetItemCategory(), EContentBrowserItemFlags::Category_Plugin);
 
 		const bool bCollectionFolder = EnumHasAnyFlags(InTileOrListItem->AssetItem->GetItem().GetItemCategory(), EContentBrowserItemFlags::Category_Collection);
 		ECollectionShareType::Type CollectionFolderShareType = ECollectionShareType::CST_All;
@@ -209,10 +208,6 @@ TSharedRef<SWidget> FAssetViewItemHelper::CreateListTileItemContents(T* const In
 		{
 			FolderBaseImage = FAppStyle::GetBrush("ContentBrowser.ListViewCodeFolderIcon");
 		}
-		else if (bPluginFolder)
-		{
-			FolderBaseImage = FAppStyle::GetBrush("ContentBrowser.ListViewPluginFolderIcon");
-		}
 		else
 		{
 			FolderBaseImage = FAppStyle::GetBrush("ContentBrowser.ListViewFolderIcon");
diff --git a/Engine/Source/Editor/ContentBrowser/Private/SourcesViewWidgets.cpp b/Engine/Source/Editor/ContentBrowser/Private/SourcesViewWidgets.cpp
--- a/Engine/Source/Editor/ContentBrowser/Private/SourcesViewWidgets.cpp
+++ b/Engine/Source/Editor/ContentBrowser/Private/SourcesViewWidgets.cpp
@@ -40,14 +40,6 @@ struct FSlateBrush;
 
 #define LOCTEXT_NAMESPACE "ContentBrowser"
 
-static bool GUsePluginFolderIcon = 0;
-static FAutoConsoleVariableRef CVarUsePluginFolderIcon(
-	TEXT("ContentBrowser.UsePluginFolderIcon"),
-	GUsePluginFolderIcon,
-	TEXT("Temporary flag to control use of separate icon for plugin folders pending final icon"),
-	ECVF_Default
-);
-
 struct FAssetTreeItemBrushes
 {
 	/** Brushes for the different folder states */
@@ -56,11 +48,7 @@ struct FAssetTreeItemBrushes
 	const FSlateBrush* FolderOpenCodeBrush;
 	const FSlateBrush* FolderClosedCodeBrush;
 	const FSlateBrush* FolderOpenDeveloperBrush;
-	const FSlateBrush* FolderClosedDeveloperBrush;
-	const FSlateBrush* FolderOpenPluginRootBrush;
-	const FSlateBrush* FolderClosedPluginRootBrush;
-	const FSlateBrush* FolderOpenPluginBrush;
-	const FSlateBrush* FolderClosedPluginBrush;
+	const FSlateBrush *FolderClosedDeveloperBrush;
 
 	FAssetTreeItemBrushes()
 	{
@@ -70,10 +58,6 @@ struct FAssetTreeItemBrushes
 		FolderClosedCodeBrush = FAppStyle::GetBrush("ContentBrowser.AssetTreeFolderClosedCode");
 		FolderOpenDeveloperBrush = FAppStyle::GetBrush("ContentBrowser.AssetTreeFolderOpenDeveloper");
 		FolderClosedDeveloperBrush = FAppStyle::GetBrush("ContentBrowser.AssetTreeFolderClosedDeveloper");
-		FolderOpenPluginRootBrush = FAppStyle::GetBrush("ContentBrowser.AssetTreeFolderOpenPluginRoot");
-		FolderClosedPluginRootBrush = FAppStyle::GetBrush("ContentBrowser.AssetTreeFolderClosedPluginRoot");
-		FolderOpenPluginBrush = FAppStyle::GetBrush("ContentBrowser.AssetTreeFolderOpenPlugin");
-		FolderClosedPluginBrush = FAppStyle::GetBrush("ContentBrowser.AssetTreeFolderClosedPlugin");
 	}
 	
 	static FAssetTreeItemBrushes& Get()
@@ -108,18 +92,6 @@ void SAssetTreeItem::Construct( const FArguments& InArgs )
 	{
 		FolderType = EFolderType::Code;
 	}
-	else if (InArgs._TreeItem->GetItem().IsInPlugin())
-	{
-		TSharedPtr<FTreeItem> Parent = InArgs._TreeItem->Parent.Pin();
-		if (Parent.IsValid() && Parent->GetItem().IsInPlugin())
-		{
-			FolderType = EFolderType::PluginSubfolder;
-		}
-		else
-		{
-			FolderType = EFolderType::PluginRoot;
-		}
-	}
 
 	bool bIsRoot = !InArgs._TreeItem->Parent.IsValid();
 
@@ -316,19 +288,6 @@ const FSlateBrush* SAssetTreeItem::GetFolderIcon() const
 	case EFolderType::Developer:
 		return (IsItemExpanded.Get()) ? Brushes.FolderOpenDeveloperBrush : Brushes.FolderClosedDeveloperBrush;
 
-	case EFolderType::PluginRoot:
-		if (GUsePluginFolderIcon)
-		{
-			return (IsItemExpanded.Get()) ? Brushes.FolderOpenPluginRootBrush : Brushes.FolderClosedPluginRootBrush;
-		}
-		// Fall through to default 
-	case EFolderType::PluginSubfolder:
-		if (GUsePluginFolderIcon)
-		{
-			return (IsItemExpanded.Get()) ? Brushes.FolderOpenPluginBrush : Brushes.FolderClosedPluginBrush;
-		}
-		// Fall through to default 
-
 	default:
 		return ( IsItemExpanded.Get() ) ? Brushes.FolderOpenBrush : Brushes.FolderClosedBrush;
 	}
diff --git a/Engine/Source/Editor/ContentBrowser/Private/SourcesViewWidgets.h b/Engine/Source/Editor/ContentBrowser/Private/SourcesViewWidgets.h
--- a/Engine/Source/Editor/ContentBrowser/Private/SourcesViewWidgets.h
+++ b/Engine/Source/Editor/ContentBrowser/Private/SourcesViewWidgets.h
@@ -108,8 +108,6 @@ private:
 		Normal,
 		Code,
 		Developer,
-		PluginSubfolder,
-		PluginRoot,
 	};
 
 	/** The data for this item */
diff --git a/Engine/Source/Editor/EditorStyle/Private/StarshipStyle.cpp b/Engine/Source/Editor/EditorStyle/Private/StarshipStyle.cpp
--- a/Engine/Source/Editor/EditorStyle/Private/StarshipStyle.cpp
+++ b/Engine/Source/Editor/EditorStyle/Private/StarshipStyle.cpp
@@ -6123,11 +6123,7 @@ void FStarshipEditorStyle::FStyle::SetupContentBrowserStyle()
 		Set( "ContentBrowser.AssetTreeFolderOpenDeveloper", new IMAGE_BRUSH_SVG("Starship/ContentBrowser/folder-developer-open", Icon16x16));
 		Set( "ContentBrowser.AssetTreeFolderClosedDeveloper", new IMAGE_BRUSH_SVG( "Starship/ContentBrowser/folder-developer", Icon16x16));
 		Set( "ContentBrowser.AssetTreeFolderOpenCode", new IMAGE_BRUSH_SVG( "Starship/ContentBrowser/folder-code-open", Icon16x16) );
-		Set( "ContentBrowser.AssetTreeFolderClosedCode", new IMAGE_BRUSH_SVG( "Starship/ContentBrowser/folder-code", Icon16x16) );
-		Set( "ContentBrowser.AssetTreeFolderOpenPluginRoot", new IMAGE_BRUSH_SVG("Starship/ContentBrowser/folder-pluginroot-open", Icon16x16));
-		Set( "ContentBrowser.AssetTreeFolderClosedPluginRoot", new IMAGE_BRUSH_SVG( "Starship/ContentBrowser/folder-pluginroot", Icon16x16));
-		Set( "ContentBrowser.AssetTreeFolderOpenPlugin", new IMAGE_BRUSH_SVG("Starship/ContentBrowser/folder-plugin-open", Icon16x16));
-		Set( "ContentBrowser.AssetTreeFolderClosedPlugin", new IMAGE_BRUSH_SVG( "Starship/ContentBrowser/folder-plugin", Icon16x16));
+		Set("ContentBrowser.AssetTreeFolderClosedCode", new IMAGE_BRUSH_SVG("Starship/ContentBrowser/folder-code", Icon16x16));
 
 		Set("ContentBrowser.DefaultFolderColor", FStyleColors::AccentFolder);
 
@@ -6309,9 +6305,8 @@ void FStarshipEditorStyle::FStyle::SetupContentBrowserStyle()
 
 		Set( "ContentBrowser.ListViewFolderIcon", new IMAGE_BRUSH_SVG( "Starship/ContentBrowser/folder", FVector2D(64, 64) ) );
 		Set( "ContentBrowser.ListViewDeveloperFolderIcon", new IMAGE_BRUSH_SVG("Starship/ContentBrowser/FolderLargeDeveloper", FVector2D(64, 64)));
-		Set( "ContentBrowser.ListViewCodeFolderIcon", new IMAGE_BRUSH_SVG("Starship/ContentBrowser/FolderLargeCode", FVector2D(64, 64)));
-		Set( "ContentBrowser.ListViewPluginFolderIcon", new IMAGE_BRUSH_SVG("Starship/ContentBrowser/FolderLargePlugin", FVector2D(64, 64)));
-		
+		Set("ContentBrowser.ListViewCodeFolderIcon", new IMAGE_BRUSH_SVG("Starship/ContentBrowser/FolderLargeCode", FVector2D(64, 64)));
+
 		Set("ContentBrowser.AssetTileItem.FolderAreaHoveredBackground", new FSlateRoundedBoxBrush(FStyleColors::Secondary, 4.0f));
 		Set("ContentBrowser.AssetTileItem.FolderAreaSelectedBackground", new FSlateRoundedBoxBrush(FStyleColors::Secondary, 4.0f));
 		Set("ContentBrowser.AssetTileItem.FolderAreaSelectedHoverBackground", new FSlateRoundedBoxBrush(FStyleColors::Secondary, 4.0f));
