Virtual and plugin folder icons for content browser
Fix code in asset view to only use large plugin folder icon for plugin roots.

diff --git a/Engine/Source/Editor/ContentBrowser/Private/AssetViewWidgets.cpp b/Engine/Source/Editor/ContentBrowser/Private/AssetViewWidgets.cpp
--- a/Engine/Source/Editor/ContentBrowser/Private/AssetViewWidgets.cpp
+++ b/Engine/Source/Editor/ContentBrowser/Private/AssetViewWidgets.cpp
@@ -194,7 +194,7 @@ TSharedRef<SWidget> FAssetViewItemHelper::CreateListTileItemContents(T* const In
 		const bool bCodeFolder = EnumHasAnyFlags(AssetItem->GetItem().GetItemCategory(), EContentBrowserItemFlags::Category_Class);
 		FContentBrowserItemDataAttributeValue VirtualAttributeValue = AssetItem->GetItem().GetItemAttribute(ContentBrowserItemAttributes::ItemIsCustomVirtualFolder);
 		const bool bVirtualFolder = VirtualAttributeValue.IsValid() && VirtualAttributeValue.GetValue<bool>();
-		const bool bPluginFolder = EnumHasAnyFlags(InTileOrListItem->AssetItem->GetItem().GetItemCategory(), EContentBrowserItemFlags::Category_Plugin);
+		const bool bPluginFolder = ContentBrowserUtils::IsItemPluginRootFolder(AssetItem->GetItem());
 		
 		const bool bCollectionFolder = EnumHasAnyFlags(AssetItem->GetItem().GetItemCategory(), EContentBrowserItemFlags::Category_Collection);
 		ECollectionShareType::Type CollectionFolderShareType = ECollectionShareType::CST_All;
diff --git a/Engine/Source/Editor/ContentBrowser/Private/ContentBrowserUtils.cpp b/Engine/Source/Editor/ContentBrowser/Private/ContentBrowserUtils.cpp
--- a/Engine/Source/Editor/ContentBrowser/Private/ContentBrowserUtils.cpp
+++ b/Engine/Source/Editor/ContentBrowser/Private/ContentBrowserUtils.cpp
@@ -436,6 +436,26 @@ bool ContentBrowserUtils::IsItemPluginContent(const FContentBrowserItem& InItem)
 	return IsPluginAttributeValue.IsValid() && IsPluginAttributeValue.GetValue<bool>();
 }
 
+bool ContentBrowserUtils::IsItemPluginRootFolder(const FContentBrowserItem& InItem)
+{
+	if (!InItem.IsFolder())
+	{
+		return false;
+	}
+	FName InternalPath = InItem.GetInternalPath();
+	if (InternalPath.IsNone())
+	{
+		return false;
+	}
+	FNameBuilder PathBuffer(InternalPath);
+	FStringView Path = PathBuffer.ToView();
+	if (int32 Index = 0; Path.RightChop(1).FindChar('/', Index) && Index != INDEX_NONE)
+	{
+		return false; // Contains a second slash, is not a root
+	}
+	return IsItemPluginContent(InItem);
+}
+
 bool ContentBrowserUtils::IsCollectionPath(const FString& InPath, FName* OutCollectionName, ECollectionShareType::Type* OutCollectionShareType)
 {
 	static const FString CollectionsRootPrefix = TEXT("/Collections");
diff --git a/Engine/Source/Editor/ContentBrowser/Private/ContentBrowserUtils.h b/Engine/Source/Editor/ContentBrowser/Private/ContentBrowserUtils.h
--- a/Engine/Source/Editor/ContentBrowser/Private/ContentBrowserUtils.h
+++ b/Engine/Source/Editor/ContentBrowser/Private/ContentBrowserUtils.h
@@ -68,6 +68,9 @@ namespace ContentBrowserUtils
 	/** Check whether the given item is considered to be plugin content (engine or project) */
 	bool IsItemPluginContent(const FContentBrowserItem& InItem);
 
+	/** Check whether the given item is the root folder of a plugin */
+	bool IsItemPluginRootFolder(const FContentBrowserItem& InItem);
+
 	/** Check to see whether the given path is rooted against a collection directory, optionally extracting the collection name and share type from the path */
 	bool IsCollectionPath(const FString& InPath, FName* OutCollectionName = nullptr, ECollectionShareType::Type* OutCollectionShareType = nullptr);
 
