Added a filter to only show the assets that are not used by another asset.
Also removed some uneeded code and removed some unessary copy of a SharedPtr in the construction of the frontend filters.

diff --git a/Engine/Source/Developer/ToolWidgets/Public/Filters/FilterBase.h b/Engine/Source/Developer/ToolWidgets/Public/Filters/FilterBase.h
--- a/Engine/Source/Developer/ToolWidgets/Public/Filters/FilterBase.h
+++ b/Engine/Source/Developer/ToolWidgets/Public/Filters/FilterBase.h
@@ -32,7 +32,7 @@ class FFilterBase : public IFilter<FilterType>
 {
 public:
 
-	FFilterBase(TSharedPtr<FFilterCategory> InCategory) : FilterCategory(InCategory) {}
+	FFilterBase(TSharedPtr<FFilterCategory> InCategory) : FilterCategory(MoveTemp(InCategory)) {}
 
 	/** Returns the system name for this filter */
 	virtual FString GetName() const override = 0;
diff --git a/Engine/Source/Editor/ContentBrowser/Private/FrontendFilters.cpp b/Engine/Source/Editor/ContentBrowser/Private/FrontendFilters.cpp
--- a/Engine/Source/Editor/ContentBrowser/Private/FrontendFilters.cpp
+++ b/Engine/Source/Editor/ContentBrowser/Private/FrontendFilters.cpp
@@ -715,7 +715,7 @@ void FFrontendFilter_Text::RebuildReferencedDynamicCollections()
 /////////////////////////////////////////
 
 FFrontendFilter_CheckedOut::FFrontendFilter_CheckedOut(TSharedPtr<FFrontendFilterCategory> InCategory) 
-	: FFrontendFilter(InCategory),
+	: FFrontendFilter(MoveTemp(InCategory)),
 	bSourceControlEnabled(false)
 {
 	
@@ -774,7 +774,7 @@ void FFrontendFilter_CheckedOut::SourceControlOperationComplete(const FSourceCon
 /////////////////////////////////////////
 
 FFrontendFilter_NotSourceControlled::FFrontendFilter_NotSourceControlled(TSharedPtr<FFrontendFilterCategory> InCategory) 
-	: FFrontendFilter(InCategory),
+	: FFrontendFilter(MoveTemp(InCategory)),
 	bSourceControlEnabled(false),
 	bIsRequestStatusRunning(false),
 	bInitialRequestCompleted(false)
@@ -869,7 +869,7 @@ void FFrontendFilter_NotSourceControlled::SourceControlOperationComplete(const F
 /////////////////////////////////////////
 
 FFrontendFilter_Modified::FFrontendFilter_Modified(TSharedPtr<FFrontendFilterCategory> InCategory)
-	: FFrontendFilter(InCategory)
+	: FFrontendFilter(MoveTemp(InCategory))
 	, bIsCurrentlyActive(false)
 {
 	UPackage::PackageDirtyStateChangedEvent.AddRaw(this, &FFrontendFilter_Modified::OnPackageDirtyStateUpdated);
@@ -920,7 +920,7 @@ bool FFrontendFilter_ReplicatedBlueprint::PassesFilter(FAssetFilterType InItem)
 #define LOCTEXT_NAMESPACE "ContentBrowser"
 
 FFrontendFilter_ArbitraryComparisonOperation::FFrontendFilter_ArbitraryComparisonOperation(TSharedPtr<FFrontendFilterCategory> InCategory)
-	: FFrontendFilter(InCategory)
+	: FFrontendFilter(MoveTemp(InCategory))
 	, TagName(TEXT("TagName"))
 	, TargetTagValue(TEXT("Value"))
 	, ComparisonOp(ETextFilterComparisonOperation::NotEqual)
@@ -1087,7 +1087,7 @@ FString FFrontendFilter_ArbitraryComparisonOperation::ConvertOperationToString(E
 /////////////////////////////////////////
 
 FFrontendFilter_ShowOtherDevelopers::FFrontendFilter_ShowOtherDevelopers(TSharedPtr<FFrontendFilterCategory> InCategory)
-	: FFrontendFilter(InCategory)
+	: FFrontendFilter(MoveTemp(InCategory))
 	, BaseDeveloperPath(TEXT("/Game/Developers/"))
 	, BaseDeveloperPathAnsi()
 	, UserDeveloperPath(BaseDeveloperPath + FPaths::GameUserDeveloperFolderName() + TEXT("/"))
@@ -1168,7 +1168,7 @@ bool FFrontendFilter_ShowOtherDevelopers::GetShowOtherDeveloperAssets() const
 /////////////////////////////////////////
 
 FFrontendFilter_ShowRedirectors::FFrontendFilter_ShowRedirectors(TSharedPtr<FFrontendFilterCategory> InCategory)
-	: FFrontendFilter(InCategory)
+	: FFrontendFilter(MoveTemp(InCategory))
 {
 	bAreRedirectorsInBaseFilter = false;
 	RedirectorClassName = UObjectRedirector::StaticClass()->GetPathName();
@@ -1197,7 +1197,7 @@ bool FFrontendFilter_ShowRedirectors::PassesFilter(FAssetFilterType InItem) cons
 /////////////////////////////////////////
 
 FFrontendFilter_InUseByLoadedLevels::FFrontendFilter_InUseByLoadedLevels(TSharedPtr<FFrontendFilterCategory> InCategory) 
-	: FFrontendFilter(InCategory)
+	: FFrontendFilter(MoveTemp(InCategory))
 {
 	FEditorDelegates::MapChange.AddRaw(this, &FFrontendFilter_InUseByLoadedLevels::OnEditorMapChange);
 
@@ -1357,7 +1357,7 @@ void FFrontendFilter_InUseByLoadedLevels::OnEditorMapChange( uint32 MapChangeFla
 /////////////////////////////////////////
 
 FFrontendFilter_UsedInAnyLevel::FFrontendFilter_UsedInAnyLevel(TSharedPtr<FFrontendFilterCategory> InCategory)
-	: FFrontendFilter(InCategory)
+	: FFrontendFilter(MoveTemp(InCategory))
 {
 	// Prepare asset registry.
 	FAssetRegistryModule& AssetRegistryModule = FModuleManager::LoadModuleChecked<FAssetRegistryModule>(TEXT("AssetRegistry"));
@@ -1365,11 +1365,6 @@ FFrontendFilter_UsedInAnyLevel::FFrontendFilter_UsedInAnyLevel(TSharedPtr<FFront
 	check (AssetRegistry != nullptr);
 }
 
-FFrontendFilter_UsedInAnyLevel::~FFrontendFilter_UsedInAnyLevel()
-{
-	AssetRegistry = nullptr;
-}
-
 void FFrontendFilter_UsedInAnyLevel::ActiveStateChanged(bool bActive)
 {
 	LevelsDependencies.Empty();
@@ -1400,7 +1395,7 @@ bool FFrontendFilter_UsedInAnyLevel::PassesFilter(FAssetFilterType InItem) const
 /////////////////////////////////////////
 
 FFrontendFilter_NotUsedInAnyLevel::FFrontendFilter_NotUsedInAnyLevel(TSharedPtr<FFrontendFilterCategory> InCategory)
-	: FFrontendFilter(InCategory)
+	: FFrontendFilter(MoveTemp(InCategory))
 {
 	// Prepare asset registry.
 	FAssetRegistryModule& AssetRegistryModule = FModuleManager::LoadModuleChecked<FAssetRegistryModule>(TEXT("AssetRegistry"));
@@ -1408,12 +1403,6 @@ FFrontendFilter_NotUsedInAnyLevel::FFrontendFilter_NotUsedInAnyLevel(TSharedPtr<
 	check (AssetRegistry != nullptr);
 }
 
-
-FFrontendFilter_NotUsedInAnyLevel::~FFrontendFilter_NotUsedInAnyLevel()
-{
-	AssetRegistry = nullptr;
-}
-
 void FFrontendFilter_NotUsedInAnyLevel::ActiveStateChanged(bool bActive)
 {
 	LevelsDependencies.Empty();
@@ -1439,13 +1428,35 @@ bool FFrontendFilter_NotUsedInAnyLevel::PassesFilter(FAssetFilterType InItem) co
 	return false;
 }
 
+/////////////////////////////////////////
+// FFrontendFilter_NotUsedInAnyAsset
+/////////////////////////////////////////
+
+FFrontendFilter_NotUsedInAnyAsset::FFrontendFilter_NotUsedInAnyAsset(TSharedPtr<FFrontendFilterCategory> InCategory)
+	: FFrontendFilter(MoveTemp(InCategory))
+{
+	FAssetRegistryModule& AssetRegistryModule = FModuleManager::LoadModuleChecked<FAssetRegistryModule>(TEXT("AssetRegistry"));
+	AssetRegistry = &AssetRegistryModule.Get();
+}
+
+bool FFrontendFilter_NotUsedInAnyAsset::PassesFilter(FAssetFilterType InItem) const
+{
+	FAssetData ItemAssetData;
+	if (InItem.Legacy_TryGetAssetData(ItemAssetData))
+	{
+		TArray<FName> OutReferencers;
+		AssetRegistry->GetReferencers(ItemAssetData.PackageName, OutReferencers);
+		return OutReferencers.IsEmpty();
+	}
+	return false;
+}
 
 /////////////////////////////////////////
 // FFrontendFilter_Recent
 /////////////////////////////////////////
 
 FFrontendFilter_Recent::FFrontendFilter_Recent(TSharedPtr<FFrontendFilterCategory> InCategory)
-	: FFrontendFilter(InCategory)
+	: FFrontendFilter(MoveTemp(InCategory))
 	, bIsCurrentlyActive(false)
 {
 	UContentBrowserSettings::OnSettingChanged().AddRaw(this, &FFrontendFilter_Recent::ResetFilter);
@@ -1506,7 +1517,7 @@ void FFrontendFilter_Recent::ResetFilter(FName InName)
 /////////////////////////////////////////
 
 FFrontendFilter_Writable::FFrontendFilter_Writable(TSharedPtr<FFrontendFilterCategory> InCategory)
-	: FFrontendFilter(InCategory)
+	: FFrontendFilter(MoveTemp(InCategory))
 {
 }
 
@@ -1533,7 +1544,7 @@ bool FFrontendFilter_Writable::PassesFilter(FAssetFilterType InItem) const
 /////////////////////////////////////////
 
 FFrontendFilter_VirtualizedData::FFrontendFilter_VirtualizedData(TSharedPtr<FFrontendFilterCategory> InCategory)
-	: FFrontendFilter(InCategory)
+	: FFrontendFilter(MoveTemp(InCategory))
 {
 }
 
diff --git a/Engine/Source/Editor/ContentBrowser/Private/SFilterList.cpp b/Engine/Source/Editor/ContentBrowser/Private/SFilterList.cpp
--- a/Engine/Source/Editor/ContentBrowser/Private/SFilterList.cpp
+++ b/Engine/Source/Editor/ContentBrowser/Private/SFilterList.cpp
@@ -79,6 +79,7 @@ void SFilterList::Construct( const FArguments& InArgs )
 	AllFrontendFilters_Internal.Add( MakeShareable(new FFrontendFilter_InUseByLoadedLevels(DefaultCategory)) );
 	AllFrontendFilters_Internal.Add( MakeShareable(new FFrontendFilter_UsedInAnyLevel(DefaultCategory)) );
 	AllFrontendFilters_Internal.Add( MakeShareable(new FFrontendFilter_NotUsedInAnyLevel(DefaultCategory)) );
+	AllFrontendFilters_Internal.Add( MakeShareable(new FFrontendFilter_NotUsedInAnyAsset(DefaultCategory)) );
 	AllFrontendFilters_Internal.Add( MakeShareable(new FFrontendFilter_ArbitraryComparisonOperation(DefaultCategory)) );
 	AllFrontendFilters_Internal.Add(MakeShareable(new FFrontendFilter_Recent(DefaultCategory)));
 	AllFrontendFilters_Internal.Add( MakeShareable(new FFrontendFilter_NotSourceControlled(DefaultCategory)) );
diff --git a/Engine/Source/Editor/ContentBrowser/Public/FrontendFilterBase.h b/Engine/Source/Editor/ContentBrowser/Public/FrontendFilterBase.h
--- a/Engine/Source/Editor/ContentBrowser/Public/FrontendFilterBase.h
+++ b/Engine/Source/Editor/ContentBrowser/Public/FrontendFilterBase.h
@@ -19,7 +19,7 @@ public:
 class FFrontendFilter : public FFilterBase<FAssetFilterType>
 {
 public:
-	FFrontendFilter(TSharedPtr<FFrontendFilterCategory> InCategory) : FFilterBase<FAssetFilterType>(InCategory) {}
+	FFrontendFilter(TSharedPtr<FFrontendFilterCategory> InCategory) : FFilterBase<FAssetFilterType>(MoveTemp(InCategory)) {}
 
 	/** Invoke to set the source filter that is currently used to filter assets in the asset view */
 	virtual void SetCurrentFilter(TArrayView<const FName> InSourcePaths, const FContentBrowserDataFilter& InBaseFilter) { }
diff --git a/Engine/Source/Editor/ContentBrowser/Public/FrontendFilters.h b/Engine/Source/Editor/ContentBrowser/Public/FrontendFilters.h
--- a/Engine/Source/Editor/ContentBrowser/Public/FrontendFilters.h
+++ b/Engine/Source/Editor/ContentBrowser/Public/FrontendFilters.h
@@ -327,7 +327,6 @@ class CONTENTBROWSER_API FFrontendFilter_UsedInAnyLevel: public FFrontendFilter
 public:
 	/** Constructor/Destructor */
 	FFrontendFilter_UsedInAnyLevel(TSharedPtr<FFrontendFilterCategory> InCategory);
-	virtual ~FFrontendFilter_UsedInAnyLevel();
 
 	// FFrontendFilter implementation
 	virtual FString GetName() const override { return TEXT("UsedInAnyLevel"); }
@@ -343,13 +342,12 @@ private:
 	TSet<FName> LevelsDependencies;
 };
 
-/** A filter that only displays assets used by any level */
+/** A filter that only displays assets not used by any level */
 class CONTENTBROWSER_API FFrontendFilter_NotUsedInAnyLevel : public FFrontendFilter
 {
 public:
 	/** Constructor/Destructor */
 	FFrontendFilter_NotUsedInAnyLevel(TSharedPtr<FFrontendFilterCategory> InCategory);
-	virtual ~FFrontendFilter_NotUsedInAnyLevel();
 
 	// FFrontendFilter implementation
 	virtual FString GetName() const override { return TEXT("NotUsedInAnyLevel"); }
@@ -365,6 +363,25 @@ private:
 	TSet<FName> LevelsDependencies;
 };
 
+/** A filter that only displays assets not used in another asset (Note It does not update itself automatically) */
+class CONTENTBROWSER_API FFrontendFilter_NotUsedInAnyAsset : public FFrontendFilter
+{
+public:
+	/** Constructor/Destructor */
+	FFrontendFilter_NotUsedInAnyAsset(TSharedPtr<FFrontendFilterCategory> InCategory);
+
+	// FFrontendFilter implementation
+	virtual FString GetName() const override { return TEXT("NotUsedInAnyAsset"); }
+	virtual FText GetDisplayName() const override { return LOCTEXT("FFrontendFilter_NotUsedInAnyAsset", "Not Used In Any Asset"); }
+	virtual FText GetToolTipText() const override { return LOCTEXT("FFrontendFilter_NotUsedInAnyAsset", "Show only the assets that aren't used by another asset."); }
+
+	// IFilter implementation
+	virtual bool PassesFilter(FAssetFilterType InItem) const override;
+
+private:
+	class IAssetRegistry* AssetRegistry = nullptr;
+};
+
 /** A filter that displays recently opened assets */
 class CONTENTBROWSER_API FFrontendFilter_Recent : public FFrontendFilter
 {
