Remove drop shadow from virtual folders until we have a matching one.

diff --git a/Engine/Source/Editor/ContentBrowser/Private/AssetViewWidgets.cpp b/Engine/Source/Editor/ContentBrowser/Private/AssetViewWidgets.cpp
--- a/Engine/Source/Editor/ContentBrowser/Private/AssetViewWidgets.cpp
+++ b/Engine/Source/Editor/ContentBrowser/Private/AssetViewWidgets.cpp
@@ -204,6 +204,7 @@ TSharedRef<SWidget> FAssetViewItemHelper::CreateListTileItemContents(T* const In
 		}
 
 		const FSlateBrush* FolderBaseImage = nullptr;
+		const FSlateBrush* DropShadowImage = FAppStyle::Get().GetBrush("ContentBrowser.FolderItem.DropShadow");
 		if (bDeveloperFolder)
 		{
 			FolderBaseImage = FAppStyle::GetBrush("ContentBrowser.ListViewDeveloperFolderIcon");
@@ -215,6 +216,7 @@ TSharedRef<SWidget> FAssetViewItemHelper::CreateListTileItemContents(T* const In
 		else if (bVirtualFolder && ContentBrowserUtils::ShouldShowCustomVirtualFolderIcon())
 		{
 			FolderBaseImage = FAppStyle::GetBrush("ContentBrowser.ListViewVirtualFolderIcon");
+			DropShadowImage = FAppStyle::GetBrush("None");
 		}
 		else if (bPluginFolder && ContentBrowserUtils::ShouldShowPluginFolderIcon())
 		{
@@ -226,11 +228,12 @@ TSharedRef<SWidget> FAssetViewItemHelper::CreateListTileItemContents(T* const In
 		}
 
 		// Folder base
+		// clang-format off
 		ItemContentsOverlay->AddSlot()
 		.Padding(FMargin(5))
 		[
 			SNew(SBorder)
-			.BorderImage(FAppStyle::Get().GetBrush("ContentBrowser.FolderItem.DropShadow"))
+			.BorderImage(DropShadowImage)
 			.Padding(FMargin(0,0,2.0f,2.0f))
 			[
 				SNew(SImage)
@@ -238,6 +241,7 @@ TSharedRef<SWidget> FAssetViewItemHelper::CreateListTileItemContents(T* const In
 				.ColorAndOpacity(InTileOrListItem, &T::GetAssetColor)
 			]
 		];
+		// clang-format on
 
 		if (bCollectionFolder)
 		{
