ContentBrowser: Fix LOCTEXT conflict

diff --git a/Engine/Source/Editor/ContentBrowser/Public/FrontendFilters.h b/Engine/Source/Editor/ContentBrowser/Public/FrontendFilters.h
--- a/Engine/Source/Editor/ContentBrowser/Public/FrontendFilters.h
+++ b/Engine/Source/Editor/ContentBrowser/Public/FrontendFilters.h
@@ -373,7 +373,7 @@ public:
 	// FFrontendFilter implementation
 	virtual FString GetName() const override { return TEXT("NotUsedInAnyAsset"); }
 	virtual FText GetDisplayName() const override { return LOCTEXT("FFrontendFilter_NotUsedInAnyAsset", "Not Used In Any Asset"); }
-	virtual FText GetToolTipText() const override { return LOCTEXT("FFrontendFilter_NotUsedInAnyAsset", "Show only the assets that aren't used by another asset."); }
+	virtual FText GetToolTipText() const override { return LOCTEXT("FFrontendFilter_NotUsedInAnyAssetTooltip", "Show only the assets that aren't used by another asset."); }
 
 	// IFilter implementation
 	virtual bool PassesFilter(FAssetFilterType InItem) const override;
