Fix overriding Developer folder icon with virtual folder icon

diff --git a/Engine/Source/Editor/ContentBrowser/Private/SourcesViewWidgets.cpp b/Engine/Source/Editor/ContentBrowser/Private/SourcesViewWidgets.cpp
--- a/Engine/Source/Editor/ContentBrowser/Private/SourcesViewWidgets.cpp
+++ b/Engine/Source/Editor/ContentBrowser/Private/SourcesViewWidgets.cpp
@@ -102,25 +102,27 @@ void SAssetTreeItem::Construct( const FArguments& InArgs )
 	{
 		FolderType = EFolderType::Code;
 	}
-
-	if (ContentBrowserUtils::ShouldShowCustomVirtualFolderIcon())
-	{
-		FContentBrowserItemDataAttributeValue VirtualAttributeValue = Item.GetItemAttribute(ContentBrowserItemAttributes::ItemIsCustomVirtualFolder);
-		if (VirtualAttributeValue.IsValid() && VirtualAttributeValue.GetValue<bool>())
-		{
-			FolderType = EFolderType::CustomVirtual;
-		}
-	}
-	
-	if (ContentBrowserUtils::ShouldShowPluginFolderIcon())
+	else
 	{
-		if (InArgs._TreeItem->GetItem().IsInPlugin())
+		if (ContentBrowserUtils::ShouldShowPluginFolderIcon() && Item.IsInPlugin())
 		{
 			TSharedPtr<FTreeItem> Parent = InArgs._TreeItem->Parent.Pin();
 			if (!Parent.IsValid() || !Parent->GetItem().IsInPlugin())
 			{
 				FolderType = EFolderType::PluginRoot;
 			}
+			else
+			{
+				FolderType = EFolderType::Normal;
+			}
+		}
+		else if (ContentBrowserUtils::ShouldShowCustomVirtualFolderIcon())
+		{
+			FContentBrowserItemDataAttributeValue VirtualAttributeValue = Item.GetItemAttribute(ContentBrowserItemAttributes::ItemIsCustomVirtualFolder);
+			if (VirtualAttributeValue.IsValid() && VirtualAttributeValue.GetValue<bool>())
+			{
+				FolderType = EFolderType::CustomVirtual;
+			}
 		}
 	}
 
