[Path Picker] Fix an issue where path picker on initial notify would pass virtual paths, when the configuration was not set to support them

diff --git a/Engine/Source/Editor/ContentBrowser/Private/SPathPicker.cpp b/Engine/Source/Editor/ContentBrowser/Private/SPathPicker.cpp
--- a/Engine/Source/Editor/ContentBrowser/Private/SPathPicker.cpp
+++ b/Engine/Source/Editor/ContentBrowser/Private/SPathPicker.cpp
@@ -71,7 +71,14 @@ void SPathPicker::Construct( const FArguments& InArgs )
 
 		if (InArgs._PathPickerConfig.bNotifyDefaultPathSelected)
 		{
-			OnPathSelected.ExecuteIfBound(VirtualPath.ToString());
+			if (bOnPathSelectedPassesVirtualPaths)
+			{
+				OnPathSelected.ExecuteIfBound(VirtualPath.ToString());
+			}
+			else
+			{				
+				OnPathSelected.ExecuteIfBound(DefaultPath);
+			}
 		}
 	}
 }
