Addded a cpu trace scope to the function SPathView::HandleItemDataUpdated to track how much it contribute to the frame hitches in the editor traces.

diff --git a/Engine/Source/Editor/ContentBrowser/Private/SPathView.cpp b/Engine/Source/Editor/ContentBrowser/Private/SPathView.cpp
--- a/Engine/Source/Editor/ContentBrowser/Private/SPathView.cpp
+++ b/Engine/Source/Editor/ContentBrowser/Private/SPathView.cpp
@@ -2003,6 +2003,8 @@ bool SPathView::IsTreeItemSelected(TSharedPtr<FTreeItem> TreeItem) const
 
 void SPathView::HandleItemDataUpdated(TArrayView<const FContentBrowserItemDataUpdate> InUpdatedItems)
 {
+	TRACE_CPUPROFILER_EVENT_SCOPE(SPathView::HandleItemDataUpdated);
+
 	if (InUpdatedItems.Num() == 0)
 	{
 		return;
