Fixed SPathView not matching localized folder names (from aliases)

#rnx

diff --git a/Engine/Source/Editor/ContentBrowser/Private/SPathView.cpp b/Engine/Source/Editor/ContentBrowser/Private/SPathView.cpp
--- a/Engine/Source/Editor/ContentBrowser/Private/SPathView.cpp
+++ b/Engine/Source/Editor/ContentBrowser/Private/SPathView.cpp
@@ -1589,8 +1589,9 @@ void SPathView::Populate(const bool bIsRefreshingFilter)
 				if (bFilteringByText)
 				{
 					// Use the whole path so we deliberately include any children of matched parents in the filtered list
+					// Also check the display name so we match on a localized folder name
 					const FString PathStr = InItemData.GetVirtualPath().ToString();
-					if (!SearchBoxFolderFilter->PassesFilter(PathStr))
+					if (!SearchBoxFolderFilter->PassesFilter(PathStr) && !SearchBoxFolderFilter->PassesFilter(InItemData.GetDisplayName().ToString()))
 					{
 						return true; // continue enumerating
 					}
