Optimizing population and filtering of SPathView
- Decouple population from text filtering
- Factor out tree data from widget as a precursor to being able to bind data externally
- Search for parents to build from leaf to root rather than reverse

diff --git a/Engine/Source/Editor/ContentBrowser/Private/ContentBrowserUtils.cpp b/Engine/Source/Editor/ContentBrowser/Private/ContentBrowserUtils.cpp
--- a/Engine/Source/Editor/ContentBrowser/Private/ContentBrowserUtils.cpp
+++ b/Engine/Source/Editor/ContentBrowser/Private/ContentBrowserUtils.cpp
@@ -41,6 +41,7 @@
 #include "Misc/Optional.h"
 #include "Misc/Paths.h"
 #include "SAssetView.h"
+#include "Settings/ContentBrowserSettings.h"
 #include "SFilterList.h"
 #include "SPathView.h"
 #include "SlateOptMacros.h"
@@ -938,4 +939,59 @@ bool ContentBrowserUtils::ShouldShowRedirectors(TSharedPtr<SFilterList> Filters)
 	return false;
 }
 
+FContentBrowserInstanceConfig* ContentBrowserUtils::GetContentBrowserConfig(FName InstanceName)
+{
+	if (InstanceName.IsNone())
+	{
+		return nullptr;
+	}
+
+	FContentBrowserInstanceConfig* Config = UContentBrowserConfig::Get()->Instances.Find(InstanceName);
+	if (Config == nullptr)
+	{
+		return nullptr;
+	}
+
+	return Config;
+}
+
+FPathViewConfig* ContentBrowserUtils::GetPathViewConfig(FName InstanceName)
+{
+	if (InstanceName.IsNone())
+	{
+		return nullptr;
+	}
+
+	FContentBrowserInstanceConfig* Config = UContentBrowserConfig::Get()->Instances.Find(InstanceName);
+	if (Config == nullptr)
+	{
+		return nullptr;
+	}
+
+	return &Config->PathView;
+}
+
+EContentBrowserItemAttributeFilter ContentBrowserUtils::GetContentBrowserItemAttributeFilter(FName InstanceName)
+{
+	const UContentBrowserSettings* ContentBrowserSettings = GetDefault<UContentBrowserSettings>();
+	bool bDisplayEngineContent = ContentBrowserSettings->GetDisplayEngineFolder();
+	bool bDisplayPluginContent = ContentBrowserSettings->GetDisplayPluginFolders();
+	bool bDisplayDevelopersContent = ContentBrowserSettings->GetDisplayDevelopersFolder();
+	bool bDisplayL10NContent = ContentBrowserSettings->GetDisplayL10NFolder();
+
+	// check to see if we have an instance config that overrides the defaults in UContentBrowserSettings
+	if (FContentBrowserInstanceConfig* EditorConfig = GetContentBrowserConfig(InstanceName))
+	{
+		bDisplayEngineContent = EditorConfig->bShowEngineContent;
+		bDisplayPluginContent = EditorConfig->bShowPluginContent;
+		bDisplayDevelopersContent = EditorConfig->bShowDeveloperContent;
+		bDisplayL10NContent = EditorConfig->bShowLocalizedContent;
+	}
+
+	return EContentBrowserItemAttributeFilter::IncludeProject
+		 | (bDisplayEngineContent ? EContentBrowserItemAttributeFilter::IncludeEngine : EContentBrowserItemAttributeFilter::IncludeNone)
+		 | (bDisplayPluginContent ? EContentBrowserItemAttributeFilter::IncludePlugins : EContentBrowserItemAttributeFilter::IncludeNone)
+		 | (bDisplayDevelopersContent ? EContentBrowserItemAttributeFilter::IncludeDeveloper : EContentBrowserItemAttributeFilter::IncludeNone)
+		 | (bDisplayL10NContent ? EContentBrowserItemAttributeFilter::IncludeLocalized : EContentBrowserItemAttributeFilter::IncludeNone);
+}
 #undef LOCTEXT_NAMESPACE
diff --git a/Engine/Source/Editor/ContentBrowser/Private/ContentBrowserUtils.h b/Engine/Source/Editor/ContentBrowser/Private/ContentBrowserUtils.h
--- a/Engine/Source/Editor/ContentBrowser/Private/ContentBrowserUtils.h
+++ b/Engine/Source/Editor/ContentBrowser/Private/ContentBrowserUtils.h
@@ -24,9 +24,12 @@ class SWidget;
 struct FARFilter;
 struct FAssetData;
 struct FContentBrowserDataFilter;
+struct FContentBrowserInstanceConfig;
 struct FContentBrowserItem;
 struct FContentBrowserItemPath;
+struct FPathViewConfig;
 enum class EContentBrowserIsFolderVisibleFlags : uint8;
+enum class EContentBrowserItemAttributeFilter : uint8;
 
 namespace ContentBrowserUtils
 {
@@ -144,4 +147,22 @@ namespace ContentBrowserUtils
 	
 	/** Returns whether the content browser should be showing redirectors based on the current filter state or content browser settings */
 	bool ShouldShowRedirectors(TSharedPtr<SFilterList> Filters);
+
+	/**
+	 * Returns config settings for the given content browser name.
+	 * A name of none returns null.
+	 */
+	FContentBrowserInstanceConfig* GetContentBrowserConfig(FName InstanceName);
+	
+	/** 
+	 * returns config settings for the path view for the given content browser name.
+	 * A name of none returns null.
+	 */
+	FPathViewConfig* GetPathViewConfig(FName InstanceName);
+
+	/**
+	 * Returns the attribute filter to use when retrieving conent browser data for the given instance.
+	 * An instance name of None gives global settings.
+	 */
+	EContentBrowserItemAttributeFilter GetContentBrowserItemAttributeFilter(FName InstanceName);
 }
diff --git a/Engine/Source/Editor/ContentBrowser/Private/PathViewTypes.cpp b/Engine/Source/Editor/ContentBrowser/Private/PathViewTypes.cpp
--- a/Engine/Source/Editor/ContentBrowser/Private/PathViewTypes.cpp
+++ b/Engine/Source/Editor/ContentBrowser/Private/PathViewTypes.cpp
@@ -2,6 +2,7 @@
 
 #include "PathViewTypes.h"
 
+#include "Algo/Copy.h"
 #include "Containers/UnrealString.h"
 #include "ContentBrowserItemData.h"
 #include "HAL/Platform.h"
@@ -57,11 +58,36 @@ void FTreeItem::RemoveItemData(const FContentBrowserItemData& InItemData)
 	Item.Remove(InItemData);
 }
 
+void FTreeItem::SetItemData(FContentBrowserItem InItem)
+{
+	Item = MoveTemp(InItem);
+}
+
 const FContentBrowserItem& FTreeItem::GetItem() const
 {
 	return Item;
 }
 
+void FTreeItem::SetVisible(bool bInIsVisible)
+{
+	bIsVisible = bInIsVisible;
+}
+
+void FTreeItem::SetHasVisibleDescendants(bool bValue)
+{
+	bHasVisibleDescendants = bValue;
+}
+
+bool FTreeItem::GetHasVisibleDescendants() const
+{
+	return bHasVisibleDescendants;
+}
+
+bool FTreeItem::IsVisible() const
+{
+	return bIsVisible || bHasVisibleDescendants;
+}
+
 FSimpleMulticastDelegate& FTreeItem::OnRenameRequested()
 {
 	return RenameRequestedEvent;
@@ -93,9 +119,36 @@ bool FTreeItem::IsChildOf(const FTreeItem& InParent)
 	return false;
 }
 
+void FTreeItem::AddChild(const TSharedRef<FTreeItem>& InChild)
+{
+	checkSlow(!AllChildren.Contains(InChild));
+	AllChildren.Add(InChild);
+	InChild->Parent = AsWeak();
+	bChildrenRequireSort = true;
+}
+
+void FTreeItem::RemoveChild(const TSharedRef<FTreeItem>& InChild)
+{
+	if (InChild->Parent == AsWeak())
+	{
+		AllChildren.Remove(InChild);
+		InChild->Parent = nullptr;
+	}
+}
+
+void FTreeItem::RemoveAllChildren()
+{
+	AllChildren.Reset();
+}
+
+TConstArrayView<TSharedPtr<FTreeItem>> FTreeItem::GetChildren() const
+{
+	return AllChildren;
+}
+
 TSharedPtr<FTreeItem> FTreeItem::GetChild(const FName InChildFolderName) const
 {
-	for (const TSharedPtr<FTreeItem>& Child : Children)
+	for (const TSharedPtr<FTreeItem>& Child : AllChildren)
 	{
 		if (Child->Item.GetItemName() == InChildFolderName)
 		{
@@ -106,6 +159,11 @@ TSharedPtr<FTreeItem> FTreeItem::GetChild(const FName InChildFolderName) const
 	return nullptr;
 }
 
+TSharedPtr<FTreeItem> FTreeItem::GetParent() const
+{
+	return Parent.Pin();
+}
+
 TSharedPtr<FTreeItem> FTreeItem::FindItemRecursive(const FName InFullPath)
 {
 	if (InFullPath == Item.GetVirtualPath())
@@ -113,7 +171,7 @@ TSharedPtr<FTreeItem> FTreeItem::FindItemRecursive(const FName InFullPath)
 		return SharedThis(this);
 	}
 
-	for (const TSharedPtr<FTreeItem>& Child : Children)
+	for (const TSharedPtr<FTreeItem>& Child : AllChildren)
 	{
 		if (TSharedPtr<FTreeItem> ChildItem = Child->FindItemRecursive(InFullPath))
 		{
@@ -124,38 +182,35 @@ TSharedPtr<FTreeItem> FTreeItem::FindItemRecursive(const FName InFullPath)
 	return nullptr;
 }
 
+void FTreeItem::ForAllChildrenRecursive(TFunctionRef<void(const TSharedRef<FTreeItem>&)> Functor)
+{
+	for (const TSharedPtr<FTreeItem>& Child : AllChildren)
+	{
+		if (Child.IsValid())
+		{
+			Functor(Child.ToSharedRef());
+			Child->ForAllChildrenRecursive(Functor);
+		}
+	}
+}
+
 void FTreeItem::RequestSortChildren()
 {
 	bChildrenRequireSort = true;
 }
 
-void FTreeItem::SortChildrenIfNeeded()
+void FTreeItem::GetSortedVisibleChildren(TArray<TSharedPtr<FTreeItem>>& OutChildren)
 {
 	if (bChildrenRequireSort)
 	{
-		if (SortOverride.IsBound())
-		{
-			SortOverride.Execute(this, Children);
-		}
-		else
-		{
-			Children.Sort([](TSharedPtr<FTreeItem> A, TSharedPtr<FTreeItem> B) -> bool
-			{
-				return A->Item.GetDisplayName().ToString() < B->Item.GetDisplayName().ToString();
-			});
-		}
-
+		UE::PathView::DefaultSort(AllChildren);
 		bChildrenRequireSort = false;
 	}
+	OutChildren.Reset();
+	Algo::CopyIf(AllChildren, OutChildren, UE_PROJECTION_MEMBER(FTreeItem, IsVisible));
 }
 
 bool FTreeItem::IsDisplayOnlyFolder() const
 {
 	return GetItem().IsDisplayOnlyFolder();
 }
-
-void FTreeItem::SetSortOverride(FSortTreeItemChildrenDelegate& InSortOverride)
-{
-	SortOverride = InSortOverride;
-}
-
diff --git a/Engine/Source/Editor/ContentBrowser/Private/PathViewTypes.h b/Engine/Source/Editor/ContentBrowser/Private/PathViewTypes.h
--- a/Engine/Source/Editor/ContentBrowser/Private/PathViewTypes.h
+++ b/Engine/Source/Editor/ContentBrowser/Private/PathViewTypes.h
@@ -14,6 +14,11 @@ class FTreeItem;
 // Sorts array of tree items
 DECLARE_DELEGATE_TwoParams(FSortTreeItemChildrenDelegate, const FTreeItem* /* OptionalOwner */, TArray<TSharedPtr<FTreeItem>>& /* Items to sort */);
 
+namespace UE::PathView
+{
+void DefaultSort(TArray<TSharedPtr<FTreeItem>>& InChildren);
+}
+
 /** A folder item shown in the asset tree */
 class FTreeItem : public TSharedFromThis<FTreeItem>
 {
@@ -26,12 +31,12 @@ public:
 	explicit FTreeItem(FContentBrowserItemData&& InItemData);
 	explicit FTreeItem(const FContentBrowserItemData& InItemData);
 
-	void AppendItemData(const FContentBrowserItem& InItem);
+	// Replace the current item data or pending item data with the given item data
+	void SetItemData(FContentBrowserItem InItem);
 
+	void AppendItemData(const FContentBrowserItem& InItem);
 	void AppendItemData(const FContentBrowserItemData& InItemData);
-
 	void RemoveItemData(const FContentBrowserItem& InItem);
-
 	void RemoveItemData(const FContentBrowserItemData& InItemData);
 
 	/** Get the underlying Content Browser item */
@@ -40,6 +45,18 @@ public:
 	/** Get the event fired whenever a rename is requested */
 	FSimpleMulticastDelegate& OnRenameRequested();
 
+	/** Set whether the item passes current filtering and should be displayed */
+	void SetVisible(bool bInIsVisible);
+
+	/** Set whether this item has any descendants which explicitly passed filtering */
+	void SetHasVisibleDescendants(bool bValue);
+
+	/** Return whether this item has any descendants which explicitly passed filtering */
+	bool GetHasVisibleDescendants() const;
+
+	/** Returns whether the item passes current filtering and should be displayed, or whether any of its descendants did */
+	bool IsVisible() const;
+
 	/** True if this folder is in the process of being named */
 	bool IsNamingFolder() const;
 
@@ -49,44 +66,61 @@ public:
 	/** Returns true if this item is a child of the specified item */
 	bool IsChildOf(const FTreeItem& InParent);
 
+	/** Add a child item and link its Parent backreference. */
+	void AddChild(const TSharedRef<FTreeItem>& InChild);
+
+	/** Remove a child item and unlink its Parent backreference. */
+	void RemoveChild(const TSharedRef<FTreeItem>& InChild);
+
+	/** Remove all children when recycling this item. They may be re-added later. */
+	void RemoveAllChildren();
+
+	/** Get a view of all the direct children of this node*/
+	TConstArrayView<TSharedPtr<FTreeItem>> GetChildren() const;
+
 	/** Returns the child item by name or NULL if the child does not exist */
 	TSharedPtr<FTreeItem> GetChild(const FName InChildFolderName) const;
 
+	/** Returns the parent item if any */
+	TSharedPtr<FTreeItem> GetParent() const;
+
 	/** Finds the child who's path matches the one specified */
 	TSharedPtr<FTreeItem> FindItemRecursive(const FName InFullPath);
 
+	/** Execute the given functor on all children of this item recursively */
+	void ForAllChildrenRecursive(TFunctionRef<void(const TSharedRef<FTreeItem>&)> Functor);
+
 	/** Request that the children be sorted the next time someone calls SortChildrenIfNeeded */
 	void RequestSortChildren();
 
-	/** Sort the children by name (but only if they need to be) */
-	void SortChildrenIfNeeded();
+	/** Sort the children if necessary and populate the output parameter with a copy */
+	void GetSortedVisibleChildren(TArray<TSharedPtr<FTreeItem>>& OutChildren);
 
 	/** Represents a folder that does not correspond to a mounted location */
 	bool IsDisplayOnlyFolder() const;
 
-	/** Set delegate to use during sort */
-	void SetSortOverride(FSortTreeItemChildrenDelegate& InSortOverride);
-
-public:
+private:
 	/** The children of this tree item */
-	TArray<TSharedPtr<FTreeItem>> Children;
+	TArray<TSharedPtr<FTreeItem>> AllChildren;
 
 	/** The parent folder for this item */
 	TWeakPtr<FTreeItem> Parent;
 
-private:
 	/** Underlying Content Browser item data */
 	FContentBrowserItem Item;
 
 	/** Broadcasts whenever a rename is requested */
 	FSimpleMulticastDelegate RenameRequestedEvent;
 
-	/** Delegate used to control sort */
-	FSortTreeItemChildrenDelegate SortOverride;
-
 	/** If true, this folder is in the process of being named */
 	bool bNamingFolder = false;
 
 	/** If true, the children of this item need sorting */
 	bool bChildrenRequireSort = false;
+
+	/** Whether this node has passed the current set of filters in use */
+	bool bIsVisible = true;
+
+	/** Whether this node has any descendants that are visible, so this node needs to be shown too */
+	bool bHasVisibleDescendants = true;
 };
diff --git a/Engine/Source/Editor/ContentBrowser/Private/SContentBrowser.cpp b/Engine/Source/Editor/ContentBrowser/Private/SContentBrowser.cpp
--- a/Engine/Source/Editor/ContentBrowser/Private/SContentBrowser.cpp
+++ b/Engine/Source/Editor/ContentBrowser/Private/SContentBrowser.cpp
@@ -3,6 +3,7 @@
 
 #include "SContentBrowser.h"
 
+#include "Algo/AnyOf.h"
 #include "AssetContextMenu.h"
 #include "AssetRegistry/ARFilter.h"
 #include "AssetRegistry/AssetData.h"
@@ -1627,16 +1628,9 @@ void SContentBrowser::PrepareToSyncItems(TArrayView<const FContentBrowserItem> I
 	// Check to see if any item paths don't exist (this can happen if we haven't ticked since the path was created)
 	if (!bRepopulate)
 	{
-		for (const FContentBrowserItem& ItemToSync : ItemsToSync)
-		{
-			const FName VirtualPath = *FPaths::GetPath(ItemToSync.GetVirtualPath().ToString());
-			TSharedPtr<FTreeItem> Item = PathViewPtr->FindTreeItem(VirtualPath);
-			if (!Item.IsValid())
- 			{
-				bRepopulate = true;
- 				break;
- 			}
-		}
+		bRepopulate = Algo::AnyOf(ItemsToSync, [this](const FContentBrowserItem& Item) {
+			return !PathViewPtr->DoesItemExist(Item.GetVirtualPath());
+		});
 	}
 	
 	if (bDisableFiltersThatHideAssets)
diff --git a/Engine/Source/Editor/ContentBrowser/Private/SPathPicker.cpp b/Engine/Source/Editor/ContentBrowser/Private/SPathPicker.cpp
--- a/Engine/Source/Editor/ContentBrowser/Private/SPathPicker.cpp
+++ b/Engine/Source/Editor/ContentBrowser/Private/SPathPicker.cpp
@@ -42,6 +42,7 @@ void SPathPicker::Construct( const FArguments& InArgs )
 	OnGetPathContextMenuExtender = InArgs._PathPickerConfig.OnGetPathContextMenuExtender;
 	bOnPathSelectedPassesVirtualPaths = InArgs._PathPickerConfig.bOnPathSelectedPassesVirtualPaths;
 
+	// clang-format off
 	ChildSlot
 	[
 		SAssignNew(PathViewPtr, SPathView)
@@ -55,20 +56,15 @@ void SPathPicker::Construct( const FArguments& InArgs )
 		.SelectionMode(ESelectionMode::Single)
 		.CustomFolderPermissionList(InArgs._PathPickerConfig.CustomFolderPermissionList)
 		.ShowFavorites(InArgs._PathPickerConfig.bShowFavorites)
+		.DefaultPath(InArgs._PathPickerConfig.bAddDefaultPath ? InArgs._PathPickerConfig.DefaultPath : FString{})
 	];
+	// clang-format on
 
 	const FString& DefaultPath = InArgs._PathPickerConfig.DefaultPath;
-	if ( !DefaultPath.IsEmpty() && PathViewPtr->InternalPathPassesBlockLists(*DefaultPath))
+	if (!DefaultPath.IsEmpty() && PathViewPtr->InternalPathPassesBlockLists(*DefaultPath))
 	{
 		const FName VirtualPath = IContentBrowserDataModule::Get().GetSubsystem()->ConvertInternalPathToVirtual(*DefaultPath);
-		if (InArgs._PathPickerConfig.bAddDefaultPath && !PathViewPtr->FindTreeItem(VirtualPath))
-		{
-			const FString DefaultPathLeafName = FPaths::GetPathLeaf(VirtualPath.ToString());
-			PathViewPtr->AddFolderItem(FContentBrowserItemData(nullptr, EContentBrowserItemFlags::Type_Folder, VirtualPath, *DefaultPathLeafName, FText(), nullptr), /*bUserNamed*/false);
-		}
-
-		PathViewPtr->SetSelectedPaths({ VirtualPath.ToString() });
-
+		// Path is created by SPathView::Construct if necessary
 		if (InArgs._PathPickerConfig.bNotifyDefaultPathSelected)
 		{
 			if (bOnPathSelectedPassesVirtualPaths)
diff --git a/Engine/Source/Editor/ContentBrowser/Private/SPathView.cpp b/Engine/Source/Editor/ContentBrowser/Private/SPathView.cpp
--- a/Engine/Source/Editor/ContentBrowser/Private/SPathView.cpp
+++ b/Engine/Source/Editor/ContentBrowser/Private/SPathView.cpp
@@ -2,6 +2,8 @@
 
 #include "SPathView.h"
 
+#include "Algo/AnyOf.h"
+#include "Algo/Copy.h"
 #include "AssetRegistry/ARFilter.h"
 #include "AssetToolsModule.h"
 #include "Containers/ContainerAllocationPolicies.h"
@@ -78,6 +80,732 @@ struct FGeometry;
 
 #define LOCTEXT_NAMESPACE "ContentBrowser"
 
+namespace UE::PathView
+{
+TSharedRef<FTreeItem> CreateOrReuseNode(FContentBrowserItemData&& InData,
+	TMap<FName, TSharedPtr<FTreeItem>>* OldItemsByInvariantPath)
+{
+	if (OldItemsByInvariantPath)
+	{
+		TSharedPtr<FTreeItem> ExistingItem;
+		// Reolve old value so we don't pick it out again when looking at another item from a different source
+		if (OldItemsByInvariantPath->RemoveAndCopyValue(InData.GetInvariantPath(), ExistingItem)
+			&& ExistingItem.IsValid())
+		{
+			ExistingItem->RemoveAllChildren();
+			ExistingItem->SetItemData(FContentBrowserItem(MoveTemp(InData)));
+			return ExistingItem.ToSharedRef();
+		}
+	}
+	return MakeShared<FTreeItem>(MoveTemp(InData));
+}
+
+void DefaultSort(TArray<TSharedPtr<FTreeItem>>& InChildren)
+{
+	if (InChildren.Num() < 2)
+	{
+		return;
+	}
+
+	static const FString ClassesPrefix = TEXT("Classes_");
+
+	struct FItemSortInfo
+	{
+		// Name to display
+		FString FolderName;
+		float Priority;
+		int32 SpecialDefaultFolderPriority;
+		bool bIsClassesFolder;
+		TSharedPtr<FTreeItem> TreeItem;
+		// Name to use when comparing "MyPlugin" vs "Classes_MyPlugin", looking up a plugin by name and other situations
+		FName ItemNameWithoutClassesPrefix;
+	};
+
+	TArray<FItemSortInfo> SortInfoArray;
+	SortInfoArray.Reserve(InChildren.Num());
+
+	const TArray<FName>& SpecialSortFolders =
+		IContentBrowserDataModule::Get().GetSubsystem()->GetPathViewSpecialSortFolders();
+
+	// Generate information needed to perform sort
+	for (TSharedPtr<FTreeItem>& It : InChildren)
+	{
+		FItemSortInfo& SortInfo = SortInfoArray.AddDefaulted_GetRef();
+		SortInfo.TreeItem = It;
+
+		const FName InvariantPathFName = It->GetItem().GetInvariantPath();
+		FNameBuilder InvariantPathBuilder(InvariantPathFName);
+		const FStringView InvariantPath(InvariantPathBuilder);
+
+		bool bIsRootInvariantFolder = false;
+		if (InvariantPath.Len() > 1)
+		{
+			FStringView RootInvariantFolder(InvariantPath);
+			RootInvariantFolder.RightChopInline(1);
+			int32 SecondSlashIndex = INDEX_NONE;
+			bIsRootInvariantFolder = !RootInvariantFolder.FindChar(TEXT('/'), SecondSlashIndex);
+		}
+
+		SortInfo.FolderName = It->GetItem().GetDisplayName().ToString();
+
+		SortInfo.bIsClassesFolder = false;
+		if (bIsRootInvariantFolder)
+		{
+			FNameBuilder ItemNameBuilder(It->GetItem().GetItemName());
+			const FStringView ItemNameView(ItemNameBuilder);
+			if (ItemNameView.StartsWith(ClassesPrefix))
+			{
+				SortInfo.bIsClassesFolder = true;
+				SortInfo.ItemNameWithoutClassesPrefix = FName(ItemNameView.RightChop(ClassesPrefix.Len()));
+			}
+
+			if (SortInfo.FolderName.StartsWith(ClassesPrefix))
+			{
+				SortInfo.bIsClassesFolder = true;
+				SortInfo.FolderName.RightChopInline(ClassesPrefix.Len(), EAllowShrinking::No);
+			}
+		}
+
+		if (SortInfo.ItemNameWithoutClassesPrefix.IsNone())
+		{
+			SortInfo.ItemNameWithoutClassesPrefix = It->GetItem().GetItemName();
+		}
+
+		if (SortInfo.bIsClassesFolder)
+		{
+			// Sort using a path without "Classes_" prefix
+			FStringView InvariantWithoutClassesPrefix(InvariantPath);
+			InvariantWithoutClassesPrefix.RightChopInline(1);
+			if (InvariantWithoutClassesPrefix.StartsWith(ClassesPrefix))
+			{
+				InvariantWithoutClassesPrefix.RightChopInline(ClassesPrefix.Len());
+				FNameBuilder Builder;
+				Builder.Append(TEXT("/"));
+				Builder.Append(InvariantWithoutClassesPrefix);
+				SortInfo.SpecialDefaultFolderPriority = SpecialSortFolders.IndexOfByKey(FName(Builder));
+			}
+			else
+			{
+				SortInfo.SpecialDefaultFolderPriority = SpecialSortFolders.IndexOfByKey(InvariantPathFName);
+			}
+		}
+		else
+		{
+			SortInfo.SpecialDefaultFolderPriority = SpecialSortFolders.IndexOfByKey(InvariantPathFName);
+		}
+
+		if (bIsRootInvariantFolder)
+		{
+			if (SortInfo.SpecialDefaultFolderPriority == INDEX_NONE)
+			{
+				SortInfo.Priority = FContentBrowserSingleton::Get()
+										.GetPluginSettings(SortInfo.ItemNameWithoutClassesPrefix)
+										.RootFolderSortPriority;
+			}
+			else
+			{
+				SortInfo.Priority = 1.f;
+			}
+		}
+		else
+		{
+			if (SortInfo.SpecialDefaultFolderPriority != INDEX_NONE)
+			{
+				SortInfo.Priority = 1.f;
+			}
+			else
+			{
+				SortInfo.Priority = 0.f;
+			}
+		}
+	}
+
+	// Perform sort
+	SortInfoArray.Sort([](const FItemSortInfo& SortInfoA, const FItemSortInfo& SortInfoB) -> bool {
+		if (SortInfoA.Priority != SortInfoB.Priority)
+		{
+			// Not the same priority, use priority to sort
+			return SortInfoA.Priority > SortInfoB.Priority;
+		}
+		else if (SortInfoA.SpecialDefaultFolderPriority != SortInfoB.SpecialDefaultFolderPriority)
+		{
+			// Special folders use the index to sort. Non special folders are all set to 0.
+			return SortInfoA.SpecialDefaultFolderPriority < SortInfoB.SpecialDefaultFolderPriority;
+		}
+		else
+		{
+			// If either is a class folder and names without classes prefix are same
+			if ((SortInfoA.bIsClassesFolder != SortInfoB.bIsClassesFolder)
+				&& (SortInfoA.ItemNameWithoutClassesPrefix == SortInfoB.ItemNameWithoutClassesPrefix))
+			{
+				return !SortInfoA.bIsClassesFolder;
+			}
+
+			// Two non special folders of the same priority, sort alphabetically
+			const int32 CompareResult =
+				UE::ComparisonUtility::CompareWithNumericSuffix(SortInfoA.FolderName, SortInfoB.FolderName);
+			if (CompareResult != 0)
+			{
+				return CompareResult < 0;
+			}
+			else
+			{
+				// Classes folders have the same name so sort them adjacent but under non-classes
+				return !SortInfoA.bIsClassesFolder;
+			}
+		}
+	});
+
+	// Replace with sorted array
+	TArray<TSharedPtr<FTreeItem>> NewList;
+	NewList.Reserve(SortInfoArray.Num());
+	for (const FItemSortInfo& It : SortInfoArray)
+	{
+		NewList.Add(It.TreeItem);
+	}
+	InChildren = MoveTemp(NewList);
+}
+
+} // namespace UE::PathView
+
+// Struct to factor out path view data fetching/filtering as a precursor to being able to bind this data to the view
+// instead of fetching it internally
+struct FPathViewData
+{
+public:
+	FPathViewData(FName InContentBrowserName, bool InFlat)
+		: OwningContentBrowserName(InContentBrowserName)
+		, bFlat(InFlat)
+		, FolderPathTextFilter(decltype(FolderPathTextFilter)::FItemToStringArray::CreateStatic(
+			  [](FStringView Input, TArray<FString>& Out) { Out.Emplace(FString(Input)); }))
+	{
+	}
+
+	~FPathViewData() { }
+
+	uint64 GetVersion()
+	{
+		return Version;
+	}
+
+	// Return an array that can be bound to a tree view widget for the current visible set of root items
+	TArray<TSharedPtr<FTreeItem>>* GetVisibleRootItems()
+	{
+		return &VisibleRootItems;
+	}
+
+	TTextFilter<FStringView>& GetFolderPathTextFilter()
+	{
+		return FolderPathTextFilter;
+	}
+
+	// Fetch all data from the content browser data backend and transform it into the tree data
+	void PopulateFullFolderTree(const FContentBrowserDataCompiledFilter& InFilter);
+	// Fetch favorite folders from config, filter them against the content browser data filter
+	// bFlat parameter adds all items at the root of the tree and not create parents
+	void PopulateWithFavorites(const FContentBrowserDataCompiledFilter& InFilter);
+	// Apply the current text filter to everything in the tree
+	void FilterFullFolderTree();
+	// Clear the filter state of all items in the tree
+	void ClearItemFilterState();
+	// Sort the roots of the tree
+	void SortRootItems();
+	// Add an item to the tree by its virtual path, reusing an old object if possible to ensure persistence of
+	// selection/expansion for the widget the items are bound to
+	TSharedRef<FTreeItem> AddFolderItem(FContentBrowserItemData&& InItemData);
+
+	// Remove the given item from the tree whether it's a root or child.
+	void RemoveFolderItem(const TSharedRef<FTreeItem>& InItem);
+
+	// Find an item with the exact virtual path
+	TSharedPtr<FTreeItem> FindTreeItem(FName InVirtualPath, bool bVisibleOnly = false);
+	// Search the tree for the item furthest from the root that matches the given path, if any
+	// Searches all items, not just visible ones according to the current text filter
+	TSharedPtr<FTreeItem> FindBestItemForPath(FStringView InVirtualPath);
+
+	// Apply new/modified/removed data notifications to the tree
+	// bFlat parameter adds all items at the root of the tree as in the favorites tree
+	void ProcessDataUpdates(TConstArrayView<FContentBrowserItemDataUpdate> InUpdatedItems,
+		const FContentBrowserDataCompiledFilter& InFilter);
+
+protected:
+	TSharedRef<FTreeItem> AddFolderItemInternal(FContentBrowserItemData&& InItemData,
+		TMap<FName, TSharedPtr<FTreeItem>>* OldItemsByInvariantPath);
+
+	// Remove the given item data from the tree - if this results in an item having no data from any sources, the tree
+	// item is removed. If the item was removed and had a parent, the parent is returned.
+	TSharedPtr<FTreeItem> TryRemoveFolderItemInternal(const FContentBrowserItemData& InItem);
+
+	bool PassesTextFilter(const TSharedPtr<FTreeItem>& InItem);
+
+	// Incremented to trigger tree rebuild from changes to the tree contents
+	uint64 Version;
+	// Items with no parent
+	TArray<TSharedPtr<FTreeItem>> RootItems;
+	// Items with no parent
+	TArray<TSharedPtr<FTreeItem>> VisibleRootItems;
+	// Mapping of full virtual path such as '/All/Game/Maps/Arena' to items
+	TMap<FName, TSharedPtr<FTreeItem>> VirtualPathToItem;
+	// Mapping of path that doesn't change based on display settings (e.g. '/MyPlugin/MyAsset') to item
+	TMap<FName, TSharedPtr<FTreeItem>> InvariantPathToItem;
+
+	// Used for retrieving saved settings per content browser instance
+	FName OwningContentBrowserName;
+
+	// If true, parent items are not created and all items are added as roots.
+	bool bFlat;
+
+	TTextFilter<FStringView> FolderPathTextFilter;
+};
+
+void FPathViewData::PopulateFullFolderTree(const FContentBrowserDataCompiledFilter& CompiledDataFilter)
+{
+	TMap<FName, TSharedPtr<FTreeItem>> OldItemsByInvariantPath = MoveTemp(InvariantPathToItem);
+	RootItems.Reset();
+	VisibleRootItems.Reset();
+	InvariantPathToItem.Reset();
+	VirtualPathToItem.Reset();
+
+	const UContentBrowserSettings* ContentBrowserSettings = GetDefault<UContentBrowserSettings>();
+	bool bDisplayEmpty = ContentBrowserSettings->DisplayEmptyFolders;
+	// check to see if we have an instance config that overrides the default in UContentBrowserSettings
+	if (FContentBrowserInstanceConfig* EditorConfig =
+			ContentBrowserUtils::GetContentBrowserConfig(OwningContentBrowserName))
+	{
+		bDisplayEmpty = EditorConfig->bShowEmptyFolders;
+	}
+
+	UContentBrowserDataSubsystem* ContentBrowserData = IContentBrowserDataModule::Get().GetSubsystem();
+	TOptional<FContentBrowserFolderContentsFilter> FolderFilter;
+	if (!bDisplayEmpty)
+	{
+		FolderFilter = FContentBrowserFolderContentsFilter{};
+		FolderFilter->ItemCategoryFilter = CompiledDataFilter.ItemCategoryFilter;
+	}
+	EContentBrowserIsFolderVisibleFlags FolderFlags = ContentBrowserUtils::GetIsFolderVisibleFlags(bDisplayEmpty);
+
+	TArray<TSharedPtr<FTreeItem>> ItemsCreated;
+	ContentBrowserData->EnumerateItemsMatchingFilter(CompiledDataFilter,
+		[this,
+			bDisplayEmpty,
+			CompiledDataFilter,
+			FolderFilter,
+			FolderFlags,
+			ContentBrowserData,
+			&OldItemsByInvariantPath](FContentBrowserItemData&& InItemData) {
+			UContentBrowserDataSource* Source = InItemData.GetOwnerDataSource();
+			if (Source && !Source->IsFolderVisible(InItemData.GetVirtualPath(), FolderFlags, FolderFilter))
+			{
+				UE_LOG(LogContentBrowser,
+					VeryVerbose,
+					TEXT("Hiding folder %s that fails current pre-text filtering"),
+					*WriteToString<256>(InItemData.GetVirtualPath()));
+				return true; // continue enumerating
+			}
+
+			AddFolderItemInternal(MoveTemp(InItemData), &OldItemsByInvariantPath);
+			return true;
+		});
+	VisibleRootItems = RootItems;
+}
+
+void FPathViewData::PopulateWithFavorites(const FContentBrowserDataCompiledFilter& CompiledDataFilter)
+{
+	// Clear all root items and clear selection
+	TMap<FName, TSharedPtr<FTreeItem>> OldItemsByInvariantPath = MoveTemp(InvariantPathToItem);
+	RootItems.Reset();
+	VisibleRootItems.Reset();
+	InvariantPathToItem.Reset();
+	VirtualPathToItem.Reset();
+
+	const TArray<FString>& FavoritePaths = ContentBrowserUtils::GetFavoriteFolders();
+	UContentBrowserDataSubsystem* ContentBrowserData = IContentBrowserDataModule::Get().GetSubsystem();
+	for (const FString& InvariantPath : FavoritePaths)
+	{
+		FName VirtualPath;
+		IContentBrowserDataModule::Get().GetSubsystem()->ConvertInternalPathToVirtual(InvariantPath, VirtualPath);
+		const FString Path = VirtualPath.ToString();
+
+		ContentBrowserData->EnumerateItemsAtPath(*Path,
+			CompiledDataFilter.ItemTypeFilter,
+			[this, &CompiledDataFilter, &OldItemsByInvariantPath](FContentBrowserItemData&& InItemData) {
+				UContentBrowserDataSource* ItemDataSource = InItemData.GetOwnerDataSource();
+				ItemDataSource->ConvertItemForFilter(InItemData, CompiledDataFilter);
+				if (ItemDataSource->DoesItemPassFilter(InItemData, CompiledDataFilter))
+				{
+					AddFolderItemInternal(MoveTemp(InItemData), &OldItemsByInvariantPath);
+				}
+
+				return true;
+			});
+	}
+}
+
+void FPathViewData::ProcessDataUpdates(TConstArrayView<FContentBrowserItemDataUpdate> InUpdatedItems,
+	const FContentBrowserDataCompiledFilter& CompiledDataFilter)
+{
+	UContentBrowserDataSubsystem* ContentBrowserData = IContentBrowserDataModule::Get().GetSubsystem();
+	const UContentBrowserSettings* ContentBrowserSettings = GetDefault<UContentBrowserSettings>();
+	bool bDisplayEmpty = ContentBrowserSettings->DisplayEmptyFolders;
+	// check to see if we have an instance config that overrides the default in UContentBrowserSettings
+	if (FContentBrowserInstanceConfig* EditorConfig =
+			ContentBrowserUtils::GetContentBrowserConfig(OwningContentBrowserName))
+	{
+		bDisplayEmpty = EditorConfig->bShowEmptyFolders;
+	}
+	TOptional<FContentBrowserFolderContentsFilter> FolderFilter;
+	if (!bDisplayEmpty)
+	{
+		FolderFilter = FContentBrowserFolderContentsFilter{};
+		FolderFilter->ItemCategoryFilter = CompiledDataFilter.ItemCategoryFilter;
+	}
+	EContentBrowserIsFolderVisibleFlags FolderFlags = ContentBrowserUtils::GetIsFolderVisibleFlags(bDisplayEmpty);
+	auto DoesItemPassFilter = [this, FolderFlags, FolderFilter, ContentBrowserData, &CompiledDataFilter](
+								  const FContentBrowserItemData& InItemData) {
+		UContentBrowserDataSource* ItemDataSource = InItemData.GetOwnerDataSource();
+		if (!ItemDataSource->DoesItemPassFilter(InItemData, CompiledDataFilter))
+		{
+			return false;
+		}
+
+		if (!ContentBrowserData->IsFolderVisible(InItemData.GetVirtualPath(), FolderFlags, FolderFilter))
+		{
+			UE_LOG(LogContentBrowser,
+				VeryVerbose,
+				TEXT("Hiding folder %s that fails broad visibility filter"),
+				*WriteToString<256>(InItemData.GetVirtualPath()));
+			return false;
+		}
+
+		return true;
+	};
+
+	TArray<TSharedRef<FTreeItem>> NewItems;
+	TArray<TSharedRef<FTreeItem>> ModifiedParents; // Parents who need their bHasVisibleDescendants updated
+	for (const FContentBrowserItemDataUpdate& ItemDataUpdate : InUpdatedItems)
+	{
+		const FContentBrowserItemData& ItemDataRef = ItemDataUpdate.GetItemData();
+		if (!ItemDataRef.IsFolder())
+		{
+			continue;
+		}
+
+		FContentBrowserItemData ItemData = ItemDataRef;
+		ItemData.GetOwnerDataSource()->ConvertItemForFilter(ItemData, CompiledDataFilter);
+
+		switch (ItemDataUpdate.GetUpdateType())
+		{
+			case EContentBrowserItemUpdateType::Added:
+				if (DoesItemPassFilter(ItemData))
+				{
+					NewItems.Emplace(AddFolderItemInternal(MoveTemp(ItemData), nullptr));
+				}
+				break;
+
+			case EContentBrowserItemUpdateType::Modified:
+				if (DoesItemPassFilter(ItemData))
+				{
+					NewItems.Emplace(AddFolderItemInternal(MoveTemp(ItemData), nullptr));
+				}
+				else
+				{
+					TSharedPtr<FTreeItem> Parent = TryRemoveFolderItemInternal(ItemData);
+					if (Parent.IsValid())
+					{
+						ModifiedParents.Emplace(Parent.ToSharedRef());
+					}
+				}
+				break;
+
+			case EContentBrowserItemUpdateType::Moved:
+			{
+				const FContentBrowserItemData OldMinimalItemData(ItemData.GetOwnerDataSource(),
+					ItemData.GetItemType(),
+					ItemDataUpdate.GetPreviousVirtualPath(),
+					NAME_None,
+					FText(),
+					nullptr);
+				TSharedPtr<FTreeItem> Parent = TryRemoveFolderItemInternal(OldMinimalItemData);
+				if (DoesItemPassFilter(ItemData))
+				{
+					NewItems.Emplace(AddFolderItemInternal(MoveTemp(ItemData), nullptr));
+				}
+				else if (Parent.IsValid())
+				{
+					ModifiedParents.Emplace(Parent.ToSharedRef());
+				}
+			}
+			break;
+
+			case EContentBrowserItemUpdateType::Removed:
+				TryRemoveFolderItemInternal(ItemData);
+				break;
+
+			default:
+				checkf(false, TEXT("Unexpected EContentBrowserItemUpdateType!"));
+				break;
+		}
+	}
+
+	++Version;
+	// Determine visibility for new items and their parents
+	if (!FolderPathTextFilter.GetRawFilterText().IsEmpty())
+	{
+		// Clear visible descendents flag on modified parents because we will reset it
+		for (const TSharedRef<FTreeItem>& Parent : ModifiedParents)
+		{
+			Parent->SetHasVisibleDescendants(false);
+		}
+
+		for (const TSharedRef<FTreeItem>& Item : NewItems)
+		{
+			bool bVisible = PassesTextFilter(Item);
+			Item->SetVisible(bVisible);
+			if (bVisible)
+			{
+				// Propagate to parents
+				for (TSharedPtr<FTreeItem> Parent = Item->GetParent(); Parent.IsValid() && !Parent->IsVisible();
+					 Parent = Parent->GetParent())
+				{
+					Parent->SetHasVisibleDescendants(true);
+				}
+			}
+		}
+
+		// Sort modified parents so if items are related, we visit the items furthest from the root first
+		Algo::Sort(ModifiedParents,
+			[](const TSharedRef<FTreeItem>& A, const TSharedRef<FTreeItem>& B) { return A->IsChildOf(*B); });
+		for (const TSharedRef<FTreeItem>& Parent : ModifiedParents)
+		{
+			// May have already figured this out when dealing with directly modified items
+			if (!Parent->GetHasVisibleDescendants())
+			{
+				bool bVisibleChildren = Algo::AnyOf(Parent->GetChildren(),
+					[](const TSharedPtr<FTreeItem>& Child) { return Child.IsValid() && Child->IsVisible(); });
+				Parent->SetHasVisibleDescendants(bVisibleChildren);
+			}
+		}
+	}
+}
+
+bool FPathViewData::PassesTextFilter(const TSharedPtr<FTreeItem>& InItem)
+{
+	return FolderPathTextFilter.PassesFilter(WriteToString<256>(InItem->GetItem().GetVirtualPath()))
+		// TODO: this will not match a string like LocName1/LocName2 when both parent and child are localized
+		|| FolderPathTextFilter.PassesFilter(FStringView(InItem->GetItem().GetDisplayName().ToString()));
+}
+
+void FPathViewData::ClearItemFilterState()
+{
+	for (const TPair<FName, TSharedPtr<FTreeItem>>& Pair : VirtualPathToItem)
+	{
+		FName VirtualPath = Pair.Key;
+		Pair.Value->SetVisible(true);
+		Pair.Value->SetHasVisibleDescendants(true);
+	}
+	VisibleRootItems = RootItems;
+	++Version;
+}
+
+void FPathViewData::FilterFullFolderTree()
+{
+	for (const TPair<FName, TSharedPtr<FTreeItem>>& Pair : VirtualPathToItem)
+	{
+		FName VirtualPath = Pair.Key;
+		Pair.Value->SetVisible(PassesTextFilter(Pair.Value));
+		Pair.Value->SetHasVisibleDescendants(false);
+	}
+
+	// Propagate visibility down to parents
+	for (const TPair<FName, TSharedPtr<FTreeItem>>& Pair : VirtualPathToItem)
+	{
+		if (Pair.Value->IsVisible())
+		{
+			for (TSharedPtr<FTreeItem> Parent = Pair.Value->GetParent(); Parent.IsValid() && !Parent->IsVisible();
+				 Parent = Parent->GetParent())
+			{
+				Parent->SetHasVisibleDescendants(true);
+			}
+		}
+	}
+	VisibleRootItems.Reset();
+	Algo::CopyIf(RootItems, VisibleRootItems, UE_PROJECTION_MEMBER(FTreeItem, IsVisible));
+	++Version;
+}
+
+void FPathViewData::SortRootItems()
+{
+	UE::PathView::DefaultSort(RootItems);
+	UE::PathView::DefaultSort(VisibleRootItems);
+}
+
+TSharedPtr<FTreeItem> FPathViewData::FindTreeItem(FName InVirtualPath, bool bVisibleOnly)
+{
+	if (TSharedPtr<FTreeItem> Found = VirtualPathToItem.FindRef(InVirtualPath))
+	{
+		if (bVisibleOnly && !Found->IsVisible())
+		{
+			return {};
+		}
+		return Found;
+	}
+	return {};
+}
+
+TSharedPtr<FTreeItem> FPathViewData::FindBestItemForPath(FStringView InVirtualPath)
+{
+	TSharedPtr<FTreeItem> Found;
+	FPathViews::IterateAncestors(InVirtualPath, [this, &Found](FStringView Ancestor) {
+		FName ItemName{ Ancestor };
+		if (TSharedPtr<FTreeItem>* Item = VirtualPathToItem.Find(ItemName))
+		{
+			Found = *Item;
+			return false; // Found the leafmost item matching this path
+		}
+		return true; // continue
+	});
+	return Found;
+}
+
+TSharedRef<FTreeItem> FPathViewData::AddFolderItem(FContentBrowserItemData&& InItemData)
+{
+	TSharedRef<FTreeItem> NewOrUpdatedItem = AddFolderItemInternal(MoveTemp(InItemData), nullptr);
+	++Version;
+	return NewOrUpdatedItem;
+}
+
+TSharedRef<FTreeItem> FPathViewData::AddFolderItemInternal(FContentBrowserItemData&& InItemData,
+	TMap<FName, TSharedPtr<FTreeItem>>* OldItemsByInvariantPath)
+{
+	FName ItemVirtualPath = InItemData.GetVirtualPath();
+	TSharedPtr<FTreeItem> LeafItem = VirtualPathToItem.FindRef(ItemVirtualPath);
+	if (LeafItem.IsValid())
+	{
+		// Item already existed - duplicate item returned by multiple data sources, merge data and move on.
+		// We will have already created all the parent items.
+		LeafItem->AppendItemData(InItemData);
+		return LeafItem.ToSharedRef();
+	}
+
+	UContentBrowserDataSource* OriginalDataSource = InItemData.GetOwnerDataSource();
+
+	FName ItemInvariantPath = InItemData.GetInvariantPath();
+	TStringBuilder<FName::StringBufferSize> PathBuffer(InPlace, ItemVirtualPath);
+	LeafItem = UE::PathView::CreateOrReuseNode(MoveTemp(InItemData), OldItemsByInvariantPath);
+	// InItemData is now no longer valid!!
+
+	VirtualPathToItem.Add(ItemVirtualPath, LeafItem);
+	InvariantPathToItem.Add(ItemInvariantPath, LeafItem);
+
+	if (bFlat)
+	{
+		RootItems.Add(LeafItem);
+		return LeafItem.ToSharedRef();
+	}
+
+	TSharedRef<FTreeItem> PreviousItem = LeafItem.ToSharedRef();
+
+	// Work backwards from the leaf path of the requested item until we encounter an item that already existed
+	FPathViews::IterateAncestors(PathBuffer.ToView(),
+		[this, &PathBuffer, &PreviousItem, OriginalDataSource, OldItemsByInvariantPath](FStringView PathView) {
+			if (PathView.Len() == PathBuffer.Len())
+			{
+				// This is the item returned by the data source, we already added it
+				return true;
+			}
+			if (PathView == TEXTVIEW("/"))
+			{
+				// PreviousItem must have been new, add it to the set of root items
+				RootItems.Add(PreviousItem);
+				return false;
+			}
+			FName ParentVirtualPath{ PathView };
+			TSharedPtr<FTreeItem> ParentItem = VirtualPathToItem.FindRef(ParentVirtualPath);
+			bool bContinue = false;
+			if (!ParentItem.IsValid())
+			{
+				// TODO: If another data source provides this path in future, can that data source become the 'primary'
+				// ?
+				ParentItem = UE::PathView::CreateOrReuseNode(FContentBrowserItemData(OriginalDataSource,
+																 EContentBrowserItemFlags::Type_Folder,
+																 ParentVirtualPath,
+																 FName(FPathViews::GetPathLeaf(PathView)),
+																 FText(),
+																 nullptr),
+					OldItemsByInvariantPath);
+				VirtualPathToItem.Add(ParentVirtualPath, ParentItem);
+				InvariantPathToItem.Add(
+					ParentItem->GetItem().GetInvariantPath()); // TODO: Do fully virtual paths have an invariant path?
+				bContinue = true;
+			}
+			ParentItem->AddChild(PreviousItem);
+			PreviousItem = ParentItem.ToSharedRef();
+			return bContinue; // If we made a node here, keep checking if we need to make more parent nodes
+		});
+	return LeafItem.ToSharedRef();
+}
+
+TSharedPtr<FTreeItem> FPathViewData::TryRemoveFolderItemInternal(const FContentBrowserItemData& InItem)
+{
+	if (!InItem.IsFolder())
+	{
+		// Not a folder
+		return {};
+	}
+
+	// Find the folder in the tree
+	if (TSharedPtr<FTreeItem> ItemToRemove = VirtualPathToItem.FindRef(InItem.GetVirtualPath()))
+	{
+		// Only fully remove this item if every sub-item is removed (items become invalid when empty)
+		ItemToRemove->RemoveItemData(InItem);
+		if (ItemToRemove->GetItem().IsValid())
+		{
+			return {};
+		}
+
+		// Found the folder to remove. Remove it.
+		TSharedPtr<FTreeItem> ItemParent = ItemToRemove->GetParent();
+		if (ItemParent.IsValid())
+		{
+			// Remove the folder from its parent's list
+			ItemParent->RemoveChild(ItemToRemove.ToSharedRef());
+		}
+		else
+		{
+			// This is a root item. Remove the folder from the root items list.
+			RootItems.Remove(ItemToRemove);
+			VisibleRootItems.Remove(ItemToRemove);
+		}
+
+		VirtualPathToItem.Remove(InItem.GetVirtualPath());
+		InvariantPathToItem.Remove(InItem.GetInvariantPath());
+		return ItemParent;
+	}
+
+	// Did not find the folder to remove
+	return {};
+}
+
+void FPathViewData::RemoveFolderItem(const TSharedRef<FTreeItem>& TreeItem)
+{
+	if (TSharedPtr<FTreeItem> Parent = TreeItem->GetParent())
+	{
+		// Remove this item from it's parent's list
+		Parent->RemoveChild(TreeItem);
+	}
+	else
+	{
+		// This was a root node, remove from the root list
+		RootItems.Remove(TreeItem);
+		VisibleRootItems.Remove(TreeItem);
+	}
+
+	VirtualPathToItem.Remove(TreeItem->GetItem().GetVirtualPath());
+	InvariantPathToItem.Remove(TreeItem->GetItem().GetInvariantPath());
+
+	Version++;
+}
+
 SPathView::FScopedSelectionChangedEvent::FScopedSelectionChangedEvent(const TSharedRef<SPathView>& InPathView, const bool InShouldEmitEvent)
 	: PathView(InPathView)
 	, bShouldEmitEvent(InShouldEmitEvent)
@@ -97,8 +825,9 @@ SPathView::FScopedSelectionChangedEvent::~FScopedSelectionChangedEvent()
 		const bool bHasSelectionChanges = InitialSelectionSet.Num() != FinalSelectionSet.Num() || InitialSelectionSet.Difference(FinalSelectionSet).Num() > 0;
 		if (bHasSelectionChanges)
 		{
-			const TArray<TSharedPtr<FTreeItem>> SelectedItems = PathView->TreeViewPtr->GetSelectedItems();
-			PathView->TreeSelectionChanged(SelectedItems.Num() > 0 ? SelectedItems[0] : nullptr, ESelectInfo::Direct);
+			const TArray<TSharedPtr<FTreeItem>> NewSelectedItems = PathView->TreeViewPtr->GetSelectedItems();
+			PathView->TreeSelectionChanged(NewSelectedItems.Num() > 0 ? NewSelectedItems[0] : nullptr,
+				ESelectInfo::Direct);
 		}
 	}
 }
@@ -106,16 +835,9 @@ SPathView::FScopedSelectionChangedEvent::~FScopedSelectionChangedEvent()
 TSet<FName> SPathView::FScopedSelectionChangedEvent::GetSelectionSet() const
 {
 	TSet<FName> SelectionSet;
-
-	const TArray<TSharedPtr<FTreeItem>> SelectedItems = PathView->TreeViewPtr->GetSelectedItems();
-	for (const TSharedPtr<FTreeItem>& Item : SelectedItems)
-	{
-		if (ensure(Item.IsValid()))
-		{
-			SelectionSet.Add(Item->GetItem().GetVirtualPath());
-		}
-	}
-
+	Algo::Transform(PathView->TreeViewPtr->GetSelectedItems(), SelectionSet, [](const TSharedPtr<FTreeItem>& Item) {
+		return Item->GetItem().GetVirtualPath();
+	});
 	return SelectionSet;
 }
 
@@ -131,7 +853,7 @@ SPathView::~SPathView()
 		}
 	}
 
-	SearchBoxFolderFilter->OnChanged().RemoveAll( this );
+	TreeData->GetFolderPathTextFilter().OnChanged().RemoveAll(this);
 }
 
 void SPathView::Construct( const FArguments& InArgs )
@@ -152,7 +874,8 @@ void SPathView::Construct( const FArguments& InArgs )
 		RegisterActiveTimer( 0.f, FWidgetActiveTimerDelegate::CreateSP( this, &SPathView::SetFocusPostConstruct ) );
 	}
 
-	SortOverride = FSortTreeItemChildrenDelegate::CreateSP(this, &SPathView::DefaultSort);
+	TreeData = MakeShared<FPathViewData>(OwningContentBrowserName, bFlat);
+	TreeData->GetFolderPathTextFilter().OnChanged().AddSP(this, &SPathView::FilterUpdated);
 
 	UContentBrowserDataSubsystem* ContentBrowserData = IContentBrowserDataModule::Get().GetSubsystem();
 	ContentBrowserData->OnItemDataUpdated().AddSP(this, &SPathView::HandleItemDataUpdated);
@@ -167,10 +890,6 @@ void SPathView::Construct( const FArguments& InArgs )
 	FContentBrowserModule& ContentBrowserModule = FModuleManager::GetModuleChecked<FContentBrowserModule>(TEXT("ContentBrowser"));
 	ContentBrowserModule.GetOnContentBrowserSettingChanged().AddSP(this, &SPathView::HandleSettingChanged);
 
-	//Setup the SearchBox filter
-	SearchBoxFolderFilter = MakeShareable( new FolderTextFilter( FolderTextFilter::FItemToStringArray::CreateSP( this, &SPathView::PopulateFolderSearchStrings ) ) );
-	SearchBoxFolderFilter->OnChanged().AddSP( this, &SPathView::FilterUpdated );
-
 	// Setup plugin filters
 	PluginPathFilters = InArgs._PluginPathFilters;
 	if (PluginPathFilters.IsValid())
@@ -188,22 +907,19 @@ void SPathView::Construct( const FArguments& InArgs )
 		}
 	}
 
-	if (!TreeViewPtr.IsValid())
-	{
-		SAssignNew(TreeViewPtr, STreeView< TSharedPtr<FTreeItem> >)
-			.TreeItemsSource(&TreeRootItems)
-			.OnGenerateRow(this, &SPathView::GenerateTreeRow)
-			.OnItemScrolledIntoView(this, &SPathView::TreeItemScrolledIntoView)
-			.ItemHeight(18)
-			.SelectionMode(InArgs._SelectionMode)
-			.OnSelectionChanged(this, &SPathView::TreeSelectionChanged)
-			.OnExpansionChanged(this, &SPathView::TreeExpansionChanged)
-			.OnGetChildren(this, &SPathView::GetChildrenForTree)
-			.OnSetExpansionRecursive(this, &SPathView::SetTreeItemExpansionRecursive)
-			.OnContextMenuOpening(this, &SPathView::MakePathViewContextMenu)
-			.ClearSelectionOnClick(false)
-			.HighlightParentNodesForSelection(true);
-	}
+	STreeView<TSharedPtr<FTreeItem>>::FArguments TreeViewArgs;
+	ConfigureTreeView(TreeViewArgs);
+	TreeViewPtr = SArgumentNew(TreeViewArgs, STreeView<TSharedPtr<FTreeItem>>)
+					  .TreeItemsSource(TreeData->GetVisibleRootItems())
+					  .OnGetChildren(this, &SPathView::GetChildrenForTree)
+					  .OnGenerateRow(this, &SPathView::GenerateTreeRow)
+					  .OnItemScrolledIntoView(this, &SPathView::TreeItemScrolledIntoView)
+					  .ItemHeight(18)
+					  .SelectionMode(InArgs._SelectionMode)
+					  .AllowInvisibleItemSelection(true)
+					  .OnSelectionChanged(this, &SPathView::TreeSelectionChanged)
+					  .OnContextMenuOpening(this, &SPathView::MakePathViewContextMenu)
+					  .ClearSelectionOnClick(false);
 
 	SearchPtr = InArgs._ExternalSearch;
 	if (!SearchPtr)
@@ -327,18 +1043,49 @@ void SPathView::Construct( const FArguments& InArgs )
 
 	for (const FName PathToExpand : GetDefaultPathsToExpand())
 	{
-		if (TSharedPtr<FTreeItem> FoundItem = FindTreeItem(PathToExpand))
+		if (TSharedPtr<FTreeItem> FoundItem = TreeData->FindTreeItem(PathToExpand))
 		{
 			RecursiveExpandParents(FoundItem);
 			TreeViewPtr->SetItemExpansion(FoundItem, true);
 		}
 	}
+
+	if (!InArgs._DefaultPath.IsEmpty() && InternalPathPassesBlockLists(InArgs._DefaultPath))
+	{
+		const FName VirtualPath =
+			IContentBrowserDataModule::Get().GetSubsystem()->ConvertInternalPathToVirtual(*InArgs._DefaultPath);
+		if (!TreeData->FindTreeItem(VirtualPath))
+		{
+			const FString DefaultPathLeafName = FPaths::GetPathLeaf(VirtualPath.ToString());
+			TreeData->AddFolderItem(FContentBrowserItemData(nullptr,
+				EContentBrowserItemFlags::Type_Folder,
+				VirtualPath,
+				*DefaultPathLeafName,
+				FText(),
+				nullptr));
+		}
+
+		SetSelectedPaths({ VirtualPath.ToString() });
+	}
+}
+
+void SPathView::ConfigureTreeView(STreeView<TSharedPtr<FTreeItem>>::FArguments& InArgs)
+{
+	InArgs.OnExpansionChanged(this, &SPathView::TreeExpansionChanged)
+		.OnSetExpansionRecursive(this, &SPathView::SetTreeItemExpansionRecursive)
+		.HighlightParentNodesForSelection(true);
 }
 
 void SPathView::Tick(const FGeometry& AllottedGeometry, const double InCurrentTime, const float InDeltaTime)
 {
 	Super::Tick(AllottedGeometry, InCurrentTime, InDeltaTime);
 
+	if (LastTreeDataVersion != TreeData->GetVersion())
+	{
+		LastTreeDataVersion = TreeData->GetVersion();
+		TreeViewPtr->RequestTreeRefresh();
+	}
+
 	const bool bNewShowRedirectors = bShowRedirectors.Get(false);
 	if (bNewShowRedirectors != bLastShowRedirectors)
 	{
@@ -447,34 +1194,12 @@ void SPathView::SetPluginPathFilterActive(const TSharedRef<FContentBrowserPlugin
 
 FPathViewConfig* SPathView::GetPathViewConfig() const
 {
-	if (OwningContentBrowserName.IsNone())
-	{
-		return nullptr;
-	}
-
-	FContentBrowserInstanceConfig* Config = UContentBrowserConfig::Get()->Instances.Find(OwningContentBrowserName);
-	if (Config == nullptr)
-	{
-		return nullptr;
-	}
-	
-	return &Config->PathView;
+	return ContentBrowserUtils::GetPathViewConfig(OwningContentBrowserName);
 }
 
 FContentBrowserInstanceConfig* SPathView::GetContentBrowserConfig() const
 {
-	if (OwningContentBrowserName.IsNone())
-	{
-		return nullptr;
-	}
-
-	FContentBrowserInstanceConfig* Config = UContentBrowserConfig::Get()->Instances.Find(OwningContentBrowserName);
-	if (Config == nullptr)
-	{
-		return nullptr;
-	}
-
-	return Config;
+	return ContentBrowserUtils::GetContentBrowserConfig(OwningContentBrowserName);
 }
 
 void SPathView::SetSelectedPaths(const TArray<FName>& Paths)
@@ -491,22 +1216,11 @@ void SPathView::SetSelectedPaths(const TArray<FString>& Paths)
 		return;
 	}
 
-	// Clear the search box if it potentially hides a path we want to select
-	for (const FString& Path : Paths)
-	{
-		if (PathIsFilteredFromViewBySearch(Path))
-		{
-			SearchPtr->ClearSearch();
-			break;
-		}
-	}
-
 	// Prevent the selection changed delegate since the invoking code requested it
 	FScopedPreventTreeItemChangedDelegate DelegatePrevention( SharedThis(this) );
 
 	// If the selection was changed before all pending initial paths were found, stop attempting to select them
 	PendingInitialPaths.Empty();
-	bPendingInitialPathsNeedsSelectionClear = false;
 
 	// Clear the selection to start, then add the selected paths as they are found
 	LastSelectedPaths.Empty();
@@ -514,72 +1228,24 @@ void SPathView::SetSelectedPaths(const TArray<FString>& Paths)
 
 	for (const FString& Path : Paths)
 	{
-		TArray<FName> PathItemList;
+		TSharedPtr<FTreeItem> BestItem = TreeData->FindBestItemForPath(Path);
+		if (BestItem.IsValid())
 		{
-			TArray<FString> PathItemListStr;
-			Path.ParseIntoArray(PathItemListStr, TEXT("/"), /*InCullEmpty=*/true);
-
-			PathItemList.Reserve(PathItemListStr.Num());
-			for (const FString& PathItemName : PathItemListStr)
+			if (!BestItem->IsVisible())
 			{
-				PathItemList.Add(*PathItemName);
+				// Clear the search box if it potentially hides a path we want to select
+				SearchPtr->ClearSearch();
 			}
-		}
 
-		if (PathItemList.Num())
-		{
-			// There is at least one element in the path
-			TArray<TSharedPtr<FTreeItem>> TreeItems;
-
-			// Find the first item in the root items list
-			for (const TSharedPtr<FTreeItem>& TreeItem : TreeRootItems)
+			for (TSharedPtr<FTreeItem> Parent = BestItem->GetParent(); Parent.IsValid(); Parent = Parent->GetParent())
 			{
-				if (TreeItem->GetItem().GetItemName() == PathItemList[0])
-				{
-					// Found the first item in the path
-					TreeItems.Add(TreeItem);
-					break;
-				}
+				TreeViewPtr->SetItemExpansion(Parent, true);
 			}
 
-			// If found in the root items list, try to find the childmost item matching the path
-			if (TreeItems.Num() > 0)
-			{
-				for ( int32 PathItemIdx = 1; PathItemIdx < PathItemList.Num(); ++PathItemIdx )
-				{
-					const FName PathItemName = PathItemList[PathItemIdx];
-					const TSharedPtr<FTreeItem> ChildItem = TreeItems.Last()->GetChild(PathItemName);
-					if (ChildItem.IsValid())
-					{
-						// Update tree items list
-						TreeItems.Add(ChildItem);
-					}
-					else
-					{
-						// Could not find the child item
-						break;
-					}
-				}
-
-				// Expand all the tree folders up to but not including the last one.
-				for (int32 ItemIdx = 0; ItemIdx < TreeItems.Num() - 1; ++ItemIdx)
-				{
-					TreeViewPtr->SetItemExpansion(TreeItems[ItemIdx], true);
-				}
-
-				// Set the selection to the closest found folder and scroll it into view
-				LastSelectedPaths.Add(TreeItems.Last()->GetItem().GetInvariantPath());
-				TreeViewPtr->SetItemSelection(TreeItems.Last(), true);
-				TreeViewPtr->RequestScrollIntoView(TreeItems.Last());
-			}
-			else
-			{
-				// Could not even find the root path... skip
-			}
-		}
-		else
-		{
-			// No path items... skip
+			// Set the selection to the closest found folder and scroll it into view
+			LastSelectedPaths.Add(BestItem->GetItem().GetInvariantPath());
+			TreeViewPtr->SetItemSelection(BestItem, true);
+			TreeViewPtr->RequestScrollIntoView(BestItem);
 		}
 	}
 
@@ -598,7 +1264,6 @@ void SPathView::ClearSelection()
 
 	// If the selection was changed before all pending initial paths were found, stop attempting to select them
 	PendingInitialPaths.Empty();
-	bPendingInitialPathsNeedsSelectionClear = false;
 
 	// Clear the selection to start, then add the selected paths as they are found
 	TreeViewPtr->ClearSelection();
@@ -645,209 +1310,6 @@ TArray<FContentBrowserItem> SPathView::GetSelectedFolderItems() const
 	return SelectedFolders;
 }
 
-TSharedPtr<FTreeItem> SPathView::AddFolderItem(FContentBrowserItemData&& InItem, const bool bUserNamed, TArray<TSharedPtr<FTreeItem>>* OutItemsCreated)
-{
-	if (!ensure(TreeViewPtr.IsValid()))
-	{
-		// No tree view for some reason
-		return nullptr;
-	}
-
-	if (!InItem.IsFolder())
-	{
-		// Not a folder
-		return nullptr;
-	}
-
-	// Clear selection if user has not changed it yet or this is the first pending initial path being selected
-	auto SelectingPendingInitialPath = [this]()
-	{
-		if (bPendingInitialPathsNeedsSelectionClear)
-		{
-			bPendingInitialPathsNeedsSelectionClear = false;
-
-			LastSelectedPaths.Empty();
-			TreeViewPtr->ClearSelection();
-		}
-	};
-
-	// The path view will add a node for each level of the path tree
-	TArray<FString> PathItemList;
-	InItem.GetVirtualPath().ToString().ParseIntoArray(PathItemList, TEXT("/"), /*InCullEmpty=*/true);
-
-	// Start at the root and work down until all required children have been added
-	TSharedPtr<FTreeItem> ParentTreeItem;
-	TArray<TSharedPtr<FTreeItem>>* CurrentTreeItems = &TreeRootItems;
-
-	TStringBuilder<512> CurrentPathStr;
-	CurrentPathStr.Append(TEXT("/"));
-	for (int32 PathItemIndex = 0; PathItemIndex < PathItemList.Num(); ++PathItemIndex)
-	{
-		const bool bIsLeafmostItem = PathItemIndex == PathItemList.Num() - 1;
-
-		const FString& FolderNameStr = PathItemList[PathItemIndex];
-		const FName FolderName = *FolderNameStr;
-		FPathViews::Append(CurrentPathStr, FolderNameStr);
-
-		// Try and find an existing tree item
-		TSharedPtr<FTreeItem> CurrentTreeItem;
-		for (const TSharedPtr<FTreeItem>& PotentialTreeItem : *CurrentTreeItems)
-		{
-			if (PotentialTreeItem->GetItem().GetItemName() == FolderName)
-			{
-				CurrentTreeItem = PotentialTreeItem;
-				break;
-			}
-		}
-
-		// Handle creating the leaf-most item that was given to us to create
-		if (bIsLeafmostItem)
-		{
-			if (CurrentTreeItem)
-			{
-				// Found a match - merge the new item data
-				CurrentTreeItem->AppendItemData(InItem);
-			}
-			else
-			{
-				// No match - create a new item
-				CurrentTreeItem = MakeShared<FTreeItem>(MoveTemp(InItem));
-				CurrentTreeItem->Parent = ParentTreeItem;
-				CurrentTreeItem->SetSortOverride(SortOverride);
-				CurrentTreeItems->Add(CurrentTreeItem);
-				if (OutItemsCreated)
-				{
-					OutItemsCreated->Add(CurrentTreeItem);
-				}
-
-				TreeItemLookup.Add(CurrentTreeItem->GetItem().GetVirtualPath(), CurrentTreeItem);
-
-				if (ParentTreeItem)
-				{
-					check(&ParentTreeItem->Children == CurrentTreeItems);
-					ParentTreeItem->RequestSortChildren();
-				}
-				else
-				{
-					SortRootItems();
-				}
-
-				// If we have pending initial paths, and this path added the path, we should select it now
-				if (PendingInitialPaths.Num() > 0 && PendingInitialPaths.Contains(CurrentTreeItem->GetItem().GetVirtualPath()))
-				{
-					SelectingPendingInitialPath();
-					RecursiveExpandParents(CurrentTreeItem);
-					TreeViewPtr->SetItemSelection(CurrentTreeItem, true);
-					TreeViewPtr->RequestScrollIntoView(CurrentTreeItem);
-				}
-			}
-
-			// If we want to name this item, select it, scroll it into view, expand the parent
-			if (bUserNamed)
-			{
-				RecursiveExpandParents(CurrentTreeItem);
-				TreeViewPtr->SetSelection(CurrentTreeItem);
-				CurrentTreeItem->SetNamingFolder(true);
-				TreeViewPtr->RequestScrollIntoView(CurrentTreeItem);
-			}
-
-			TreeViewPtr->RequestTreeRefresh();
-			return CurrentTreeItem;
-		}
-
-		// If we're missing an item on the way down to the leaf-most item then we'll add a placeholder
-		// This shouldn't usually happen as Populate will create paths in the correct order, but 
-		// the path picker may force add a path that hasn't been discovered (or doesn't exist) yet
-		if (!CurrentTreeItem)
-		{
-			CurrentTreeItem = MakeShared<FTreeItem>(FContentBrowserItemData(InItem.GetOwnerDataSource(), EContentBrowserItemFlags::Type_Folder, *CurrentPathStr, FolderName, FText(), nullptr));
-			CurrentTreeItem->Parent = ParentTreeItem;
-			CurrentTreeItem->SetSortOverride(SortOverride);
-			CurrentTreeItems->Add(CurrentTreeItem);
-			if (OutItemsCreated)
-			{
-				OutItemsCreated->Add(CurrentTreeItem);
-			}
-
-			TreeItemLookup.Add(CurrentTreeItem->GetItem().GetVirtualPath(), CurrentTreeItem);
-
-			if (ParentTreeItem)
-			{
-				check(&ParentTreeItem->Children == CurrentTreeItems);
-				ParentTreeItem->RequestSortChildren();
-			}
-			else
-			{
-				SortRootItems();
-			}
-
-			// If we have pending initial paths, and this path added the path, we should select it now
-			if (PendingInitialPaths.Num() > 0 && PendingInitialPaths.Contains(CurrentTreeItem->GetItem().GetVirtualPath()))
-			{
-				SelectingPendingInitialPath();
-				RecursiveExpandParents(CurrentTreeItem);
-				TreeViewPtr->SetItemSelection(CurrentTreeItem, true);
-				TreeViewPtr->RequestScrollIntoView(CurrentTreeItem);
-			}
-		}
-
-		// Set-up the data for the next level
-		ParentTreeItem = CurrentTreeItem;
-		CurrentTreeItems = &ParentTreeItem->Children;
-	}
-
-	return nullptr;
-}
-
-bool SPathView::RemoveFolderItem(const FContentBrowserItemData& InItem)
-{
-	if (!ensure(TreeViewPtr.IsValid()))
-	{
-		// No tree view for some reason
-		return false;
-	}
-
-	if (!InItem.IsFolder())
-	{
-		// Not a folder
-		return false;
-	}
-
-	// Find the folder in the tree
-	const FName VirtualPath = InItem.GetVirtualPath();
-	if (TSharedPtr<FTreeItem> ItemToRemove = FindTreeItem(VirtualPath))
-	{
-		// Only fully remove this item if every sub-item is removed (items become invalid when empty)
-		ItemToRemove->RemoveItemData(InItem);
-		if (ItemToRemove->GetItem().IsValid())
-		{
-			return true;
-		}
-
-		// Found the folder to remove. Remove it.
-		if (TSharedPtr<FTreeItem> ItemParent = ItemToRemove->Parent.Pin())
-		{
-			// Remove the folder from its parent's list
-			ItemParent->Children.Remove(ItemToRemove);
-		}
-		else
-		{
-			// This is a root item. Remove the folder from the root items list.
-			TreeRootItems.Remove(ItemToRemove);
-		}
-
-		TreeItemLookup.Remove(VirtualPath);
-
-		// Refresh the tree
-		TreeViewPtr->RequestTreeRefresh();
-
-		return true;
-	}
-	
-	// Did not find the folder to remove
-	return false;
-}
-
 void SPathView::RenameFolderItem(const FContentBrowserItem& InItem)
 {
 	if (!ensure(TreeViewPtr.IsValid()))
@@ -863,8 +1325,12 @@ void SPathView::RenameFolderItem(const FContentBrowserItem& InItem)
 	}
 
 	// Find the folder in the tree
-	if (TSharedPtr<FTreeItem> ItemToRename = FindTreeItem(InItem.GetVirtualPath()))
+	if (TSharedPtr<FTreeItem> ItemToRename = TreeData->FindTreeItem(InItem.GetVirtualPath()))
 	{
+		if (!ItemToRename->IsVisible())
+		{
+			SearchPtr->ClearSearch();
+		}
 		ItemToRename->SetNamingFolder(true);
 
 		TreeViewPtr->SetSelection(ItemToRename);
@@ -918,6 +1384,11 @@ FContentBrowserDataCompiledFilter SPathView::CreateCompiledFolderFilter() const
 		}
 	}
 
+	UE_LOG(LogContentBrowser,
+		VeryVerbose,
+		TEXT("Compiled folder permission list: %s"),
+		*CombinedFolderPermissionList->ToString());
+
 	ContentBrowserUtils::AppendAssetFilterToContentBrowserFilter(FARFilter(), nullptr, CombinedFolderPermissionList, DataFilter);
 
 	FContentBrowserDataCompiledFilter CompiledDataFilter;
@@ -1054,35 +1525,23 @@ void SPathView::SyncToItems(TArrayView<const FContentBrowserItem> ItemsToSync, c
 
 void SPathView::SyncToVirtualPaths(TArrayView<const FName> VirtualPathsToSync, const bool bAllowImplicitSync)
 {
-	// Clear the search box if it potentially hides a path we want to select
+	TSet<TSharedRef<FTreeItem>> SyncTreeItems;
 	for (const FName& VirtualPathToSync : VirtualPathsToSync)
 	{
-		if (PathIsFilteredFromViewBySearch(VirtualPathToSync.ToString()))
+		TSharedPtr<FTreeItem> Item = TreeData->FindTreeItem(VirtualPathToSync);
+		if (Item.IsValid())
 		{
-			SearchPtr->ClearSearch();
-			break;
+			SyncTreeItems.Add(Item.ToSharedRef());
 		}
 	}
-
-	TArray<TSharedPtr<FTreeItem>> SyncTreeItems;
-	{
-		TSet<FName> UniqueVirtualPathsToSync;
-		for (const FName& VirtualPathToSync : VirtualPathsToSync)
-		{
-			if (!UniqueVirtualPathsToSync.Contains(VirtualPathToSync))
-			{
-				UniqueVirtualPathsToSync.Add(VirtualPathToSync);
-
-				TSharedPtr<FTreeItem> Item = FindTreeItem(VirtualPathToSync);
-				if (Item.IsValid())
-				{
-					SyncTreeItems.Add(Item);
-				}
-			}
-		}
+
+	if (Algo::AnyOf(SyncTreeItems, [](const TSharedRef<FTreeItem>& Item) { return !Item->IsVisible(); }))
+	{
+		// Clear the search box if it potentially hides a path we want to select
+		SearchPtr->ClearSearch();
 	}
 
-	if ( SyncTreeItems.Num() > 0 )
+	if (SyncTreeItems.Num() > 0)
 	{
 		// Batch the selection changed event
 		FScopedSelectionChangedEvent ScopedSelectionChangedEvent(SharedThis(this));
@@ -1091,40 +1550,34 @@ void SPathView::SyncToVirtualPaths(TArrayView<const FName> VirtualPathsToSync, c
 		{
 			// Prune the current selection so that we don't unnecessarily change the path which might disorientate the user.
 			// If a parent tree item is currently selected we don't need to clear it and select the child
-			TArray<TSharedPtr<FTreeItem>> SelectedTreeItems = TreeViewPtr->GetSelectedItems();
-
-			for (int32 Index = 0; Index < SelectedTreeItems.Num(); ++Index)
+			TSet<TSharedPtr<FTreeItem>> SelectedTreeItems{ TreeViewPtr->GetSelectedItems() };
+			TSet<TSharedRef<FTreeItem>> FinalItems;
+			for (const TSharedRef<FTreeItem>& ItemToSelect : SyncTreeItems)
 			{
-				// For each item already selected in the tree
-				const TSharedPtr<FTreeItem>& AlreadySelectedTreeItem = SelectedTreeItems[Index];
-				if (!AlreadySelectedTreeItem.IsValid())
+				// If the target item or any of its parents are already selected, maintain that object in the final
+				// selection
+				TSharedPtr<FTreeItem> It = ItemToSelect;
+				while (It.IsValid() && !SelectedTreeItems.Contains(It))
 				{
-					continue;
+					It = It->GetParent();
 				}
 
-				// Check to see if any of the items to sync are already synced
-				for (int32 ToSyncIndex = SyncTreeItems.Num()-1; ToSyncIndex >= 0; --ToSyncIndex)
+				if (It.IsValid())
 				{
-					const TSharedPtr<FTreeItem>& ToSyncItem = SyncTreeItems[ToSyncIndex];
-					if (ToSyncItem == AlreadySelectedTreeItem || ToSyncItem->IsChildOf(*AlreadySelectedTreeItem.Get()))
-					{
-						// A parent is already selected
-						SyncTreeItems.Pop();
-					}
-					else if (ToSyncIndex == 0)
-					{
-						// AlreadySelectedTreeItem is not required for SyncTreeItems, so deselect it
-						TreeViewPtr->SetItemSelection(AlreadySelectedTreeItem, false);
-					}
+					FinalItems.Add(It.ToSharedRef());
+				}
+				else
+				{
+					// Otherwise select the specific folder we were asked for
+					FinalItems.Add(ItemToSelect);
 				}
 			}
-		}
-		else
-		{
-			// Explicit sync so just clear the selection
-			TreeViewPtr->ClearSelection();
+			SyncTreeItems = FinalItems;
 		}
 
+		// SyncTreeItems now shows exactly what we want to be selected and no more
+		TreeViewPtr->ClearSelection();
+
 		// SyncTreeItems should now only contain items which aren't already shown explicitly or implicitly (as a child)
 		for (const TSharedPtr<FTreeItem>& Item : SyncTreeItems)
 		{
@@ -1133,11 +1586,11 @@ void SPathView::SyncToVirtualPaths(TArrayView<const FName> VirtualPathsToSync, c
 		}
 	}
 
-	// > 0 as some may have been popped off in the code above
+	// > 0 as some may have been removed in the code above
 	if (SyncTreeItems.Num() > 0)
 	{
 		// Scroll the first item into view if applicable
-		TreeViewPtr->RequestScrollIntoView(SyncTreeItems[0]);
+		TreeViewPtr->RequestScrollIntoView(*SyncTreeItems.CreateConstIterator());
 	}
 }
 
@@ -1149,21 +1602,9 @@ void SPathView::SyncToLegacy(TArrayView<const FAssetData> AssetDataList, TArrayV
 	SyncToVirtualPaths(VirtualPathsToSync, bAllowImplicitSync);
 }
 
-void SPathView::ClearTreeItems()
-{
-	TreeRootItems.Empty();
-	TreeViewPtr->ClearSelection();
-	TreeItemLookup.Empty();
-}
-
-TSharedPtr<FTreeItem> SPathView::FindTreeItem(FName InPath) const
+bool SPathView::DoesItemExist(FName InVirtualPath) const
 {
-	if (const TWeakPtr<FTreeItem>* FoundWeak = TreeItemLookup.Find(InPath))
-	{
-		return FoundWeak->Pin();
-	}
-
-	return TSharedPtr<FTreeItem>();
+	return TreeData->FindTreeItem(InVirtualPath).IsValid();
 }
 
 void SPathView::ApplyHistoryData( const FHistoryData& History )
@@ -1246,48 +1687,35 @@ void SPathView::LoadSettings(const FString& IniFilename, const FString& IniSecti
 		// Batch the selection changed event
 		FScopedSelectionChangedEvent ScopedSelectionChangedEvent(SharedThis(this));
 
-		bPendingInitialPathsNeedsSelectionClear = false;
-
 		UContentBrowserDataSubsystem* ContentBrowserData = IContentBrowserDataModule::Get().GetSubsystem();
 		if (ContentBrowserData->IsDiscoveringItems())
 		{
-			// Determine if any of the items already exist
-			bool bFoundAnyItemInTree = false;
-			for (const FName& Path : NewSelectedPaths)
-			{
-				if (TSharedPtr<FTreeItem> FoundItem = FindTreeItem(Path))
-				{
-					bFoundAnyItemInTree = true;
-					break;
-				}
-			}
+			PendingInitialPaths = NewSelectedPaths;
 
-			if (bFoundAnyItemInTree)
+			// If any of the pending paths are available, select only them
+			// otherwise, leave the selection unchanged until we discover some
+			if (Algo::AnyOf(NewSelectedPaths, [this](FName VirtualPath) {
+					return TreeData->FindTreeItem(VirtualPath, /* bVisibleOnly */ true).IsValid();
+				}))
 			{
 				// Clear any previously selected paths
 				LastSelectedPaths.Empty();
 				TreeViewPtr->ClearSelection();
 			}
-			else
-			{
-				// Delay clear until first path discovered
-				// No clear occurs if selection changes during asset discovery
-				bPendingInitialPathsNeedsSelectionClear = true;
-			}
 
 			// If the selected paths is empty, the path was "All assets"
 			// This should handle that case properly
 			for (const FName& Path : NewSelectedPaths)
 			{
-				if (!ExplicitlyAddPathToSelection(Path))
-				{
-					// If we could not initially select these paths, but are still discovering assets, add them to a pending list to select them later
-					PendingInitialPaths.Add(Path);
-				}
+				ExplicitlyAddPathToSelection(Path);
 			}
+
+			// Keep entire list of pending paths around until discovery is complete or all of them are selected
+			PendingInitialPaths = NewSelectedPaths;
 		}
 		else
 		{
+			PendingInitialPaths.Reset();
 			// If all assets are already discovered, just select paths the best we can
 			SetSelectedPaths(NewSelectedPaths);
 		}
@@ -1340,26 +1768,35 @@ TSharedPtr<SWidget> SPathView::MakePathViewContextMenu()
 		return nullptr;
 	}
 
-	const TArray<FContentBrowserItem> SelectedItems = GetSelectedFolderItems();
-	if (SelectedItems.Num() == 0)
+	const TArray<FContentBrowserItem> CurrentSelectedItems = GetSelectedFolderItems();
+	if (CurrentSelectedItems.Num() == 0)
 	{
 		return nullptr;
 	}
-	
-	return OnGetItemContextMenu.Execute(SelectedItems);
+
+	return OnGetItemContextMenu.Execute(CurrentSelectedItems);
 }
 
 void SPathView::NewFolderItemRequested(const FContentBrowserItemTemporaryContext& NewItemContext)
 {
 	bool bAddedTemporaryFolder = false;
+	TSharedPtr<FTreeItem> NewItem;
+	// TODO: Consider having FTreeItem explicitly store FContentBrowserItemTemporaryContext
 	for (const FContentBrowserItemData& NewItemData : NewItemContext.GetItem().GetInternalItems())
 	{
-		bAddedTemporaryFolder |= AddFolderItem(CopyTemp(NewItemData), /*bUserNamed=*/true).IsValid();
+		NewItem = TreeData->AddFolderItem(CopyTemp(NewItemData));
 	}
 
-	if (bAddedTemporaryFolder)
+	if (NewItem.IsValid())
 	{
 		PendingNewFolderContext = NewItemContext;
+
+		PendingInitialPaths.Reset();
+
+		RecursiveExpandParents(NewItem);
+		TreeViewPtr->SetSelection(NewItem);
+		NewItem->SetNamingFolder(true);
+		TreeViewPtr->RequestScrollIntoView(NewItem);
 	}
 }
 
@@ -1370,8 +1807,13 @@ bool SPathView::ExplicitlyAddPathToSelection(const FName Path)
 		return false;
 	}
 
-	if (TSharedPtr<FTreeItem> FoundItem = FindTreeItem(Path))
+	if (TSharedPtr<FTreeItem> FoundItem = TreeData->FindTreeItem(Path))
 	{
+		if (!FoundItem->IsVisible())
+		{
+			SearchPtr->ClearSearch();
+		}
+
 		// Set the selection to the closest found folder and scroll it into view
 		RecursiveExpandParents(FoundItem);
 		LastSelectedPaths.Add(FoundItem->GetItem().GetInvariantPath());
@@ -1391,10 +1833,10 @@ bool SPathView::ShouldAllowTreeItemChangedDelegate() const
 
 void SPathView::RecursiveExpandParents(const TSharedPtr<FTreeItem>& Item)
 {
-	if ( Item->Parent.IsValid() )
+	if (TSharedPtr<FTreeItem> Parent = Item->GetParent())
 	{
-		RecursiveExpandParents(Item->Parent.Pin());
-		TreeViewPtr->SetItemExpansion(Item->Parent.Pin(), true);
+		RecursiveExpandParents(Parent);
+		TreeViewPtr->SetItemExpansion(Parent, true);
 	}
 }
 
@@ -1426,31 +1868,28 @@ void SPathView::TreeItemScrolledIntoView( TSharedPtr<FTreeItem> TreeItem, const
 
 void SPathView::GetChildrenForTree( TSharedPtr< FTreeItem > TreeItem, TArray< TSharedPtr<FTreeItem> >& OutChildren )
 {
-	TreeItem->SortChildrenIfNeeded();
-	OutChildren = TreeItem->Children;
+	TreeItem->GetSortedVisibleChildren(OutChildren);
 }
 
 void SPathView::SetTreeItemExpansionRecursive( TSharedPtr< FTreeItem > TreeItem, bool bInExpansionState )
 {
 	TreeViewPtr->SetItemExpansion(TreeItem, bInExpansionState);
 
-	// Recursively go through the children.
-	for(auto It = TreeItem->Children.CreateIterator(); It; ++It)
-	{
-		SetTreeItemExpansionRecursive( *It, bInExpansionState );
-	}
+	TreeItem->ForAllChildrenRecursive([this, bInExpansionState](const TSharedRef<FTreeItem>& Child) {
+		TreeViewPtr->SetItemExpansion(Child, bInExpansionState);
+	});
 }
 
 void SPathView::TreeSelectionChanged( TSharedPtr< FTreeItem > TreeItem, ESelectInfo::Type SelectInfo )
 {
 	if ( ShouldAllowTreeItemChangedDelegate() )
 	{
-		const TArray<TSharedPtr<FTreeItem>> SelectedItems = TreeViewPtr->GetSelectedItems();
+		const TArray<TSharedPtr<FTreeItem>> NewSelectedItems = TreeViewPtr->GetSelectedItems();
 
 		LastSelectedPaths.Empty();
-		for (int32 ItemIdx = 0; ItemIdx < SelectedItems.Num(); ++ItemIdx)
+		for (int32 ItemIdx = 0; ItemIdx < NewSelectedItems.Num(); ++ItemIdx)
 		{
-			const TSharedPtr<FTreeItem> Item = SelectedItems[ItemIdx];
+			const TSharedPtr<FTreeItem> Item = NewSelectedItems[ItemIdx];
 			if ( !ensure(Item.IsValid()) )
 			{
 				// All items must exist
@@ -1497,12 +1936,12 @@ void SPathView::TreeExpansionChanged( TSharedPtr< FTreeItem > TreeItem, bool bIs
 
 		if (!bIsExpanded)
 		{
-			const TArray<TSharedPtr<FTreeItem>> SelectedItems = TreeViewPtr->GetSelectedItems();
+			const TArray<TSharedPtr<FTreeItem>> CurrentSelectedItems = TreeViewPtr->GetSelectedItems();
 			bool bSelectTreeItem = false;
 
 			// If any selected item was a child of the collapsed node, then add the collapsed node to the current selection
 			// This avoids the selection ever becoming empty, as this causes the Content Browser to show everything
-			for (const TSharedPtr<FTreeItem>& SelectedItem : SelectedItems)
+			for (const TSharedPtr<FTreeItem>& SelectedItem : CurrentSelectedItems)
 			{
 				if (SelectedItem->IsChildOf(*TreeItem.Get()))
 				{
@@ -1521,303 +1960,84 @@ void SPathView::TreeExpansionChanged( TSharedPtr< FTreeItem > TreeItem, bool bIs
 
 void SPathView::FilterUpdated()
 {
-	Populate(/*bIsRefreshingFilter*/true);
-}
-
-void SPathView::SetSearchFilterText(const FText& InSearchText, TArray<FText>& OutErrors)
-{
-	SearchBoxFolderFilter->SetRawFilterText(InSearchText);
-
-	const FText ErrorText = SearchBoxFolderFilter->GetFilterErrorText();
-	if (!ErrorText.IsEmpty())
-	{
-		OutErrors.Add(ErrorText);
-	}
-}
-
-FText SPathView::GetHighlightText() const
-{
-	return SearchBoxFolderFilter->GetRawFilterText();
-}
-
-void SPathView::Populate(const bool bIsRefreshingFilter)
-{
-	TRACE_CPUPROFILER_EVENT_SCOPE(SPathView::Populate);
-	UE_LOG(LogContentBrowser, Verbose, TEXT("Repopulating path view"));
-
-	// Update the list of expanded path before removing the items
-	UpdateLastExpandedPathsIfDirty();
-
-	const bool bFilteringByText = !SearchBoxFolderFilter->GetRawFilterText().IsEmpty();
+	TRACE_CPUPROFILER_EVENT_SCOPE(SPathView::FilterUpdated);
 
 	// Batch the selection changed event
 	// Only emit events when the user isn't filtering, as the selection may be artificially limited by the filter
-	FScopedSelectionChangedEvent ScopedSelectionChangedEvent(SharedThis(this), !bFilteringByText && !bIsRefreshingFilter);
-
-	// Clear all root items and clear selection
-	ClearTreeItems();
-
-	// Populate the view
-	{
-		const UContentBrowserSettings* ContentBrowserSettings = GetDefault<UContentBrowserSettings>();
-		bool bDisplayEmpty = ContentBrowserSettings->DisplayEmptyFolders;
-		// check to see if we have an instance config that overrides the default in UContentBrowserSettings
-		if (FContentBrowserInstanceConfig* EditorConfig = GetContentBrowserConfig())
-		{
-			bDisplayEmpty = EditorConfig->bShowEmptyFolders;
-		}
-
-		UContentBrowserDataSubsystem* ContentBrowserData = IContentBrowserDataModule::Get().GetSubsystem();
-		const FContentBrowserDataCompiledFilter CompiledDataFilter = CreateCompiledFolderFilter();
-		TOptional<FContentBrowserFolderContentsFilter> FolderFilter;
-		if (!bDisplayEmpty)
-		{
-			FolderFilter = FContentBrowserFolderContentsFilter{};
-			FolderFilter->ItemCategoryFilter = CompiledDataFilter.ItemCategoryFilter;
-		}
-		EContentBrowserIsFolderVisibleFlags FolderFlags = ContentBrowserUtils::GetIsFolderVisibleFlags(bDisplayEmpty);
-
-		TArray<TSharedPtr<FTreeItem>> ItemsCreated;
-		ContentBrowserData->EnumerateItemsMatchingFilter(CompiledDataFilter,
-			[this, bFilteringByText, bDisplayEmpty, CompiledDataFilter, FolderFilter, FolderFlags, ContentBrowserData, &ItemsCreated](FContentBrowserItemData&& InItemData) {
-				if (!ContentBrowserData->IsFolderVisible(InItemData.GetVirtualPath(), FolderFlags, FolderFilter))
-				{
-					UE_LOG(LogContentBrowser, VeryVerbose, TEXT("Hiding folder %s that fails current pre-text filtering"), *WriteToString<256>(InItemData.GetVirtualPath()));
-					return true; // continue enumerating
-				}
-
-				if (bFilteringByText)
-				{
-					// Use the whole path so we deliberately include any children of matched parents in the filtered list
-					// Also check the display name so we match on a localized folder name
-					const FString PathStr = InItemData.GetVirtualPath().ToString();
-					if (!SearchBoxFolderFilter->PassesFilter(PathStr) && !SearchBoxFolderFilter->PassesFilter(InItemData.GetDisplayName().ToString()))
-					{
-						return true; // continue enumerating
-					}
-				}
-
-				// Using array of all items created to handle item expansion of fully virtual paths that may not be included in enumeration
-				ItemsCreated.Reset();
-				AddFolderItem(MoveTemp(InItemData), /*bUserNamed=*/false, &ItemsCreated);
-
-				for (TSharedPtr<FTreeItem> Item : ItemsCreated)
-				{
-					const FName InvariantPath = Item->GetItem().GetInvariantPath();
-					const bool bSelectedItem = LastSelectedPaths.Contains(InvariantPath);
-					const bool bExpandedItem = LastExpandedPaths.Contains(InvariantPath);
-
-					if (bFilteringByText || bSelectedItem)
-					{
-						RecursiveExpandParents(Item);
-					}
-
-					if (bSelectedItem)
-					{
-						// Tree items that match the last broadcasted paths should be re-selected them after they are added
-						if (!TreeViewPtr->IsItemSelected(Item))
-						{
-							TreeViewPtr->SetItemSelection(Item, true);
-						}
-						TreeViewPtr->RequestScrollIntoView(Item);
-					}
-
-					if (bExpandedItem)
-					{
-						// Tree items that were previously expanded should be re-expanded when repopulating
-						if (!TreeViewPtr->IsItemExpanded(Item))
-						{
-							TreeViewPtr->SetItemExpansion(Item, true);
-						}
-					}
-				}
-
-				return true;
-			});
-	}
-
-	SortRootItems();
-}
-
-void SPathView::DefaultSort(const FTreeItem* InTreeItem, TArray<TSharedPtr<FTreeItem>>& InChildren)
-{
-	if (InChildren.Num() < 2)
-	{
-		return;
-	}
-
-	static const FString ClassesPrefix = TEXT("Classes_");
-
-	struct FItemSortInfo
-	{
-		// Name to display
-		FString FolderName;
-		float Priority;
-		int32 SpecialDefaultFolderPriority;
-		bool bIsClassesFolder;
-		TSharedPtr<FTreeItem> TreeItem;
-		// Name to use when comparing "MyPlugin" vs "Classes_MyPlugin", looking up a plugin by name and other situations
-		FName ItemNameWithoutClassesPrefix;
-	};
-
-	TArray<FItemSortInfo> SortInfoArray;
-	SortInfoArray.Reserve(InChildren.Num());
+	FScopedSelectionChangedEvent ScopedSelectionChangedEvent(SharedThis(this), false);
 
-	const TArray<FName>& SpecialSortFolders = IContentBrowserDataModule::Get().GetSubsystem()->GetPathViewSpecialSortFolders();
-
-	// Generate information needed to perform sort
-	for (TSharedPtr<FTreeItem>& It : InChildren)
+	if (TreeData->GetFolderPathTextFilter().GetRawFilterText().IsEmpty())
 	{
-		FItemSortInfo& SortInfo = SortInfoArray.AddDefaulted_GetRef();
-		SortInfo.TreeItem = It;
-
-		const FName InvariantPathFName = It->GetItem().GetInvariantPath();
-		FNameBuilder InvariantPathBuilder(InvariantPathFName);
-		const FStringView InvariantPath(InvariantPathBuilder);
-
-		bool bIsRootInvariantFolder = false;
-		if (InvariantPath.Len() > 1)
-		{
-			FStringView RootInvariantFolder(InvariantPath);
-			RootInvariantFolder.RightChopInline(1);
-			int32 SecondSlashIndex = INDEX_NONE;
-			bIsRootInvariantFolder = !RootInvariantFolder.FindChar(TEXT('/'), SecondSlashIndex);
-		}
-
-		SortInfo.FolderName = It->GetItem().GetDisplayName().ToString();
-
-		SortInfo.bIsClassesFolder = false;
-		if (bIsRootInvariantFolder)
-		{
-			FNameBuilder ItemNameBuilder(It->GetItem().GetItemName());
-			const FStringView ItemNameView(ItemNameBuilder);
-			if (ItemNameView.StartsWith(ClassesPrefix))
-			{
-				SortInfo.bIsClassesFolder = true;
-				SortInfo.ItemNameWithoutClassesPrefix = FName(ItemNameView.RightChop(ClassesPrefix.Len()));
-			}
-
-			if (SortInfo.FolderName.StartsWith(ClassesPrefix))
-			{
-				SortInfo.bIsClassesFolder = true;
-				SortInfo.FolderName.RightChopInline(ClassesPrefix.Len(), EAllowShrinking::No);
-			}
-		}
-
-		if (SortInfo.ItemNameWithoutClassesPrefix.IsNone())
-		{
-			SortInfo.ItemNameWithoutClassesPrefix = It->GetItem().GetItemName();
-		}
-
-		if (SortInfo.bIsClassesFolder)
+		TreeData->ClearItemFilterState();
+		TreeViewPtr->ClearExpandedItems();
+		TArray<TSharedPtr<FTreeItem>> SelectedItems = TreeViewPtr->GetSelectedItems();
+		for (const TSharedPtr<FTreeItem>& SelectedItem : SelectedItems)
 		{
-			// Sort using a path without "Classes_" prefix
-			FStringView InvariantWithoutClassesPrefix(InvariantPath);
-			InvariantWithoutClassesPrefix.RightChopInline(1);
-			if (InvariantWithoutClassesPrefix.StartsWith(ClassesPrefix))
-			{
-				InvariantWithoutClassesPrefix.RightChopInline(ClassesPrefix.Len());
-				FNameBuilder Builder;
-				Builder.Append(TEXT("/"));
-				Builder.Append(InvariantWithoutClassesPrefix);
-				SortInfo.SpecialDefaultFolderPriority = SpecialSortFolders.IndexOfByKey(FName(Builder));
-			}
-			else
+			for (TSharedPtr<FTreeItem> Parent = SelectedItem->GetParent(); Parent.IsValid();
+				 Parent = Parent->GetParent())
 			{
-				SortInfo.SpecialDefaultFolderPriority = SpecialSortFolders.IndexOfByKey(InvariantPathFName);
+				TreeViewPtr->SetItemExpansion(Parent, true);
 			}
 		}
-		else
-		{
-			SortInfo.SpecialDefaultFolderPriority = SpecialSortFolders.IndexOfByKey(InvariantPathFName);
-		}
 
-		if (bIsRootInvariantFolder)
-		{
-			if (SortInfo.SpecialDefaultFolderPriority == INDEX_NONE)
-			{
-				SortInfo.Priority = FContentBrowserSingleton::Get().GetPluginSettings(SortInfo.ItemNameWithoutClassesPrefix).RootFolderSortPriority;
-			}
-			else
-			{
-				SortInfo.Priority = 1.f;
-			}
-		}
-		else
+		if (SelectedItems.Num())
 		{
-			if (SortInfo.SpecialDefaultFolderPriority != INDEX_NONE)
-			{
-				SortInfo.Priority = 1.f;
-			}
-			else
-			{
-				SortInfo.Priority = 0.f;
-			}
+			TreeViewPtr->RequestScrollIntoView(SelectedItems[0]);
 		}
 	}
-
-	// Perform sort
-	SortInfoArray.Sort([](const FItemSortInfo& SortInfoA, const FItemSortInfo& SortInfoB) -> bool
+	else
 	{
-		if (SortInfoA.Priority != SortInfoB.Priority)
-		{
-			// Not the same priority, use priority to sort
-			return SortInfoA.Priority > SortInfoB.Priority;
-		}
-		else if (SortInfoA.SpecialDefaultFolderPriority != SortInfoB.SpecialDefaultFolderPriority)
-		{
-			// Special folders use the index to sort. Non special folders are all set to 0.
-			return SortInfoA.SpecialDefaultFolderPriority < SortInfoB.SpecialDefaultFolderPriority;
-		}
-		else
+		TreeData->FilterFullFolderTree();
+		TreeViewPtr->ClearExpandedItems();
+		for (const TSharedPtr<FTreeItem>& Root : *TreeData->GetVisibleRootItems())
 		{
-			// If either is a class folder and names without classes prefix are same
-			if ((SortInfoA.bIsClassesFolder != SortInfoB.bIsClassesFolder) && (SortInfoA.ItemNameWithoutClassesPrefix == SortInfoB.ItemNameWithoutClassesPrefix))
-			{
-				return !SortInfoA.bIsClassesFolder;
-			}
-
-			// Two non special folders of the same priority, sort alphabetically
-			const int32 CompareResult = UE::ComparisonUtility::CompareWithNumericSuffix(SortInfoA.FolderName, SortInfoB.FolderName);
-			if (CompareResult != 0)
-			{
-				return CompareResult < 0;
-			}
-			else
-			{
-				// Classes folders have the same name so sort them adjacent but under non-classes
-				return !SortInfoA.bIsClassesFolder;
-			}
+			TreeViewPtr->SetItemExpansion(Root, true);
+			Root->ForAllChildrenRecursive([this](const TSharedPtr<FTreeItem> Descendant) {
+				if (Descendant->GetHasVisibleDescendants())
+				{
+					TreeViewPtr->SetItemExpansion(Descendant, true);
+				}
+			});
 		}
-	});
-
-	// Replace with sorted array
-	TArray<TSharedPtr<FTreeItem>> NewList;
-	NewList.Reserve(SortInfoArray.Num());
-	for (const FItemSortInfo& It : SortInfoArray)
-	{
-		NewList.Add(It.TreeItem);
 	}
-	InChildren = MoveTemp(NewList);
 }
 
-void SPathView::SortRootItems()
+void SPathView::SetSearchFilterText(const FText& InSearchText, TArray<FText>& OutErrors)
 {
-	if (SortOverride.IsBound())
-	{
-		SortOverride.Execute(nullptr, TreeRootItems);
-	}
-	else
+	TreeData->GetFolderPathTextFilter().SetRawFilterText(InSearchText);
+
+	const FText ErrorText = TreeData->GetFolderPathTextFilter().GetFilterErrorText();
+	if (!ErrorText.IsEmpty())
 	{
-		DefaultSort(nullptr, TreeRootItems);
+		OutErrors.Add(ErrorText);
 	}
+}
 
-	TreeViewPtr->RequestTreeRefresh();
+FText SPathView::GetHighlightText() const
+{
+	return TreeData->GetFolderPathTextFilter().GetRawFilterText();
 }
 
-void SPathView::PopulateFolderSearchStrings( const FString& FolderName, OUT TArray< FString >& OutSearchStrings ) const
+void SPathView::Populate(const bool bIsRefreshingFilter)
 {
-	OutSearchStrings.Add( FolderName );
+	TRACE_CPUPROFILER_EVENT_SCOPE(SPathView::Populate);
+	UE_LOG(LogContentBrowser, Verbose, TEXT("Repopulating path view"));
+
+	const bool bFilteringByText = !TreeData->GetFolderPathTextFilter().GetRawFilterText().IsEmpty();
+
+	// Batch the selection changed event
+	// Only emit events when the user isn't filtering, as the selection may be artificially limited by the filter
+	FScopedSelectionChangedEvent ScopedSelectionChangedEvent(SharedThis(this),
+		!bFilteringByText && !bIsRefreshingFilter);
+	TreeData->PopulateFullFolderTree(CreateCompiledFolderFilter());
+	TreeData->FilterFullFolderTree();
+	TreeData->SortRootItems();
+
+	// Select any of our initial paths which aren't currently selected
+	for (FName VirtualPath : PendingInitialPaths)
+	{
+		ExplicitlyAddPathToSelection(VirtualPath);
+	}
 }
 
 FReply SPathView::OnFolderDragDetected(const FGeometry& Geometry, const FPointerEvent& MouseEvent)
@@ -1849,8 +2069,16 @@ bool SPathView::VerifyFolderNameChanged(const TSharedPtr< FTreeItem >& TreeItem,
 	return true;
 }
 
-void SPathView::FolderNameChanged( const TSharedPtr< FTreeItem >& TreeItem, const FString& ProposedName, const UE::Slate::FDeprecateVector2DParameter& MessageLocation, const ETextCommit::Type CommitType )
+void SPathView::FolderNameChanged(const TSharedPtr<FTreeItem>& TreeItem,
+	const FString& ProposedName,
+	const UE::Slate::FDeprecateVector2DParameter& MessageLocation,
+	const ETextCommit::Type CommitType)
 {
+	if (!TreeItem.IsValid())
+	{
+		return;
+	}
+
 	bool bSuccess = false;
 	FText ErrorMessage;
 
@@ -1860,13 +2088,13 @@ void SPathView::FolderNameChanged( const TSharedPtr< FTreeItem >& TreeItem, cons
 		checkf(FContentBrowserItemKey(TreeItem->GetItem()) == FContentBrowserItemKey(PendingNewFolderContext.GetItem()), TEXT("PendingNewFolderContext was still set when attempting to rename a different item!"));
 
 		// Remove the temporary item before we do any work to ensure the new item creation is not prevented
-		RemoveFolderItem(TreeItem);
+		TreeData->RemoveFolderItem(TreeItem.ToSharedRef());
 
 		// Clearing the rename box on a newly created item cancels the entire creation process
 		if (CommitType == ETextCommit::OnCleared)
 		{
 			// We need to select the parent item of this folder, as the folder would have become selected while it was being named
-			if (TSharedPtr<FTreeItem> ParentTreeItem = TreeItem->Parent.Pin())
+			if (TSharedPtr<FTreeItem> ParentTreeItem = TreeItem->GetParent())
 			{
 				TreeViewPtr->SetItemSelection(ParentTreeItem, true);
 			}
@@ -1909,7 +2137,7 @@ void SPathView::FolderNameChanged( const TSharedPtr< FTreeItem >& TreeItem, cons
 		TSharedPtr<FTreeItem> NewTreeItem;
 		for (const FContentBrowserItemData& NewItemData : NewItem.GetInternalItems())
 		{
-			NewTreeItem = AddFolderItem(CopyTemp(NewItemData));
+			NewTreeItem = TreeData->AddFolderItem(CopyTemp(NewItemData));
 		}
 
 		// Select the new item
@@ -1928,70 +2156,6 @@ void SPathView::FolderNameChanged( const TSharedPtr< FTreeItem >& TreeItem, cons
 	}
 }
 
-bool SPathView::FolderAlreadyExists(const TSharedPtr< FTreeItem >& TreeItem, TSharedPtr< FTreeItem >& ExistingItem)
-{
-	ExistingItem.Reset();
-
-	if ( TreeItem.IsValid() )
-	{
-		if ( TreeItem->Parent.IsValid() )
-		{
-			// This item has a parent, try to find it in its parent's children
-			TSharedPtr<FTreeItem> ParentItem = TreeItem->Parent.Pin();
-
-			for ( auto ChildIt = ParentItem->Children.CreateConstIterator(); ChildIt; ++ChildIt )
-			{
-				const TSharedPtr<FTreeItem>& Child = *ChildIt;
-				if ( Child != TreeItem && Child->GetItem().GetItemName() == TreeItem->GetItem().GetItemName() )
-				{
-					// The item is in its parent already
-					ExistingItem = Child;
-					break;
-				}
-			}
-		}
-		else
-		{
-			// This item is part of the root set
-			for ( auto RootIt = TreeRootItems.CreateConstIterator(); RootIt; ++RootIt )
-			{
-				const TSharedPtr<FTreeItem>& Root = *RootIt;
-				if ( Root != TreeItem && Root->GetItem().GetItemName() == TreeItem->GetItem().GetItemName() )
-				{
-					// The item is part of the root set already
-					ExistingItem = Root;
-					break;
-				}
-			}
-		}
-	}
-
-	return ExistingItem.IsValid();
-}
-
-void SPathView::RemoveFolderItem(const TSharedPtr< FTreeItem >& TreeItem)
-{
-	if ( TreeItem.IsValid() )
-	{
-		if ( TreeItem->Parent.IsValid() )
-		{
-			// Remove this item from it's parent's list
-			TreeItem->Parent.Pin()->Children.Remove(TreeItem);
-		}
-		else
-		{
-			// This was a root node, remove from the root list
-			TreeRootItems.Remove(TreeItem);
-		}
-
-		const FName VirtualPath = TreeItem->GetItem().GetVirtualPath();
-		ensure(!FindTreeItem(VirtualPath) || (FindTreeItem(VirtualPath) == TreeItem));
-		TreeItemLookup.Remove(VirtualPath);
-
-		TreeViewPtr->RequestTreeRefresh();
-	}
-}
-
 bool SPathView::IsTreeItemExpanded(TSharedPtr<FTreeItem> TreeItem) const
 {
 	return TreeViewPtr->IsItemExpanded(TreeItem);
@@ -2011,125 +2175,31 @@ void SPathView::HandleItemDataUpdated(TArrayView<const FContentBrowserItemDataUp
 		return;
 	}
 
-	const bool bFilteringByText = !SearchBoxFolderFilter->GetRawFilterText().IsEmpty();
+	// TODO: Consider batching if sometimes we get very few items and filter construction time dominates
+	if (!Algo::AnyOf(InUpdatedItems,
+			[](const FContentBrowserItemDataUpdate& Update) { return Update.GetItemData().IsFolder(); }))
+	{
+		return;
+	}
+
+	const bool bFilteringByText = !TreeData->GetFolderPathTextFilter().GetRawFilterText().IsEmpty();
 
 	// Batch the selection changed event
 	// Only emit events when the user isn't filtering, as the selection may be artificially limited by the filter
 	FScopedSelectionChangedEvent ScopedSelectionChangedEvent(SharedThis(this), !bFilteringByText);
 
 	const double HandleItemDataUpdatedStartTime = FPlatformTime::Seconds();
+	TreeData->ProcessDataUpdates(InUpdatedItems, CreateCompiledFolderFilter());
+	UE_LOG(LogContentBrowser,
+		VeryVerbose,
+		TEXT("PathView - HandleItemDataUpdated completed in %0.4f seconds for %d items"),
+		FPlatformTime::Seconds() - HandleItemDataUpdatedStartTime,
+		InUpdatedItems.Num());
 
-	const UContentBrowserSettings* ContentBrowserSettings = GetDefault<UContentBrowserSettings>();
-	bool bDisplayEmpty = ContentBrowserSettings->DisplayEmptyFolders;
-	// check to see if we have an instance config that overrides the default in UContentBrowserSettings
-	if (FContentBrowserInstanceConfig* EditorConfig = GetContentBrowserConfig())
-	{
-		bDisplayEmpty = EditorConfig->bShowEmptyFolders;
-	}
-
-	UContentBrowserDataSubsystem* ContentBrowserData = IContentBrowserDataModule::Get().GetSubsystem();
-
-	// We defer this compilation as it's quite expensive due to being recursive, and not all updates will contain new folders
-	bool bHasCompiledDataFilter = false;
-	FContentBrowserDataCompiledFilter CompiledDataFilter;
-	auto ConditionalCompileFilter = [this, &bHasCompiledDataFilter, &CompiledDataFilter]()
-	{
-		if (!bHasCompiledDataFilter)
-		{
-			bHasCompiledDataFilter = true;
-			CompiledDataFilter = CreateCompiledFolderFilter();
-		}
-	};
-
-	TOptional<FContentBrowserFolderContentsFilter> FolderFilter;
-	if (!bDisplayEmpty)
-	{
-		FolderFilter = FContentBrowserFolderContentsFilter{};
-		FolderFilter->ItemCategoryFilter = CompiledDataFilter.ItemCategoryFilter;
-	}
-	EContentBrowserIsFolderVisibleFlags FolderFlags = ContentBrowserUtils::GetIsFolderVisibleFlags(bDisplayEmpty);
-	auto DoesItemPassFilter = [this, bFilteringByText, FolderFlags, FolderFilter, ContentBrowserData, &CompiledDataFilter](const FContentBrowserItemData& InItemData) {
-		UContentBrowserDataSource* ItemDataSource = InItemData.GetOwnerDataSource();
-		if (!ItemDataSource->DoesItemPassFilter(InItemData, CompiledDataFilter))
-		{
-			return false;
-		}
-
-		if (!ContentBrowserData->IsFolderVisible(InItemData.GetVirtualPath(), FolderFlags, FolderFilter))
-		{
-			UE_LOG(LogContentBrowser, VeryVerbose, TEXT("Hiding folder %s that fails broad visibility filter"), *WriteToString<256>(InItemData.GetVirtualPath()));
-			return false;
-		}
-
-		if (bFilteringByText)
-		{
-			// Use the whole path so we deliberately include any children of matched parents in the filtered list
-			const FString PathStr = InItemData.GetVirtualPath().ToString();
-			if (!SearchBoxFolderFilter->PassesFilter(PathStr))
-			{
-				return false;
-			}
-		}
-
-		return true;
-	};
-
-	for (const FContentBrowserItemDataUpdate& ItemDataUpdate : InUpdatedItems)
+	for (FName VirtualPath : PendingInitialPaths)
 	{
-		const FContentBrowserItemData& ItemDataRef = ItemDataUpdate.GetItemData();
-		if (!ItemDataRef.IsFolder())
-		{
-			continue;
-		}
-
-		ConditionalCompileFilter();
-
-		FContentBrowserItemData ItemData = ItemDataRef;
-		ItemData.GetOwnerDataSource()->ConvertItemForFilter(ItemData, CompiledDataFilter);
-
-		switch (ItemDataUpdate.GetUpdateType())
-		{
-		case EContentBrowserItemUpdateType::Added:
-			if (DoesItemPassFilter(ItemData))
-			{
-				AddFolderItem(MoveTemp(ItemData));
-			}
-			break;
-
-		case EContentBrowserItemUpdateType::Modified:
-			if (DoesItemPassFilter(ItemData))
-			{
-				AddFolderItem(MoveTemp(ItemData));
-			}
-			else
-			{
-				RemoveFolderItem(ItemData);
-			}
-			break;
-
-		case EContentBrowserItemUpdateType::Moved:
-		{
-			const FContentBrowserItemData OldMinimalItemData(ItemData.GetOwnerDataSource(), ItemData.GetItemType(), ItemDataUpdate.GetPreviousVirtualPath(), NAME_None, FText(), nullptr);
-			RemoveFolderItem(OldMinimalItemData);
-
-			if (DoesItemPassFilter(ItemData))
-			{
-				AddFolderItem(MoveTemp(ItemData));
-			}
-		}
-		break;
-
-		case EContentBrowserItemUpdateType::Removed:
-			RemoveFolderItem(ItemData);
-			break;
-
-		default:
-			checkf(false, TEXT("Unexpected EContentBrowserItemUpdateType!"));
-			break;
-		}
+		ExplicitlyAddPathToSelection(VirtualPath);
 	}
-
-	UE_LOG(LogContentBrowser, VeryVerbose, TEXT("PathView - HandleItemDataUpdated completed in %0.4f seconds for %d items"), FPlatformTime::Seconds() - HandleItemDataUpdatedStartTime, InUpdatedItems.Num());
 }
 
 void SPathView::HandleItemDataRefreshed()
@@ -2148,14 +2218,6 @@ void SPathView::HandleItemDataDiscoveryComplete()
 {
 	// If there were any more initial paths, they no longer exist so clear them now.
 	PendingInitialPaths.Empty();
-	bPendingInitialPathsNeedsSelectionClear = false;
-}
-
-bool SPathView::PathIsFilteredFromViewBySearch(const FString& InPath) const
-{
-	return !SearchBoxFolderFilter->GetRawFilterText().IsEmpty()
-		&& !SearchBoxFolderFilter->PassesFilter(InPath)
-		&& !FindTreeItem(*InPath);
 }
 
 void SPathView::HandleSettingChanged(FName PropertyName)
@@ -2170,6 +2232,11 @@ void SPathView::HandleSettingChanged(FName PropertyName)
 		|| (PropertyName == GET_MEMBER_NAME_CHECKED(UContentBrowserSettings, bDisplayFriendlyNameForPluginFolders))
 		|| (PropertyName == NAME_None)) // @todo: Needed if PostEditChange was called manually, for now
 	{
+		UE_LOG(LogContentBrowser,
+			Log,
+			TEXT("[%s][PathView] HandleSettingChanged %s"),
+			*WriteToString<256>(OwningContentBrowserName),
+			*WriteToString<256>(PropertyName));
 		const bool bHadSelectedPath = TreeViewPtr->GetNumItemsSelected() > 0;
 
 		// Update our path view so that it can include/exclude the dev folder
@@ -2180,7 +2247,7 @@ void SPathView::HandleSettingChanged(FName PropertyName)
 		{
 			for (const FName VirtualPath : GetDefaultPathsToSelect())
 			{
-				if (TSharedPtr<FTreeItem> TreeItemToSelect = FindTreeItem(VirtualPath))
+				if (TSharedPtr<FTreeItem> TreeItemToSelect = TreeData->FindTreeItem(VirtualPath))
 				{
 					TreeViewPtr->SetSelection(TreeItemToSelect);
 					break;
@@ -2239,8 +2306,8 @@ TArray<FName> SPathView::GetDefaultPathsToSelect() const
 TArray<FName> SPathView::GetRootPathItemNames() const
 {
 	TArray<FName> RootPathItemNames;
-	RootPathItemNames.Reserve(TreeRootItems.Num());
-	for (const TSharedPtr<FTreeItem>& RootItem : TreeRootItems)
+	RootPathItemNames.Reserve(TreeData->GetVisibleRootItems()->Num());
+	for (const TSharedPtr<FTreeItem>& RootItem : *TreeData->GetVisibleRootItems())
 	{
 		if (RootItem.IsValid())
 		{
@@ -2327,17 +2394,6 @@ TSharedRef<SWidget> SPathView::CreateFavoritesView()
 
 void SFavoritePathView::Construct(const FArguments& InArgs)
 {
-	SAssignNew(TreeViewPtr, STreeView< TSharedPtr<FTreeItem> >)
-		.TreeItemsSource(&TreeRootItems)
-		.OnGetChildren(this, &SFavoritePathView::GetChildrenForTree)
-		.OnGenerateRow(this, &SFavoritePathView::GenerateTreeRow)
-		.OnItemScrolledIntoView(this, &SFavoritePathView::TreeItemScrolledIntoView)
-		.ItemHeight(18)
-		.SelectionMode(InArgs._SelectionMode)
-		.OnSelectionChanged(this, &SFavoritePathView::TreeSelectionChanged)
-		.OnContextMenuOpening(this, &SFavoritePathView::MakePathViewContextMenu)
-		.ClearSelectionOnClick(false);
-
 	// Bind the favorites menu to update after folder changes
 	AssetViewUtils::OnFolderPathChanged().AddSP(this, &SFavoritePathView::FixupFavoritesFromExternalChange); 
 
@@ -2346,6 +2402,16 @@ void SFavoritePathView::Construct(const FArguments& InArgs)
 	SPathView::Construct(InArgs);
 }
 
+void SFavoritePathView::ConfigureTreeView(STreeView<TSharedPtr<FTreeItem>>::FArguments& InArgs)
+{
+	// Don't bind some stuff that the parent class binds such as item expansion
+}
+
+SFavoritePathView::SFavoritePathView()
+{
+	bFlat = true;
+}
+
 SFavoritePathView::~SFavoritePathView()
 {
 	FContentBrowserSingleton::Get().UnregisterOnFavoritesChangedDelegate(OnFavoritesChangedHandle);
@@ -2357,54 +2423,9 @@ void SFavoritePathView::Populate(const bool bIsRefreshingFilter)
 
 	// Don't allow the selection changed delegate to be fired here
 	FScopedPreventTreeItemChangedDelegate DelegatePrevention(SharedThis(this));
-
-	// Clear all root items and clear selection
-	ClearTreeItems();
-
-	const TArray<FString>& FavoritePaths = ContentBrowserUtils::GetFavoriteFolders();
-	if (FavoritePaths.Num() > 0)
-	{
-		UContentBrowserDataSubsystem* ContentBrowserData = IContentBrowserDataModule::Get().GetSubsystem();
-		const FContentBrowserDataCompiledFilter CompiledDataFilter = CreateCompiledFolderFilter();
-
-		for (const FString& InvariantPath : FavoritePaths)
-		{
-			FName VirtualPath;
-			IContentBrowserDataModule::Get().GetSubsystem()->ConvertInternalPathToVirtual(InvariantPath, VirtualPath);
-			const FString Path = VirtualPath.ToString();
-
-			// Use the whole path so we deliberately include any children of matched parents in the filtered list
-			if (!SearchBoxFolderFilter->PassesFilter(Path))
-			{
-				continue;
-			}
-
-			ContentBrowserData->EnumerateItemsAtPath(*Path, CompiledDataFilter.ItemTypeFilter, 
-			[this, &CompiledDataFilter](FContentBrowserItemData&& InItemData)
-				{
-					UContentBrowserDataSource* ItemDataSource = InItemData.GetOwnerDataSource();
-					ItemDataSource->ConvertItemForFilter(InItemData, CompiledDataFilter);
-					if (ItemDataSource->DoesItemPassFilter(InItemData, CompiledDataFilter))
-					{
-						if (TSharedPtr<FTreeItem> Item = AddFolderItem(MoveTemp(InItemData)))
-						{
-							const bool bSelectedItem = LastSelectedPaths.Contains(Item->GetItem().GetInvariantPath());
-							if (bSelectedItem)
-							{
-								// Tree items that match the last broadcasted paths should be re-selected them after they are added
-								TreeViewPtr->SetItemSelection(Item, true);
-								TreeViewPtr->RequestScrollIntoView(Item);
-							}
-						}
-					}
-
-					return true;
-				}
-			);
-		}
-	}
-
-	SortRootItems();
+	TreeData->PopulateWithFavorites(CreateCompiledFolderFilter());
+	TreeData->SortRootItems();
+	TreeData->FilterFullFolderTree();
 }
 
 void SFavoritePathView::SaveSettings(const FString& IniFilename, const FString& IniSection, const FString& SettingsString) const
@@ -2468,41 +2489,6 @@ void SFavoritePathView::LoadSettings(const FString& IniFilename, const FString&
 	}
 }
 
-TSharedPtr<FTreeItem> SFavoritePathView::AddFolderItem(FContentBrowserItemData&& InItem, const bool bUserNamed, TArray<TSharedPtr<FTreeItem>>* OutItemsCreated)
-{
-	if (!ensure(TreeViewPtr.IsValid()))
-	{
-		// No tree view for some reason
-		return nullptr;
-	}
-
-	// The favorite view will add all items at the root level
-
-	// Try and find an existing tree item
-	TWeakPtr<FTreeItem>* WeakExistingItem = TreeItemLookup.Find(InItem.GetVirtualPath());
-	if (WeakExistingItem != nullptr)
-	{
-		if (TSharedPtr<FTreeItem> ExistingItem = WeakExistingItem->Pin())
-		{
-			ExistingItem->AppendItemData(InItem);
-			return ExistingItem;
-		}
-	}
-
-	// No match - create a new item
-	TSharedPtr<FTreeItem> CurrentTreeItem = MakeShared<FTreeItem>(MoveTemp(InItem));
-	TreeRootItems.Add(CurrentTreeItem);
-	TreeItemLookup.Add(CurrentTreeItem->GetItem().GetVirtualPath(), CurrentTreeItem);
-	//TreeViewPtr->SetSelection(CurrentTreeItem);
-	TreeViewPtr->RequestTreeRefresh();
-	if (OutItemsCreated)
-	{
-		OutItemsCreated->Add(CurrentTreeItem);
-	}
-
-	return CurrentTreeItem;
-}
-
 TSharedRef<ITableRow> SFavoritePathView::GenerateTreeRow(TSharedPtr<FTreeItem> TreeItem, const TSharedRef<STableViewBase>& OwnerTable)
 {
 	check(TreeItem.IsValid());
@@ -2529,6 +2515,12 @@ void SFavoritePathView::HandleItemDataUpdated(TArrayView<const FContentBrowserIt
 		return;
 	}
 
+	if (!Algo::AnyOf(InUpdatedItems,
+			[](const FContentBrowserItemDataUpdate& Update) { return Update.GetItemData().IsFolder(); }))
+	{
+		return;
+	}
+
 	TSet<FName> FavoritePaths;
 	{
 		const TArray<FString>& FavoritePathStrs = ContentBrowserUtils::GetFavoriteFolders();
@@ -2546,116 +2538,50 @@ void SFavoritePathView::HandleItemDataUpdated(TArrayView<const FContentBrowserIt
 
 	// Don't allow the selection changed delegate to be fired here
 	FScopedPreventTreeItemChangedDelegate DelegatePrevention(SharedThis(this));
-
 	const double HandleItemDataUpdatedStartTime = FPlatformTime::Seconds();
 
-	const bool bFilteringByText = !SearchBoxFolderFilter->GetRawFilterText().IsEmpty();
-
-	// We defer this compilation as it's quite expensive due to being recursive, and not all updates will contain new folders
-	bool bHasCompiledDataFilter = false;
-	FContentBrowserDataCompiledFilter CompiledDataFilter;
-	auto ConditionalCompileFilter = [this, &bHasCompiledDataFilter, &CompiledDataFilter]()
-	{
-		if (!bHasCompiledDataFilter)
-		{
-			bHasCompiledDataFilter = true;
-			CompiledDataFilter = CreateCompiledFolderFilter();
-		}
-	};
-
-	auto DoesItemPassFilter = [this, bFilteringByText, &CompiledDataFilter, &FavoritePaths](const FContentBrowserItemData& InItemData)
+	// Limit the updates to only folders which are favorites
+	TArray<FContentBrowserItemDataUpdate> FilteredUpdates;
+	Algo::CopyIf(InUpdatedItems, FilteredUpdates, [FavoritePaths](const FContentBrowserItemDataUpdate& Update) {
+		return FavoritePaths.Contains(Update.GetItemData().GetVirtualPath());
+	});
+	if (FilteredUpdates.Num())
 	{
-		if (!FavoritePaths.Contains(InItemData.GetVirtualPath()))
-		{
-			return false;
-		}
-
-		UContentBrowserDataSource* ItemDataSource = InItemData.GetOwnerDataSource();
-		if (!ItemDataSource->DoesItemPassFilter(InItemData, CompiledDataFilter))
-		{
-			return false;
-		}
-
-		if (bFilteringByText)
-		{
-			// Use the whole path so we deliberately include any children of matched parents in the filtered list
-			const FString PathStr = InItemData.GetVirtualPath().ToString();
-			if (!SearchBoxFolderFilter->PassesFilter(PathStr))
-			{
-				return false;
-			}
-		}
-
-		return true;
-	};
+		TreeData->ProcessDataUpdates(MakeArrayView(FilteredUpdates), CreateCompiledFolderFilter());
+	}
 
+	// Update saved favorites
 	for (const FContentBrowserItemDataUpdate& ItemDataUpdate : InUpdatedItems)
 	{
-		const FContentBrowserItemData& ItemDataRef = ItemDataUpdate.GetItemData();
-		if (!ItemDataRef.IsFolder())
+		const FContentBrowserItemData& ItemData = ItemDataUpdate.GetItemData();
+		if (!ItemData.IsFolder())
 		{
 			continue;
 		}
 
-		ConditionalCompileFilter();
-
-		FContentBrowserItemData ItemData = ItemDataUpdate.GetItemData();
-		ItemData.GetOwnerDataSource()->ConvertItemForFilter(ItemData, CompiledDataFilter);
-
 		switch (ItemDataUpdate.GetUpdateType())
 		{
-		case EContentBrowserItemUpdateType::Added:
-			if (DoesItemPassFilter(ItemData))
-			{
-				AddFolderItem(MoveTemp(ItemData));
-			}
-			break;
-
-		case EContentBrowserItemUpdateType::Modified:
-			if (DoesItemPassFilter(ItemData))
-			{
-				AddFolderItem(MoveTemp(ItemData));
-			}
-			else
-			{
-				RemoveFolderItem(ItemData);
-			}
-			break;
-
-		case EContentBrowserItemUpdateType::Moved:
-		{
-			const FContentBrowserItemData OldMinimalItemData(ItemData.GetOwnerDataSource(), ItemData.GetItemType(), ItemDataUpdate.GetPreviousVirtualPath(), NAME_None, FText(), nullptr);
-			RemoveFolderItem(OldMinimalItemData);
-
-			if (DoesItemPassFilter(ItemData))
-			{
-				AddFolderItem(MoveTemp(ItemData));
-			}
-
-			ContentBrowserUtils::RemoveFavoriteFolder(FContentBrowserItemPath(ItemDataUpdate.GetPreviousVirtualPath(), EContentBrowserPathType::Virtual));
-		}
-		break;
-
-		case EContentBrowserItemUpdateType::Removed:
-			RemoveFolderItem(ItemData);
-			ContentBrowserUtils::RemoveFavoriteFolder(FContentBrowserItemPath(ItemData.GetVirtualPath(), EContentBrowserPathType::Virtual));
-			break;
-
-		default:
-			checkf(false, TEXT("Unexpected EContentBrowserItemUpdateType!"));
-			break;
+			case EContentBrowserItemUpdateType::Added:
+				break;
+			case EContentBrowserItemUpdateType::Modified:
+				break;
+			case EContentBrowserItemUpdateType::Moved:
+				ContentBrowserUtils::RemoveFavoriteFolder(
+					FContentBrowserItemPath(ItemDataUpdate.GetPreviousVirtualPath(), EContentBrowserPathType::Virtual));
+				break;
+			case EContentBrowserItemUpdateType::Removed:
+				ContentBrowserUtils::RemoveFavoriteFolder(
+					FContentBrowserItemPath(ItemData.GetVirtualPath(), EContentBrowserPathType::Virtual));
+				break;
+			default:
+				checkf(false, TEXT("Unexpected EContentBrowserItemUpdateType!"));
+				break;
 		}
 	}
 
 	UE_LOG(LogContentBrowser, VeryVerbose, TEXT("FavoritePathView - HandleItemDataUpdated completed in %0.4f seconds for %d items"), FPlatformTime::Seconds() - HandleItemDataUpdatedStartTime, InUpdatedItems.Num());
 }
 
-bool SFavoritePathView::PathIsFilteredFromViewBySearch(const FString& InPath) const
-{
-	return SPathView::PathIsFilteredFromViewBySearch(InPath)
-		&& ContentBrowserUtils::IsFavoriteFolder(FContentBrowserItemPath(InPath, EContentBrowserPathType::Virtual));
-}
-
 void SFavoritePathView::FixupFavoritesFromExternalChange(TArrayView<const AssetViewUtils::FMovedContentFolder> MovedFolders)
 {
 	for (const AssetViewUtils::FMovedContentFolder& MovedFolder : MovedFolders)
@@ -2670,7 +2596,7 @@ void SFavoritePathView::FixupFavoritesFromExternalChange(TArrayView<const AssetV
 			// Add the new path to favorites instead
 			const FString& NewPath = MovedFolder.Value;
 			ContentBrowserUtils::AddFavoriteFolder(FContentBrowserItemPath(NewPath, EContentBrowserPathType::Virtual));
-			TSharedPtr<FTreeItem> Item = FindTreeItem(*NewPath);
+			TSharedPtr<FTreeItem> Item = TreeData->FindTreeItem(*NewPath);
 			if (Item.IsValid())
 			{
 				TreeViewPtr->SetItemSelection(Item, true);
diff --git a/Engine/Source/Editor/ContentBrowser/Private/SPathView.h b/Engine/Source/Editor/ContentBrowser/Private/SPathView.h
--- a/Engine/Source/Editor/ContentBrowser/Private/SPathView.h
+++ b/Engine/Source/Editor/ContentBrowser/Private/SPathView.h
@@ -52,11 +52,10 @@ struct FAssetData;
 struct FGeometry;
 struct FHistoryData;
 struct FPathViewConfig;
+struct FPathViewData;
 struct FPointerEvent;
 struct FContentBrowserInstanceConfig;
 
-typedef TTextFilter< const FString& > FolderTextFilter;
-
 /**
  * The tree view of folders which contain content.
  */
@@ -135,6 +134,9 @@ public:
 		/** The instance name of the owning content browser. */
 		SLATE_ARGUMENT( FName, OwningContentBrowserName )
 
+		/** Default path to create even if it doesn't exist for use by path picker */
+		SLATE_ARGUMENT(FString, DefaultPath)
+
 	SLATE_END_ARGS()
 
 	/** Destructor */
@@ -167,12 +169,6 @@ public:
 	/** Called when "new folder" is selected in the context menu */
 	void NewFolderItemRequested(const FContentBrowserItemTemporaryContext& NewItemContext);
 
-	/** Adds nodes to the tree in order to construct the specified item. If bUserNamed is true, the user will name the folder and the item includes the default name. */
-	virtual TSharedPtr<FTreeItem> AddFolderItem(FContentBrowserItemData&& InItem, const bool bUserNamed = false, TArray<TSharedPtr<FTreeItem>>* OutItemsCreated = nullptr);
-
-	/** Attempts to remove the item from the tree. Returns true when successful. */
-	bool RemoveFolderItem(const FContentBrowserItemData& InItem);
-
 	/** Sets up an inline rename for the specified folder */
 	void RenameFolderItem(const FContentBrowserItem& InItem);
 
@@ -207,8 +203,8 @@ public:
 	 */
 	void SyncToLegacy( TArrayView<const FAssetData> AssetDataList, TArrayView<const FString> FolderList, const bool bAllowImplicitSync = false );
 
-	/** Finds the item that represents the specified path, if it exists. */
-	TSharedPtr<FTreeItem> FindTreeItem(FName InPath) const;
+	/** Returns whether the tree contains an item with the given virtual path. */
+	bool DoesItemExist(FName InVirtualPath) const;
 
 	/** Sets the state of the path view to the one described by the history data */
 	void ApplyHistoryData( const FHistoryData& History );
@@ -259,9 +255,6 @@ protected:
 	/** Expands all parents of the specified item */
 	void RecursiveExpandParents(const TSharedPtr<FTreeItem>& Item);
 
-	/** Sort the root items into the correct order */
-	void SortRootItems();
-
 	/** Handles updating the view when content items are changed */
 	virtual void HandleItemDataUpdated(TArrayView<const FContentBrowserItemDataUpdate> InUpdatedItems);
 
@@ -271,9 +264,6 @@ protected:
 	/** Notification for when the content browser has completed it's initial search */
 	void HandleItemDataDiscoveryComplete();
 
-	/** Query to see whether the given path is currently filtered from the view */
-	virtual bool PathIsFilteredFromViewBySearch(const FString& InPath) const;
-
 	/** Creates a list item for the tree view */
 	virtual TSharedRef<ITableRow> GenerateTreeRow(TSharedPtr<FTreeItem> TreeItem, const TSharedRef<STableViewBase>& OwnerTable);
 
@@ -309,15 +299,14 @@ protected:
 
 	FContentBrowserDataCompiledFilter CreateCompiledFolderFilter() const;
 
-	/** Clear all root items and clear selection */
-	void ClearTreeItems();
-
 	/** Get this path view's editor config if OwningContentBrowserName is set. */
 	FPathViewConfig* GetPathViewConfig() const;
 
 	/** Get this path view's content browser instance config if OwningContentBrowserName is set. */
 	FContentBrowserInstanceConfig* GetContentBrowserConfig() const;
 
+	virtual void ConfigureTreeView(STreeView<TSharedPtr<FTreeItem>>::FArguments& InArgs);
+
 private:
 	/** Selects the given path only if it exists. Returns true if selected. */
 	bool ExplicitlyAddPathToSelection(const FName Path);
@@ -334,15 +323,9 @@ private:
 	/** Handler for when the search box filter has changed */
 	void FilterUpdated();
 
-	/** Populates OutSearchStrings with the strings that should be used in searching */
-	void PopulateFolderSearchStrings( const FString& FolderName, OUT TArray< FString >& OutSearchStrings ) const;
-
 	/** Returns true if the supplied folder item already exists in the tree. If so, ExistingItem will be set to the found item. */
 	bool FolderAlreadyExists(const TSharedPtr< FTreeItem >& TreeItem, TSharedPtr< FTreeItem >& ExistingItem);
 
-	/** Removes the supplied folder from the tree. */
-	void RemoveFolderItem(const TSharedPtr< FTreeItem >& TreeItem);
-
 	/** True if the specified item is expanded in the asset tree */
 	bool IsTreeItemExpanded(TSharedPtr<FTreeItem> TreeItem) const;
 
@@ -367,9 +350,6 @@ private:
 	/** Returns true if filter is being used. */
 	bool IsPluginPathFilterInUse(TSharedRef<FContentBrowserPluginFilter> Filter) const;
 
-	/** Sorts tree items */
-	void DefaultSort(const FTreeItem* InTreeItem, TArray<TSharedPtr<FTreeItem>>& InChildren);
-
 	TArray<FName> GetDefaultPathsToExpand() const;
 
 	/** Tell the tree that the LastExpandedPath set should be refreshed */
@@ -423,21 +403,24 @@ protected:
 	/** The path view search interface */
 	TSharedPtr<FSourcesSearch> SearchPtr;
 
-	/** The list of folders in the tree */
-	TArray< TSharedPtr<FTreeItem> > TreeRootItems;
+	/** Items in the tree and associated data. Shared ptr for easy binding to delegates/attributes */
+	TSharedPtr<FPathViewData> TreeData;
+
+	// Last version number retrieved from TreeData so we can decide if the tree view need rebuilding
+	uint64 LastTreeDataVersion = 0;
 
-	/** The The TextFilter attached to the SearchBox widget */
-	TSharedPtr< FolderTextFilter > SearchBoxFolderFilter;
+	/** Should this path tree be flat like the favorites tree */
+	bool bFlat = false;
 
 	/** The paths that were last reported by OnPathSelected event. Used in preserving selection when filtering folders */
 	TSet<FName> LastSelectedPaths;
 
-	/** If not empty, this is the path of the folders to sync once they are available while assets are still being discovered */
+	/**
+	 * If not empty, this is the path of the folders to sync once they are available while assets are still being discovered.
+	 * It should be emptied when user interaction overides the default selections.
+	 */
 	TArray<FName> PendingInitialPaths;
 
-	/** Delay clear until first pending path is found */
-	bool bPendingInitialPathsNeedsSelectionClear = false;
-
 	/** Context information for the folder item that is currently being created, if any */
 	FContentBrowserItemTemporaryContext PendingNewFolderContext;
 
@@ -449,8 +432,6 @@ protected:
 	/** Writable folder filter */
 	TSharedPtr<FPathPermissionList> WritableFolderPermissionList;
 
-	TMap<FName, TWeakPtr<FTreeItem>> TreeItemLookup;
-
 	/** Custom Folder permissions */
 	TSharedPtr<FPathPermissionList> CustomFolderPermissionList;
 
@@ -494,9 +475,6 @@ private:
 	/** Plugins filters that are currently active */
 	TArray< TSharedRef<FContentBrowserPluginFilter> > AllPluginPathFilters;
 
-	/** Delegate to sort with */
-	FSortTreeItemChildrenDelegate SortOverride;
-
 	/** The favorites path view if one is set. */
 	TSharedPtr<SExpandableArea> FavoritesArea;
 
@@ -512,7 +490,7 @@ private:
 class SFavoritePathView : public SPathView
 {
 public:
-
+	SFavoritePathView();
 	virtual ~SFavoritePathView();
 
 	/** Constructs this widget with InArgs */
@@ -526,9 +504,6 @@ public:
 	/** Loads any settings to config that should be persistent between editor sessions */
 	virtual void LoadSettings(const FString& IniFilename, const FString& IniSection, const FString& SettingsString) override;
 
-	/** Adds nodes to the tree in order to construct the specified item. If bUserNamed is true, the user will name the folder and the item includes the default name. */
-	virtual TSharedPtr<FTreeItem> AddFolderItem(FContentBrowserItemData&& InItem, const bool bUserNamed = false, TArray<TSharedPtr<FTreeItem>>* OutItemsCreated=nullptr) override;
-
 	/** Updates favorites based on an external change. */
 	void FixupFavoritesFromExternalChange(TArrayView<const AssetViewUtils::FMovedContentFolder> MovedFolders);
 
@@ -538,8 +513,7 @@ private:
 	/** Handles updating the view when content items are changed */
 	virtual void HandleItemDataUpdated(TArrayView<const FContentBrowserItemDataUpdate> InUpdatedItems) override;
 
-	/** Query to see whether the given path is currently filtered from the view */
-	virtual bool PathIsFilteredFromViewBySearch(const FString& InPath) const override;
+	virtual void ConfigureTreeView(STreeView<TSharedPtr<FTreeItem>>::FArguments& InArgs) override;
 
 private:
 	TArray<FString> RemovedByFolderMove;
diff --git a/Engine/Source/Editor/ContentBrowser/Private/SourcesViewWidgets.cpp b/Engine/Source/Editor/ContentBrowser/Private/SourcesViewWidgets.cpp
--- a/Engine/Source/Editor/ContentBrowser/Private/SourcesViewWidgets.cpp
+++ b/Engine/Source/Editor/ContentBrowser/Private/SourcesViewWidgets.cpp
@@ -106,7 +106,7 @@ void SAssetTreeItem::Construct( const FArguments& InArgs )
 	{
 		if (ContentBrowserUtils::ShouldShowPluginFolderIcon() && Item.IsInPlugin())
 		{
-			TSharedPtr<FTreeItem> Parent = InArgs._TreeItem->Parent.Pin();
+			TSharedPtr<FTreeItem> Parent = InArgs._TreeItem->GetParent();
 			if (!Parent.IsValid() || !Parent->GetItem().IsInPlugin())
 			{
 				FolderType = EFolderType::PluginRoot;
@@ -126,7 +126,7 @@ void SAssetTreeItem::Construct( const FArguments& InArgs )
 		}
 	}
 
-	bool bIsRoot = !InArgs._TreeItem->Parent.IsValid();
+	bool bIsRoot = !InArgs._TreeItem->GetParent().IsValid();
 
 	ChildSlot
 	[
