Fix copy of TSharedRef to TSharedPtr in loop variable

diff --git a/Engine/Source/Editor/ContentBrowser/Private/SPathView.cpp b/Engine/Source/Editor/ContentBrowser/Private/SPathView.cpp
--- a/Engine/Source/Editor/ContentBrowser/Private/SPathView.cpp
+++ b/Engine/Source/Editor/ContentBrowser/Private/SPathView.cpp
@@ -1579,7 +1579,7 @@ void SPathView::SyncToVirtualPaths(TArrayView<const FName> VirtualPathsToSync, c
 		TreeViewPtr->ClearSelection();
 
 		// SyncTreeItems should now only contain items which aren't already shown explicitly or implicitly (as a child)
-		for (const TSharedPtr<FTreeItem>& Item : SyncTreeItems)
+		for (const TSharedRef<FTreeItem>& Item : SyncTreeItems)
 		{
 			RecursiveExpandParents(Item);
 			TreeViewPtr->SetItemSelection(Item, true);
