Only select exact path matches for favorites view.
Increment tree data version in Populate functions.

diff --git a/Engine/Source/Editor/ContentBrowser/Private/SPathView.cpp b/Engine/Source/Editor/ContentBrowser/Private/SPathView.cpp
--- a/Engine/Source/Editor/ContentBrowser/Private/SPathView.cpp
+++ b/Engine/Source/Editor/ContentBrowser/Private/SPathView.cpp
@@ -407,6 +407,7 @@ void FPathViewData::PopulateFullFolderTree(const FContentBrowserDataCompiledFilt
 			return true;
 		});
 	VisibleRootItems = RootItems;
+	++Version;
 }
 
 void FPathViewData::PopulateWithFavorites(const FContentBrowserDataCompiledFilter& CompiledDataFilter)
@@ -439,6 +440,7 @@ void FPathViewData::PopulateWithFavorites(const FContentBrowserDataCompiledFilte
 				return true;
 			});
 	}
+	++Version;
 }
 
 void FPathViewData::ProcessDataUpdates(TConstArrayView<FContentBrowserItemDataUpdate> InUpdatedItems,
@@ -654,6 +656,11 @@ TSharedPtr<FTreeItem> FPathViewData::FindTreeItem(FName InVirtualPath, bool bVis
 
 TSharedPtr<FTreeItem> FPathViewData::FindBestItemForPath(FStringView InVirtualPath)
 {
+	if (bFlat)
+	{
+		return FindTreeItem(FName(InVirtualPath), false);
+	}
+
 	TSharedPtr<FTreeItem> Found;
 	FPathViews::IterateAncestors(InVirtualPath, [this, &Found](FStringView Ancestor) {
 		FName ItemName{ Ancestor };
