Reset pending initial paths on user-driven selections.

diff --git a/Engine/Source/Editor/ContentBrowser/Private/SPathView.cpp b/Engine/Source/Editor/ContentBrowser/Private/SPathView.cpp
--- a/Engine/Source/Editor/ContentBrowser/Private/SPathView.cpp
+++ b/Engine/Source/Editor/ContentBrowser/Private/SPathView.cpp
@@ -1889,6 +1889,11 @@ void SPathView::SetTreeItemExpansionRecursive( TSharedPtr< FTreeItem > TreeItem,
 
 void SPathView::TreeSelectionChanged( TSharedPtr< FTreeItem > TreeItem, ESelectInfo::Type SelectInfo )
 {
+	if (SelectInfo != ESelectInfo::Direct)
+	{
+		PendingInitialPaths.Reset();
+	}
+
 	if ( ShouldAllowTreeItemChangedDelegate() )
 	{
 		const TArray<TSharedPtr<FTreeItem>> NewSelectedItems = TreeViewPtr->GetSelectedItems();
