[ContentBrowser] Wrong Thumbnail for Folder Drag in Content Browser

diff --git a/Engine/Source/Editor/ContentBrowser/Private/AssetViewWidgets.cpp b/Engine/Source/Editor/ContentBrowser/Private/AssetViewWidgets.cpp
--- a/Engine/Source/Editor/ContentBrowser/Private/AssetViewWidgets.cpp
+++ b/Engine/Source/Editor/ContentBrowser/Private/AssetViewWidgets.cpp
@@ -190,11 +190,10 @@ TSharedRef<SWidget> FAssetViewItemHelper::CreateListTileItemContents(T* const In
 		// TODO: Allow items to customize their widget
 		TSharedPtr<FAssetViewItem>& AssetItem = InTileOrListItem->AssetItem;
 
-		const bool bDeveloperFolder = ContentBrowserUtils::IsItemDeveloperContent(AssetItem->GetItem());
-		const bool bCodeFolder = EnumHasAnyFlags(AssetItem->GetItem().GetItemCategory(), EContentBrowserItemFlags::Category_Class);
-		FContentBrowserItemDataAttributeValue VirtualAttributeValue = AssetItem->GetItem().GetItemAttribute(ContentBrowserItemAttributes::ItemIsCustomVirtualFolder);
-		const bool bVirtualFolder = VirtualAttributeValue.IsValid() && VirtualAttributeValue.GetValue<bool>();
-		const bool bPluginFolder = ContentBrowserUtils::IsItemPluginRootFolder(AssetItem->GetItem());
+		// Default values
+		FName FolderBrushName = TEXT("ContentBrowser.ListViewFolderIcon");
+		FName FolderShadowBrushName = TEXT("ContentBrowser.FolderItem.DropShadow");
+		ContentBrowserUtils::TryGetFolderBrushAndShadowName(AssetItem->GetItem(), FolderBrushName, FolderShadowBrushName);
 		
 		const bool bCollectionFolder = EnumHasAnyFlags(AssetItem->GetItem().GetItemCategory(), EContentBrowserItemFlags::Category_Collection);
 		ECollectionShareType::Type CollectionFolderShareType = ECollectionShareType::CST_All;
@@ -203,29 +202,8 @@ TSharedRef<SWidget> FAssetViewItemHelper::CreateListTileItemContents(T* const In
 			ContentBrowserUtils::IsCollectionPath(AssetItem->GetItem().GetVirtualPath().ToString(), nullptr, &CollectionFolderShareType);
 		}
 
-		const FSlateBrush* FolderBaseImage = nullptr;
-		const FSlateBrush* DropShadowImage = FAppStyle::Get().GetBrush("ContentBrowser.FolderItem.DropShadow");
-		if (bDeveloperFolder)
-		{
-			FolderBaseImage = FAppStyle::GetBrush("ContentBrowser.ListViewDeveloperFolderIcon");
-		}
-		else if (bCodeFolder)
-		{
-			FolderBaseImage = FAppStyle::GetBrush("ContentBrowser.ListViewCodeFolderIcon");
-		}
-		else if (bVirtualFolder && ContentBrowserUtils::ShouldShowCustomVirtualFolderIcon())
-		{
-			FolderBaseImage = FAppStyle::GetBrush("ContentBrowser.ListViewVirtualFolderIcon");
-			DropShadowImage = FAppStyle::GetBrush("None");
-		}
-		else if (bPluginFolder && ContentBrowserUtils::ShouldShowPluginFolderIcon())
-		{
-			FolderBaseImage = FAppStyle::GetBrush("ContentBrowser.ListViewPluginFolderIcon");
-		}
-		else
-		{
-			FolderBaseImage = FAppStyle::GetBrush("ContentBrowser.ListViewFolderIcon");
-		}
+		const FSlateBrush* FolderBaseImage = FAppStyle::GetBrush(FolderBrushName);
+		const FSlateBrush* DropShadowImage = FAppStyle::GetBrush(FolderShadowBrushName);
 
 		// Folder base
 		// clang-format off
diff --git a/Engine/Source/Editor/ContentBrowser/Private/ContentBrowserUtils.cpp b/Engine/Source/Editor/ContentBrowser/Private/ContentBrowserUtils.cpp
--- a/Engine/Source/Editor/ContentBrowser/Private/ContentBrowserUtils.cpp
+++ b/Engine/Source/Editor/ContentBrowser/Private/ContentBrowserUtils.cpp
@@ -459,6 +459,44 @@ bool ContentBrowserUtils::IsItemPluginRootFolder(const FContentBrowserItem& InIt
 	return IsItemPluginContent(InItem);
 }
 
+bool ContentBrowserUtils::TryGetFolderBrushAndShadowName(const FContentBrowserItem& InFolder, FName& OutBrushName, FName& OutShadowBrushName)
+{
+	if (!InFolder.IsValid() || !InFolder.IsFolder())
+	{
+		return false;
+	}
+
+	OutShadowBrushName = TEXT("ContentBrowser.FolderItem.DropShadow");
+	const bool bDeveloperFolder = IsItemDeveloperContent(InFolder);
+	const bool bCodeFolder = EnumHasAnyFlags(InFolder.GetItemCategory(), EContentBrowserItemFlags::Category_Class);
+	const FContentBrowserItemDataAttributeValue VirtualAttributeValue = InFolder.GetItemAttribute(ContentBrowserItemAttributes::ItemIsCustomVirtualFolder);
+	const bool bVirtualFolder = VirtualAttributeValue.IsValid() && VirtualAttributeValue.GetValue<bool>();
+	const bool bPluginFolder = IsItemPluginRootFolder(InFolder);
+
+	if (bDeveloperFolder)
+	{
+		OutBrushName = TEXT("ContentBrowser.ListViewDeveloperFolderIcon");
+	}
+	else if (bCodeFolder)
+	{
+		OutBrushName = TEXT("ContentBrowser.ListViewCodeFolderIcon");
+	}
+	else if (bVirtualFolder && ShouldShowCustomVirtualFolderIcon())
+	{
+		OutBrushName = TEXT("ContentBrowser.ListViewVirtualFolderIcon");
+		OutShadowBrushName = TEXT("None");
+	}
+	else if (bPluginFolder && ShouldShowPluginFolderIcon())
+	{
+		OutBrushName = TEXT("ContentBrowser.ListViewPluginFolderIcon");
+	}
+	else
+	{
+		OutBrushName = TEXT("ContentBrowser.ListViewFolderIcon");
+	}
+	return true;
+}
+
 bool ContentBrowserUtils::IsCollectionPath(const FString& InPath, FName* OutCollectionName, ECollectionShareType::Type* OutCollectionShareType)
 {
 	static const FString CollectionsRootPrefix = TEXT("/Collections");
diff --git a/Engine/Source/Editor/ContentBrowser/Private/ContentBrowserUtils.h b/Engine/Source/Editor/ContentBrowser/Private/ContentBrowserUtils.h
--- a/Engine/Source/Editor/ContentBrowser/Private/ContentBrowserUtils.h
+++ b/Engine/Source/Editor/ContentBrowser/Private/ContentBrowserUtils.h
@@ -75,6 +75,9 @@ namespace ContentBrowserUtils
 	/** Check whether the given item is the root folder of a plugin */
 	bool IsItemPluginRootFolder(const FContentBrowserItem& InItem);
 
+	/** Given a Folder, work out the BrushName and ShadowBrushName to use for it */
+	bool TryGetFolderBrushAndShadowName(const FContentBrowserItem& InFolder, FName& OutBrushName, FName& OutShadowBrushName);
+
 	/** Check to see whether the given path is rooted against a collection directory, optionally extracting the collection name and share type from the path */
 	bool IsCollectionPath(const FString& InPath, FName* OutCollectionName = nullptr, ECollectionShareType::Type* OutCollectionShareType = nullptr);
 
diff --git a/Engine/Source/Editor/ContentBrowser/Private/DragDropHandler.cpp b/Engine/Source/Editor/ContentBrowser/Private/DragDropHandler.cpp
--- a/Engine/Source/Editor/ContentBrowser/Private/DragDropHandler.cpp
+++ b/Engine/Source/Editor/ContentBrowser/Private/DragDropHandler.cpp
@@ -3,6 +3,7 @@
 #include "DragDropHandler.h"
 
 #include "AssetViewUtils.h"
+#include "CollectionViewUtils.h"
 #include "Containers/Array.h"
 #include "Containers/Map.h"
 #include "ContentBrowserDataDragDropOp.h"
@@ -40,6 +41,38 @@
 
 #define LOCTEXT_NAMESPACE "ContentBrowser"
 
+namespace DragDropHandlerUtilsPrivate
+{
+	void RetrieveFolderInformation(const FContentBrowserItem& InFolder, FName& OutFolderBrushName, FName& OutFolderShadowName, FLinearColor& OutFolderOverrideColor)
+	{
+		// We are just dragging folders so add information on the FolderBrush to use and the color
+		if (InFolder.IsValid() && InFolder.IsFolder())
+		{
+			ContentBrowserUtils::TryGetFolderBrushAndShadowName(InFolder, OutFolderBrushName, OutFolderShadowName);
+
+			const bool bCollectionFolder = EnumHasAnyFlags(InFolder.GetItemCategory(), EContentBrowserItemFlags::Category_Collection);
+			if (bCollectionFolder)
+			{
+				FName CollectionName;
+				ECollectionShareType::Type CollectionFolderShareType = ECollectionShareType::CST_All;
+				ContentBrowserUtils::IsCollectionPath(InFolder.GetVirtualPath().ToString(), &CollectionName, &CollectionFolderShareType);
+			
+				if (TOptional<FLinearColor> Color = CollectionViewUtils::GetCustomColor(CollectionName, CollectionFolderShareType))
+				{
+					OutFolderOverrideColor = Color.GetValue();
+				}
+			}
+			else
+			{
+				if (TOptional<FLinearColor> Color = ContentBrowserUtils::GetPathColor(InFolder.GetInvariantPath().ToString()))
+				{
+					OutFolderOverrideColor = Color.GetValue();
+				}
+			}
+		}
+	}
+}
+
 namespace DragDropHandler
 {
 
@@ -75,7 +108,22 @@ TSharedPtr<FDragDropOperation> CreateDragOperation(TArrayView<const FContentBrow
 	}
 
 	// Generic handling
-	return FContentBrowserDataDragDropOp::New(InItems);
+	FName FolderBrushName = NAME_None;
+	FName FolderShadowBrushName = NAME_None;
+	FLinearColor FolderColor = ContentBrowserUtils::GetDefaultColor();
+	// This information will be later used by the ContentBrowserDragAnDropOp in case you are dragging only folders
+	if (InItems.Num() > 0)
+	{
+		DragDropHandlerUtilsPrivate::RetrieveFolderInformation(InItems[0], FolderBrushName, FolderShadowBrushName, FolderColor);
+	}
+
+	FThumbnailOverrideParams FolderThumbnailOverrideParams;
+	FolderThumbnailOverrideParams.FolderBrushName = FolderBrushName;
+	FolderThumbnailOverrideParams.FolderShadowBrushName = FolderShadowBrushName;
+	FolderThumbnailOverrideParams.FolderColorOverride = FolderColor;
+
+	TSharedRef<FContentBrowserDataDragDropOp> ContentBrowserDragAndDrop = FContentBrowserDataDragDropOp::New(InItems, FolderThumbnailOverrideParams);
+	return ContentBrowserDragAndDrop;
 }
 
 template <typename FuncType>
diff --git a/Engine/Source/Editor/ContentBrowserData/Private/ContentBrowserDataDragDropOp.cpp b/Engine/Source/Editor/ContentBrowserData/Private/ContentBrowserDataDragDropOp.cpp
--- a/Engine/Source/Editor/ContentBrowserData/Private/ContentBrowserDataDragDropOp.cpp
+++ b/Engine/Source/Editor/ContentBrowserData/Private/ContentBrowserDataDragDropOp.cpp
@@ -16,11 +16,23 @@
 
 class UActorFactory;
 
-TSharedRef<FContentBrowserDataDragDropOp> FContentBrowserDataDragDropOp::New(TArrayView<const FContentBrowserItem> InDraggedItems)
+TSharedRef<FContentBrowserDataDragDropOp> FContentBrowserDataDragDropOp::New(TArrayView<const FContentBrowserItem> InDraggedItems, FThumbnailOverrideParams InThumbnailOverrideParams)
 {
 	TSharedRef<FContentBrowserDataDragDropOp> Operation = MakeShared<FContentBrowserDataDragDropOp>();
 
 	Operation->Init(MoveTemp(InDraggedItems));
+	Operation->FolderBrushName = InThumbnailOverrideParams.FolderBrushName;
+	Operation->FolderShadowBrushName = InThumbnailOverrideParams.FolderShadowBrushName;
+	Operation->FolderColorOverride = InThumbnailOverrideParams.FolderColorOverride;
+
+	if (Operation->ShouldOverrideThumbnailWidget())
+	{
+		const TSharedPtr<SWidget> CustomThumbnailWidget = Operation->GetFolderWidget();
+		if (CustomThumbnailWidget.IsValid())
+		{
+			Operation->SetCustomThumbnailWidget(CustomThumbnailWidget.ToSharedRef());
+		}
+	}
 
 	Operation->Construct();
 	return Operation;
@@ -163,3 +175,25 @@ FText FContentBrowserDataDragDropOp::GetFirstItemText() const
 
 	return FText::GetEmpty();
 }
+
+bool FContentBrowserDataDragDropOp::ShouldOverrideThumbnailWidget() const
+{
+	return HasFolders() && DraggedFiles.IsEmpty();
+}
+
+TSharedPtr<SWidget> FContentBrowserDataDragDropOp::GetFolderWidget() const
+{
+	const FSlateBrush* FolderBrush = FAppStyle::GetBrush(FolderShadowBrushName);
+	if (FolderBrush != FAppStyle::GetDefaultBrush())
+	{
+		return SNew(SBorder)
+			.BorderImage(FolderBrush)
+			.Padding(FMargin(0,0,2.0f,2.0f))
+			[
+				SNew(SImage)
+				.Image(FAppStyle::GetBrush(FolderBrushName))
+				.ColorAndOpacity(FolderColorOverride)
+			];
+	}
+	return nullptr;
+}
diff --git a/Engine/Source/Editor/ContentBrowserData/Public/ContentBrowserDataDragDropOp.h b/Engine/Source/Editor/ContentBrowserData/Public/ContentBrowserDataDragDropOp.h
--- a/Engine/Source/Editor/ContentBrowserData/Public/ContentBrowserDataDragDropOp.h
+++ b/Engine/Source/Editor/ContentBrowserData/Public/ContentBrowserDataDragDropOp.h
@@ -16,12 +16,29 @@
 class UActorFactory;
 struct FAssetData;
 
+/**
+ * Additional params for FContentBrowserDataDragDropOp to override the thumbnail used if dragging only folders
+ */
+struct CONTENTBROWSERDATA_API FThumbnailOverrideParams
+{
+	FThumbnailOverrideParams() {}
+
+	/** Folder brush name to use */
+	FName FolderBrushName = NAME_None;
+
+	/** Folder shadow brush name to use */
+	FName FolderShadowBrushName = NAME_None;
+
+	/** Override color for the final brush */
+	FLinearColor FolderColorOverride = FLinearColor::Gray;
+};
+
 class CONTENTBROWSERDATA_API FContentBrowserDataDragDropOp : public FAssetDragDropOp
 {
 public:
 	DRAG_DROP_OPERATOR_TYPE(FContentBrowserDataDragDropOp, FAssetDragDropOp)
 
-	static TSharedRef<FContentBrowserDataDragDropOp> New(TArrayView<const FContentBrowserItem> InDraggedItems);
+	static TSharedRef<FContentBrowserDataDragDropOp> New(TArrayView<const FContentBrowserItem> InDraggedItems, FThumbnailOverrideParams InThumbnailOverrideParams = FThumbnailOverrideParams());
 
 	static TSharedRef<FContentBrowserDataDragDropOp> Legacy_New(TArrayView<const FAssetData> InAssetData, TArrayView<const FString> InAssetPaths = TArrayView<const FString>(), UActorFactory* InActorFactory = nullptr);
 
@@ -51,7 +68,16 @@ private:
 	virtual int32 GetTotalCount() const override;
 	virtual FText GetFirstItemText() const override;
 
+	/** Return whether or not it should override the thumbnail widget */
+	bool ShouldOverrideThumbnailWidget() const;
+
+	/** Return the folder widget or nullptr if the BrushName is not valid */
+	TSharedPtr<SWidget> GetFolderWidget() const;
+
 private:
+	FName FolderBrushName;
+	FName FolderShadowBrushName;
+	FLinearColor FolderColorOverride;
 	TArray<FContentBrowserItem> DraggedItems;
 	TArray<FContentBrowserItem> DraggedFiles;
 	TArray<FContentBrowserItem> DraggedFolders;
diff --git a/Engine/Source/Editor/UnrealEd/Private/DragAndDrop/AssetDragDropOp.cpp b/Engine/Source/Editor/UnrealEd/Private/DragAndDrop/AssetDragDropOp.cpp
--- a/Engine/Source/Editor/UnrealEd/Private/DragAndDrop/AssetDragDropOp.cpp
+++ b/Engine/Source/Editor/UnrealEd/Private/DragAndDrop/AssetDragDropOp.cpp
@@ -79,27 +79,19 @@ TSharedPtr<SWidget> FAssetDragDropOp::GetDefaultDecorator() const
 	const int32 TotalCount = GetTotalCount();
 
 	TSharedPtr<SWidget> ThumbnailWidget;
-	if (AssetThumbnail.IsValid())
+	if (CustomThumbnailWidget.IsSet())
+	{
+		ThumbnailWidget = CustomThumbnailWidget.Get();
+	}
+	else if (AssetThumbnail.IsValid())
 	{
 		ThumbnailWidget = AssetThumbnail->MakeThumbnailWidget();
 	}
 	else if (HasFolders())
 	{
-		ThumbnailWidget = 
-			SNew(SOverlay)
-
-			+SOverlay::Slot()
-			[
-				SNew(SImage)
-				.Image(FAppStyle::GetBrush("ContentBrowser.ListViewFolderIcon.Base"))
-				.ColorAndOpacity(FLinearColor::Gray)
-			]
-		
-			+SOverlay::Slot()
-			[
-				SNew(SImage)
-				.Image(FAppStyle::GetBrush("ContentBrowser.ListViewFolderIcon.Mask"))
-			];
+		ThumbnailWidget = SNew(SImage)
+				.Image(FAppStyle::GetBrush("ContentBrowser.ListViewFolderIcon"))
+				.ColorAndOpacity(FLinearColor::Gray);
 	}
 	else
 	{
@@ -295,3 +287,8 @@ FText FAssetDragDropOp::GetFirstItemText() const
 
 	return FText::GetEmpty();
 }
+
+void FAssetDragDropOp::SetCustomThumbnailWidget(const TSharedRef<SWidget>& InNewThumbnailWidget)
+{
+	CustomThumbnailWidget = InNewThumbnailWidget;
+}
diff --git a/Engine/Source/Editor/UnrealEd/Public/DragAndDrop/AssetDragDropOp.h b/Engine/Source/Editor/UnrealEd/Public/DragAndDrop/AssetDragDropOp.h
--- a/Engine/Source/Editor/UnrealEd/Public/DragAndDrop/AssetDragDropOp.h
+++ b/Engine/Source/Editor/UnrealEd/Public/DragAndDrop/AssetDragDropOp.h
@@ -86,6 +86,9 @@ protected:
 	UNREALED_API virtual int32 GetTotalCount() const;
 	UNREALED_API virtual FText GetFirstItemText() const;
 
+	/** Set a custom thumbnail widget to use instead of the default one */
+	UNREALED_API void SetCustomThumbnailWidget(const TSharedRef<SWidget>& InNewThumbnailWidget);
+
 private:
 	/** Data for the assets this item represents */
 	TArray<FAssetData> AssetData;
@@ -93,6 +96,9 @@ private:
 	/** Data for the asset paths this item represents */
 	TArray<FString> AssetPaths;
 
+	/** Custom thumbnail widget */
+	TAttribute<TSharedPtr<SWidget>> CustomThumbnailWidget;
+
 	/** The factory to use if converting this asset to a placed object */
 	TWeakInterfacePtr<IAssetFactoryInterface> AssetFactory;
 
