Fixing bugs with favorites path view.
- Fix use of per-content browser settings in favorites view.
- Fix 'empty folder' filtering to work the same in initial population and incremental updates.
- Correctly populate visible root items after incremental updates when filter is inactive.

diff --git a/Engine/Source/Editor/ContentBrowser/Private/SContentBrowser.cpp b/Engine/Source/Editor/ContentBrowser/Private/SContentBrowser.cpp
--- a/Engine/Source/Editor/ContentBrowser/Private/SContentBrowser.cpp
+++ b/Engine/Source/Editor/ContentBrowser/Private/SContentBrowser.cpp
@@ -1125,6 +1125,7 @@ void SContentBrowser::SetFavoritesExpanded(bool bExpanded)
 
 TSharedRef<SWidget> SContentBrowser::CreateFavoritesView(const FContentBrowserConfig* Config)
 {
+	// clang-format off
 	return
 		SAssignNew(FavoritesArea, SExpandableArea)
 		.BorderImage(FAppStyle::Get().GetBrush("Brushes.Header"))
@@ -1178,10 +1179,12 @@ TSharedRef<SWidget> SContentBrowser::CreateFavoritesView(const FContentBrowserCo
 				.ShowTreeTitle(false)
 				.ShowSeparator(false)
 				.AllowClassesFolder(true)
+				.OwningContentBrowserName(InstanceName)
 				.AddMetaData<FTagMetaData>(FTagMetaData(TEXT("ContentBrowserFavorites")))
 				.ExternalSearch(FavoritesSearch)
 			]
 		];
+	// clang-format on
 }
 
 void SContentBrowser::SetPathViewExpanded(bool bExpanded)
diff --git a/Engine/Source/Editor/ContentBrowser/Private/SPathView.cpp b/Engine/Source/Editor/ContentBrowser/Private/SPathView.cpp
--- a/Engine/Source/Editor/ContentBrowser/Private/SPathView.cpp
+++ b/Engine/Source/Editor/ContentBrowser/Private/SPathView.cpp
@@ -338,6 +338,13 @@ protected:
 
 	bool PassesTextFilter(const TSharedPtr<FTreeItem>& InItem);
 
+	struct FEmptyFolderFilter
+	{
+		TOptional<FContentBrowserFolderContentsFilter> FolderFilter;
+		EContentBrowserIsFolderVisibleFlags FolderFlags;
+	};
+	FEmptyFolderFilter GetEmptyFolderFilter(const FContentBrowserDataCompiledFilter& CompiledDataFilter) const;
+
 	// Incremented to trigger tree rebuild from changes to the tree contents
 	uint64 Version;
 	// Items with no parent
@@ -358,14 +365,8 @@ protected:
 	TTextFilter<FStringView> FolderPathTextFilter;
 };
 
-void FPathViewData::PopulateFullFolderTree(const FContentBrowserDataCompiledFilter& CompiledDataFilter)
+FPathViewData::FEmptyFolderFilter FPathViewData::GetEmptyFolderFilter(const FContentBrowserDataCompiledFilter& CompiledDataFilter) const
 {
-	TMap<FName, TSharedPtr<FTreeItem>> OldItemsByInvariantPath = MoveTemp(InvariantPathToItem);
-	RootItems.Reset();
-	VisibleRootItems.Reset();
-	InvariantPathToItem.Reset();
-	VirtualPathToItem.Reset();
-
 	const UContentBrowserSettings* ContentBrowserSettings = GetDefault<UContentBrowserSettings>();
 	bool bDisplayEmpty = ContentBrowserSettings->DisplayEmptyFolders;
 	// check to see if we have an instance config that overrides the default in UContentBrowserSettings
@@ -383,18 +384,28 @@ void FPathViewData::PopulateFullFolderTree(const FContentBrowserDataCompiledFilt
 		FolderFilter->ItemCategoryFilter = CompiledDataFilter.ItemCategoryFilter;
 	}
 	EContentBrowserIsFolderVisibleFlags FolderFlags = ContentBrowserUtils::GetIsFolderVisibleFlags(bDisplayEmpty);
+	return { FolderFilter, FolderFlags };
+}
 
+void FPathViewData::PopulateFullFolderTree(const FContentBrowserDataCompiledFilter& CompiledDataFilter)
+{
+	TMap<FName, TSharedPtr<FTreeItem>> OldItemsByInvariantPath = MoveTemp(InvariantPathToItem);
+	RootItems.Reset();
+	VisibleRootItems.Reset();
+	InvariantPathToItem.Reset();
+	VirtualPathToItem.Reset();
+
+	UContentBrowserDataSubsystem* ContentBrowserData = IContentBrowserDataModule::Get().GetSubsystem();
+	FEmptyFolderFilter EmptyFilter = GetEmptyFolderFilter(CompiledDataFilter);
 	TArray<TSharedPtr<FTreeItem>> ItemsCreated;
 	ContentBrowserData->EnumerateItemsMatchingFilter(CompiledDataFilter,
 		[this,
-			bDisplayEmpty,
 			CompiledDataFilter,
-			FolderFilter,
-			FolderFlags,
+			EmptyFilter,
 			ContentBrowserData,
 			&OldItemsByInvariantPath](FContentBrowserItemData&& InItemData) {
 			UContentBrowserDataSource* Source = InItemData.GetOwnerDataSource();
-			if (Source && !Source->IsFolderVisible(InItemData.GetVirtualPath(), FolderFlags, FolderFilter))
+			if (Source && !Source->IsFolderVisible(InItemData.GetVirtualPath(), EmptyFilter.FolderFlags, EmptyFilter.FolderFilter))
 			{
 				UE_LOG(LogContentBrowser,
 					VeryVerbose,
@@ -421,6 +432,7 @@ void FPathViewData::PopulateWithFavorites(const FContentBrowserDataCompiledFilte
 
 	const TArray<FString>& FavoritePaths = ContentBrowserUtils::GetFavoriteFolders();
 	UContentBrowserDataSubsystem* ContentBrowserData = IContentBrowserDataModule::Get().GetSubsystem();
+	FEmptyFolderFilter EmptyFilter = GetEmptyFolderFilter(CompiledDataFilter);
 	for (const FString& InvariantPath : FavoritePaths)
 	{
 		FName VirtualPath;
@@ -429,8 +441,16 @@ void FPathViewData::PopulateWithFavorites(const FContentBrowserDataCompiledFilte
 
 		ContentBrowserData->EnumerateItemsAtPath(*Path,
 			CompiledDataFilter.ItemTypeFilter,
-			[this, &CompiledDataFilter, &OldItemsByInvariantPath](FContentBrowserItemData&& InItemData) {
+			[this, &CompiledDataFilter, EmptyFilter, &OldItemsByInvariantPath](FContentBrowserItemData&& InItemData) {
 				UContentBrowserDataSource* ItemDataSource = InItemData.GetOwnerDataSource();
+				if (!ItemDataSource->IsFolderVisible(InItemData.GetVirtualPath(), EmptyFilter.FolderFlags, EmptyFilter.FolderFilter))
+				{
+					UE_LOG(LogContentBrowser,
+						VeryVerbose,
+						TEXT("Hiding folder %s that fails current pre-text filtering"),
+						*WriteToString<256>(InItemData.GetVirtualPath()));
+					return true; // continue enumerating
+				}
 				ItemDataSource->ConvertItemForFilter(InItemData, CompiledDataFilter);
 				if (ItemDataSource->DoesItemPassFilter(InItemData, CompiledDataFilter))
 				{
@@ -447,22 +467,8 @@ void FPathViewData::ProcessDataUpdates(TConstArrayView<FContentBrowserItemDataUp
 	const FContentBrowserDataCompiledFilter& CompiledDataFilter)
 {
 	UContentBrowserDataSubsystem* ContentBrowserData = IContentBrowserDataModule::Get().GetSubsystem();
-	const UContentBrowserSettings* ContentBrowserSettings = GetDefault<UContentBrowserSettings>();
-	bool bDisplayEmpty = ContentBrowserSettings->DisplayEmptyFolders;
-	// check to see if we have an instance config that overrides the default in UContentBrowserSettings
-	if (FContentBrowserInstanceConfig* EditorConfig =
-			ContentBrowserUtils::GetContentBrowserConfig(OwningContentBrowserName))
-	{
-		bDisplayEmpty = EditorConfig->bShowEmptyFolders;
-	}
-	TOptional<FContentBrowserFolderContentsFilter> FolderFilter;
-	if (!bDisplayEmpty)
-	{
-		FolderFilter = FContentBrowserFolderContentsFilter{};
-		FolderFilter->ItemCategoryFilter = CompiledDataFilter.ItemCategoryFilter;
-	}
-	EContentBrowserIsFolderVisibleFlags FolderFlags = ContentBrowserUtils::GetIsFolderVisibleFlags(bDisplayEmpty);
-	auto DoesItemPassFilter = [this, FolderFlags, FolderFilter, ContentBrowserData, &CompiledDataFilter](
+	FEmptyFolderFilter EmptyFilter = GetEmptyFolderFilter(CompiledDataFilter);
+	auto DoesItemPassFilter = [this, EmptyFilter, ContentBrowserData, &CompiledDataFilter](
 								  const FContentBrowserItemData& InItemData) {
 		UContentBrowserDataSource* ItemDataSource = InItemData.GetOwnerDataSource();
 		if (!ItemDataSource->DoesItemPassFilter(InItemData, CompiledDataFilter))
@@ -470,7 +476,7 @@ void FPathViewData::ProcessDataUpdates(TConstArrayView<FContentBrowserItemDataUp
 			return false;
 		}
 
-		if (!ContentBrowserData->IsFolderVisible(InItemData.GetVirtualPath(), FolderFlags, FolderFilter))
+		if (!ContentBrowserData->IsFolderVisible(InItemData.GetVirtualPath(), EmptyFilter.FolderFlags, EmptyFilter.FolderFilter))
 		{
 			UE_LOG(LogContentBrowser,
 				VeryVerbose,
@@ -588,6 +594,11 @@ void FPathViewData::ProcessDataUpdates(TConstArrayView<FContentBrowserItemDataUp
 			}
 		}
 	}
+	else 
+	{
+		// If filtering is not active and we created some new root items, we need them to be visible
+		VisibleRootItems = RootItems;
+	}
 }
 
 bool FPathViewData::PassesTextFilter(const TSharedPtr<FTreeItem>& InItem)
@@ -2409,7 +2420,7 @@ void SFavoritePathView::Construct(const FArguments& InArgs)
 	// Bind the favorites menu to update after folder changes
 	AssetViewUtils::OnFolderPathChanged().AddSP(this, &SFavoritePathView::FixupFavoritesFromExternalChange); 
 
-	OnFavoritesChangedHandle = FContentBrowserSingleton::Get().RegisterOnFavoritesChangedHandler(FSimpleDelegate::CreateSP(this, &SFavoritePathView::Populate, false));
+	OnFavoritesChangedHandle = FContentBrowserSingleton::Get().RegisterOnFavoritesChangedHandler(FSimpleDelegate::CreateSP(this, &SFavoritePathView::OnFavoriteAdded));
 
 	SPathView::Construct(InArgs);
 }
@@ -2462,6 +2473,7 @@ void SFavoritePathView::SaveSettings(const FString& IniFilename, const FString&
 
 void SFavoritePathView::LoadSettings(const FString& IniFilename, const FString& IniSection, const FString& SettingsString)
 {
+	TGuardValue<bool> Guard(bIsLoadingSettings, true);
 	SPathView::LoadSettings(IniFilename, IniSection, SettingsString);
 
 	// We clear the initial selection for the favorite view, as it conflicts with the main paths view and results in a phantomly selected favorite item
@@ -2520,6 +2532,14 @@ TSharedRef<ITableRow> SFavoritePathView::GenerateTreeRow(TSharedPtr<FTreeItem> T
 		];
 }
 
+void SFavoritePathView::OnFavoriteAdded()
+{
+	if (!bIsLoadingSettings)
+	{
+		Populate();
+	}
+}
+
 void SFavoritePathView::HandleItemDataUpdated(TArrayView<const FContentBrowserItemDataUpdate> InUpdatedItems)
 {
 	if (InUpdatedItems.Num() == 0)
diff --git a/Engine/Source/Editor/ContentBrowser/Private/SPathView.h b/Engine/Source/Editor/ContentBrowser/Private/SPathView.h
--- a/Engine/Source/Editor/ContentBrowser/Private/SPathView.h
+++ b/Engine/Source/Editor/ContentBrowser/Private/SPathView.h
@@ -510,6 +510,8 @@ public:
 private:
 	virtual TSharedRef<ITableRow> GenerateTreeRow(TSharedPtr<FTreeItem> TreeItem, const TSharedRef<STableViewBase>& OwnerTable) override;
 
+	void OnFavoriteAdded();
+
 	/** Handles updating the view when content items are changed */
 	virtual void HandleItemDataUpdated(TArrayView<const FContentBrowserItemDataUpdate> InUpdatedItems) override;
 
@@ -518,4 +520,5 @@ private:
 private:
 	TArray<FString> RemovedByFolderMove;
 	FDelegateHandle OnFavoritesChangedHandle;
+	bool bIsLoadingSettings = false;
 };
