Explicitly deselect renamed item when committing new name. Fixes broken context menu entry when creating a new folder called NewFolder.

diff --git a/Engine/Source/Editor/ContentBrowser/Private/SPathView.cpp b/Engine/Source/Editor/ContentBrowser/Private/SPathView.cpp
--- a/Engine/Source/Editor/ContentBrowser/Private/SPathView.cpp
+++ b/Engine/Source/Editor/ContentBrowser/Private/SPathView.cpp
@@ -2105,6 +2105,8 @@ void SPathView::FolderNameChanged(const TSharedPtr<FTreeItem>& TreeItem,
 	bool bSuccess = false;
 	FText ErrorMessage;
 
+	// Group the deselect and reselect into a single operation
+	FScopedSelectionChangedEvent ScopedSelectionChangedEvent(SharedThis(this));
 	FContentBrowserItem NewItem;
 	if (PendingNewFolderContext.IsValid())
 	{
@@ -2112,6 +2114,7 @@ void SPathView::FolderNameChanged(const TSharedPtr<FTreeItem>& TreeItem,
 
 		// Remove the temporary item before we do any work to ensure the new item creation is not prevented
 		TreeData->RemoveFolderItem(TreeItem.ToSharedRef());
+		TreeViewPtr->SetItemSelection(TreeItem.ToSharedRef(), false);
 
 		// Clearing the rename box on a newly created item cancels the entire creation process
 		if (CommitType == ETextCommit::OnCleared)
