Update FFrontendFilter_ShowOtherDevelopers to avoid string allocations

diff --git a/Engine/Source/Editor/ContentBrowser/Private/FrontendFilters.cpp b/Engine/Source/Editor/ContentBrowser/Private/FrontendFilters.cpp
--- a/Engine/Source/Editor/ContentBrowser/Private/FrontendFilters.cpp
+++ b/Engine/Source/Editor/ContentBrowser/Private/FrontendFilters.cpp
@@ -1118,32 +1118,32 @@ void FFrontendFilter_ShowOtherDevelopers::SetCurrentFilter(TArrayView<const FNam
 bool FFrontendFilter_ShowOtherDevelopers::PassesFilter(FAssetFilterType InItem) const
 {
 	// Pass all assets if other developer assets are allowed
-	if ( !bShowOtherDeveloperAssets )
+	if ( bShowOtherDeveloperAssets )
 	{
-		// Never hide developer assets when a single developer folder is selected.
-		if ( !bIsOnlyOneDeveloperPathSelected )
-		{
-			// TODO: Have attribute flags for this so you can tell from the item whether it's a developer folder, and also whether it's yours
-			// If selecting multiple folders, the Developers folder/parent folder, or "All Assets", hide assets which are found in the development folder unless they are in the current user's folder
-			bool bPackageInDeveloperFolder = !TextFilterUtils::NameStrincmp(InItem.GetInternalPath(), BaseDeveloperPath, BaseDeveloperPathAnsi, BaseDeveloperPath.Len());
-			if ( bPackageInDeveloperFolder )
-			{
-				// Test again using only the path part to avoid filtering files directly in the Developers folder
-				// This happens after the above check to avoid string manipulation when not required
-				FString PackagePath = FPaths::GetPath(InItem.GetInternalPath().ToString());
-				bPackageInDeveloperFolder = PackagePath.StartsWith(BaseDeveloperPath);
-				if ( bPackageInDeveloperFolder )
-				{
-					PackagePath /= FString();
+		return true;
+	}
+	// Never hide developer assets when a single developer folder is selected.
+	if ( bIsOnlyOneDeveloperPathSelected )
+	{
+		return true;
+	}
 
-					const bool bPackageInUserDeveloperFolder = PackagePath.StartsWith(UserDeveloperPath);
-					if ( !bPackageInUserDeveloperFolder )
-					{
-						return false;
-					}
-				}
-			}
+	// TODO: Have attribute flags for this so you can tell from the item whether it's a developer folder, and also whether it's yours
+	// If selecting multiple folders, the Developers folder/parent folder, or "All Assets", hide assets which are found in the development folder unless they are in the current user's folder
+	TStringBuilder<256> InternalPath{InPlace, InItem.GetInternalPath()};
+	const bool bPackageInDeveloperFolder = InternalPath.ToView().StartsWith(BaseDeveloperPath, ESearchCase::IgnoreCase);
+	const bool bPackageInUserDeveloperFolder = InternalPath.ToView().StartsWith(UserDeveloperPath, ESearchCase::IgnoreCase);
+	if (bPackageInDeveloperFolder && !bPackageInUserDeveloperFolder)
+	{
+		// Test again using only the path part to avoid filtering files directly in the Developers folder
+		// This happens after the above check to avoid string manipulation when not required
+		FStringView ParentPath = FPathViews::GetPath(InternalPath.ToView());
+		const bool bIsDirectlyInDevleoperFolder = ParentPath.Equals(FStringView(BaseDeveloperPath).LeftChop(1), ESearchCase::IgnoreCase);
+		if(bIsDirectlyInDevleoperFolder )
+		{
+			return true;
 		}
+		return false;
 	}
 
 	return true;
