[ContentBrowser] Impossible to save asset in Engine / Plugins / Developers folder in the asset picker dialog

diff --git a/Engine/Source/Editor/ContentBrowser/Private/ContentBrowserMenuUtils.cpp b/Engine/Source/Editor/ContentBrowser/Private/ContentBrowserMenuUtils.cpp
new file mode 100644
--- /dev/null
+++ b/Engine/Source/Editor/ContentBrowser/Private/ContentBrowserMenuUtils.cpp
@@ -0,0 +1,279 @@
+// Copyright Epic Games, Inc. All Rights Reserved.
+
+#include "ContentBrowserMenuUtils.h"
+#include "ContentBrowserConfig.h"
+#include "Settings/ContentBrowserSettings.h"
+#include "ToolMenu.h"
+
+#define LOCTEXT_NAMESPACE "ContentBrowserMenuUtils"
+
+namespace ContentBrowserMenuUtils::Filters
+{
+	FContentBrowserInstanceConfig* GetContentBrowserConfig(FName InOwningContentBrowserName)
+	{
+		if (InOwningContentBrowserName.IsNone())
+		{
+			return nullptr;
+		}
+
+		if (UContentBrowserConfig* EditorConfig = UContentBrowserConfig::Get())
+		{
+			return EditorConfig->Instances.Find(InOwningContentBrowserName);
+		}
+
+		return nullptr;
+	}
+
+	bool IsToggleAllowed(FCanExecuteAction InCanExecuteAction)
+	{
+		if (InCanExecuteAction.IsBound())
+		{
+			return InCanExecuteAction.Execute();
+		}
+
+		return true;
+	}
+
+	bool IsShowingCppContent(FName InOwningContentBrowserName, FCanExecuteAction InCanExecuteAction)
+	{
+		if (!IsToggleAllowed(InCanExecuteAction))
+		{
+			return false;
+		}
+
+		if (const FContentBrowserInstanceConfig* Config = GetContentBrowserConfig(InOwningContentBrowserName))
+		{
+			return Config->bShowCppFolders;
+		}
+	
+		return GetDefault<UContentBrowserSettings>()->GetDisplayCppFolders();
+	}
+
+	/** Toggle whether localized content should be shown or not */
+	void ToggleShowCppContent(FName InOwningContentBrowserName, FCanExecuteAction InCanExecuteAction)
+	{
+		check(IsToggleAllowed(InCanExecuteAction));
+
+		bool bNewState = !GetDefault<UContentBrowserSettings>()->GetDisplayCppFolders();
+
+		if (FContentBrowserInstanceConfig* Config = GetContentBrowserConfig(InOwningContentBrowserName))
+		{
+			bNewState = !Config->bShowCppFolders;
+			Config->bShowCppFolders = bNewState;
+			UContentBrowserConfig::Get()->SaveEditorConfig();
+		}
+
+		GetMutableDefault<UContentBrowserSettings>()->SetDisplayCppFolders(bNewState);
+		GetMutableDefault<UContentBrowserSettings>()->PostEditChange();
+	}
+
+	bool IsShowingDevelopersContent(FName InOwningContentBrowserName, FCanExecuteAction InCanExecuteAction)
+	{
+		if (!IsToggleAllowed(InCanExecuteAction))
+		{
+			return false;
+		}
+
+		if (const FContentBrowserInstanceConfig* Config = GetContentBrowserConfig(InOwningContentBrowserName))
+		{
+			return Config->bShowDeveloperContent;
+		}
+
+		return GetDefault<UContentBrowserSettings>()->GetDisplayDevelopersFolder();
+	}
+
+	void ToggleShowDevelopersContent(FName InOwningContentBrowserName, FCanExecuteAction InCanExecuteAction)
+	{
+		check(IsToggleAllowed(InCanExecuteAction));
+
+		bool bNewState = !GetDefault<UContentBrowserSettings>()->GetDisplayDevelopersFolder();
+
+		if (FContentBrowserInstanceConfig* EditorConfig = GetContentBrowserConfig(InOwningContentBrowserName))
+		{
+			bNewState = !EditorConfig->bShowDeveloperContent;
+			EditorConfig->bShowDeveloperContent = bNewState;
+			UContentBrowserConfig::Get()->SaveEditorConfig();
+		}
+
+		GetMutableDefault<UContentBrowserSettings>()->SetDisplayDevelopersFolder(bNewState);
+		GetMutableDefault<UContentBrowserSettings>()->PostEditChange();
+	}
+	
+	bool IsShowingEngineContent(FName InOwningContentBrowserName, FCanExecuteAction InCanExecuteAction)
+	{
+		if (!IsToggleAllowed(InCanExecuteAction))
+		{
+			// Engine toggle is allowed to change if not forced to always be shown
+			// Return true in case the IsToggleAllowed for it is false since it means that it is shown
+			return true;
+		}
+
+		if (const FContentBrowserInstanceConfig* Config = GetContentBrowserConfig(InOwningContentBrowserName))
+		{
+			return Config->bShowEngineContent;
+		}
+
+		return GetDefault<UContentBrowserSettings>()->GetDisplayEngineFolder();
+	}
+
+	void ToggleShowEngineContent(FName InOwningContentBrowserName, FCanExecuteAction InCanExecuteAction)
+	{
+		check(IsToggleAllowed(InCanExecuteAction));
+
+		bool bNewState = !GetDefault<UContentBrowserSettings>()->GetDisplayEngineFolder();
+
+		if (FContentBrowserInstanceConfig* EditorConfig = GetContentBrowserConfig(InOwningContentBrowserName))
+		{
+			bNewState = !EditorConfig->bShowEngineContent;
+			EditorConfig->bShowEngineContent = bNewState;
+			UContentBrowserConfig::Get()->SaveEditorConfig();
+		}
+
+		GetMutableDefault<UContentBrowserSettings>()->SetDisplayEngineFolder(bNewState);
+		GetMutableDefault<UContentBrowserSettings>()->PostEditChange();
+	}
+
+	bool IsShowingPluginContent(FName InOwningContentBrowserName, FCanExecuteAction InCanExecuteAction)
+	{
+		if (!IsToggleAllowed(InCanExecuteAction))
+		{
+			// Plugin toggle is allowed to change if not forced to always be shown
+			// Return true in case the IsToggleAllowed for it is false since it means that it is shown
+			return true;
+		}
+
+		if (const FContentBrowserInstanceConfig* Config = GetContentBrowserConfig(InOwningContentBrowserName))
+		{
+			return Config->bShowPluginContent;
+		}
+
+		return GetDefault<UContentBrowserSettings>()->GetDisplayPluginFolders();
+	}
+
+	void ToggleShowPluginContent(FName InOwningContentBrowserName, FCanExecuteAction InCanExecuteAction)
+	{
+		check(IsToggleAllowed(InCanExecuteAction));
+
+		bool bNewState = !GetDefault<UContentBrowserSettings>()->GetDisplayPluginFolders();
+
+		if (FContentBrowserInstanceConfig* EditorConfig = GetContentBrowserConfig(InOwningContentBrowserName))
+		{
+			bNewState = !EditorConfig->bShowPluginContent;
+			EditorConfig->bShowPluginContent = bNewState;
+			UContentBrowserConfig::Get()->SaveEditorConfig();
+		}
+
+		GetMutableDefault<UContentBrowserSettings>()->SetDisplayPluginFolders(bNewState);
+		GetMutableDefault<UContentBrowserSettings>()->PostEditChange();
+	}
+	
+	bool IsShowingLocalizedContent(FName InOwningContentBrowserName, FCanExecuteAction InCanExecuteAction)
+	{
+		if (!IsToggleAllowed(InCanExecuteAction))
+		{
+			return false;
+		}
+
+		if (const FContentBrowserInstanceConfig* Config = GetContentBrowserConfig(InOwningContentBrowserName))
+		{
+			return Config->bShowLocalizedContent;
+		}
+
+		return GetDefault<UContentBrowserSettings>()->GetDisplayL10NFolder();
+	}
+
+	void ToggleShowLocalizedContent(FName InOwningContentBrowserName, FCanExecuteAction InCanExecuteAction)
+	{
+		check(IsToggleAllowed(InCanExecuteAction));
+
+		bool bNewState = !GetDefault<UContentBrowserSettings>()->GetDisplayL10NFolder();
+
+		if (FContentBrowserInstanceConfig* Config = GetContentBrowserConfig(InOwningContentBrowserName))
+		{
+			bNewState = !Config->bShowLocalizedContent;
+			Config->bShowLocalizedContent = bNewState;
+			UContentBrowserConfig::Get()->SaveEditorConfig();
+		}
+
+		GetMutableDefault<UContentBrowserSettings>()->SetDisplayL10NFolder(bNewState);
+		GetMutableDefault<UContentBrowserSettings>()->PostEditChange();
+	}
+}
+
+void ContentBrowserMenuUtils::AddFiltersToMenu(UToolMenu* InMenu, const FName& InOwningContentBrowserName, FFiltersAdditionalParams InFiltersAdditionalParams)
+{
+	FToolMenuInsert ContentMenuInsert = FToolMenuInsert();
+	if (InMenu->ContainsSection("View"))
+	{
+		ContentMenuInsert.Name = TEXT("View");
+		ContentMenuInsert.Position = EToolMenuInsertType::After;
+	}
+
+	FToolMenuSection& Section = InMenu->FindOrAddSection("Content", LOCTEXT("ContentHeading", "Content"), ContentMenuInsert);
+	Section.AddMenuEntry(
+		"ShowCppClasses",
+		LOCTEXT("ShowCppClassesOption", "Show C++ Classes"),
+		LOCTEXT("ShowCppClassesOptionToolTip", "Show C++ classes in the view?"),
+		FSlateIcon(),
+		FUIAction(
+			FExecuteAction::CreateStatic(&Filters::ToggleShowCppContent, InOwningContentBrowserName, InFiltersAdditionalParams.CanShowCPPClasses),
+			FCanExecuteAction::CreateStatic(&Filters::IsToggleAllowed, InFiltersAdditionalParams.CanShowCPPClasses),
+			FIsActionChecked::CreateStatic(&Filters::IsShowingCppContent, InOwningContentBrowserName, InFiltersAdditionalParams.CanShowCPPClasses)
+		),
+		EUserInterfaceActionType::ToggleButton
+	);
+
+	Section.AddMenuEntry(
+		"ShowDevelopersContent",
+		LOCTEXT("ShowDevelopersContentOption", "Show Developers Content"),
+		LOCTEXT("ShowDevelopersContentOptionToolTip", "Show developers content in the view?"),
+		FSlateIcon(),
+		FUIAction(
+			FExecuteAction::CreateStatic(&Filters::ToggleShowDevelopersContent, InOwningContentBrowserName, InFiltersAdditionalParams.CanShowDevelopersContent),
+			FCanExecuteAction::CreateStatic(&Filters::IsToggleAllowed, InFiltersAdditionalParams.CanShowDevelopersContent),
+			FIsActionChecked::CreateStatic(&Filters::IsShowingDevelopersContent, InOwningContentBrowserName, InFiltersAdditionalParams.CanShowDevelopersContent)
+			),
+		EUserInterfaceActionType::ToggleButton
+	);
+
+	Section.AddMenuEntry(
+		"ShowEngineFolder",
+		LOCTEXT("ShowEngineFolderOption", "Show Engine Content"),
+		LOCTEXT("ShowEngineFolderOptionToolTip", "Show engine content in the view?"),
+		FSlateIcon(),
+		FUIAction(
+			FExecuteAction::CreateStatic(&Filters::ToggleShowEngineContent, InOwningContentBrowserName, InFiltersAdditionalParams.CanShowEngineFolder),
+			FCanExecuteAction::CreateStatic(&Filters::IsToggleAllowed, InFiltersAdditionalParams.CanShowEngineFolder),
+			FIsActionChecked::CreateStatic(&Filters::IsShowingEngineContent, InOwningContentBrowserName, InFiltersAdditionalParams.CanShowEngineFolder)
+		),
+		EUserInterfaceActionType::ToggleButton
+	);
+
+	Section.AddMenuEntry(
+		"ShowPluginFolder",
+		LOCTEXT("ShowPluginFolderOption", "Show Plugin Content"),
+		LOCTEXT("ShowPluginFolderOptionToolTip", "Show plugin content in the view?"),
+		FSlateIcon(),
+		FUIAction(
+			FExecuteAction::CreateStatic(&Filters::ToggleShowPluginContent, InOwningContentBrowserName, InFiltersAdditionalParams.CanShowPluginFolder),
+			FCanExecuteAction::CreateStatic(&Filters::IsToggleAllowed, InFiltersAdditionalParams.CanShowPluginFolder),
+			FIsActionChecked::CreateStatic(&Filters::IsShowingPluginContent, InOwningContentBrowserName, InFiltersAdditionalParams.CanShowPluginFolder)
+		),
+		EUserInterfaceActionType::ToggleButton
+	);
+
+	Section.AddMenuEntry(
+		"ShowLocalizedContent",
+		LOCTEXT("ShowLocalizedContentOption", "Show Localized Content"),
+		LOCTEXT("ShowLocalizedContentOptionToolTip", "Show localized content in the view?"),
+		FSlateIcon(),
+		FUIAction(
+			FExecuteAction::CreateStatic(&Filters::ToggleShowLocalizedContent, InOwningContentBrowserName, InFiltersAdditionalParams.CanShowLocalizedContent),
+			FCanExecuteAction::CreateStatic(&Filters::IsToggleAllowed, InFiltersAdditionalParams.CanShowLocalizedContent),
+			FIsActionChecked::CreateStatic(&Filters::IsShowingLocalizedContent, InOwningContentBrowserName, InFiltersAdditionalParams.CanShowLocalizedContent)
+			),
+		EUserInterfaceActionType::ToggleButton
+	);
+}
+
+#undef LOCTEXT_NAMESPACE
diff --git a/Engine/Source/Editor/ContentBrowser/Private/ContentBrowserMenuUtils.h b/Engine/Source/Editor/ContentBrowser/Private/ContentBrowserMenuUtils.h
new file mode 100644
--- /dev/null
+++ b/Engine/Source/Editor/ContentBrowser/Private/ContentBrowserMenuUtils.h
@@ -0,0 +1,42 @@
+// Copyright Epic Games, Inc. All Rights Reserved.
+
+#pragma once
+
+#include "Framework/Commands/UIAction.h"
+#include "Templates/SharedPointer.h"
+
+class FName;
+class SContentBrowser;
+class UToolMenu;
+
+/**
+ * Additional parameter for the content filters
+ */
+struct FFiltersAdditionalParams
+{
+	FFiltersAdditionalParams()
+	{
+		CanShowCPPClasses = FCanExecuteAction();
+		CanShowDevelopersContent = FCanExecuteAction();
+		CanShowEngineFolder = FCanExecuteAction();
+		CanShowPluginFolder = FCanExecuteAction();
+		CanShowLocalizedContent = FCanExecuteAction();
+	}
+
+	FCanExecuteAction CanShowCPPClasses;
+	FCanExecuteAction CanShowDevelopersContent;
+	FCanExecuteAction CanShowEngineFolder;
+	FCanExecuteAction CanShowPluginFolder;
+	FCanExecuteAction CanShowLocalizedContent;
+};
+
+namespace ContentBrowserMenuUtils
+{
+	/**
+	 * Add content browser filters to the given InMenu
+	 * @param InMenu Menu to add filter to
+	 * @param InOwningContentBrowserName ContentBrowser owner name
+	 * @param InFiltersAdditionalParams Additional filters params
+	 */
+	void AddFiltersToMenu(UToolMenu* InMenu, const FName& InOwningContentBrowserName, FFiltersAdditionalParams InFiltersAdditionalParams = FFiltersAdditionalParams());
+}
diff --git a/Engine/Source/Editor/ContentBrowser/Private/ContentBrowserPathViewMenuContexts.h b/Engine/Source/Editor/ContentBrowser/Private/ContentBrowserPathViewMenuContexts.h
new file mode 100644
--- /dev/null
+++ b/Engine/Source/Editor/ContentBrowser/Private/ContentBrowserPathViewMenuContexts.h
@@ -0,0 +1,24 @@
+// Copyright Epic Games, Inc. All Rights Reserved.
+
+#pragma once
+
+#include "Templates/SharedPointer.h"
+#include "UObject/NameTypes.h"
+#include "UObject/Object.h"
+
+#include "ContentBrowserPathViewMenuContexts.generated.h"
+
+class SPathView;
+
+/**
+ * Context for the PathView setting combo button
+ */
+UCLASS()
+class UContentBrowserPathViewContextMenuContext : public UObject
+{
+	GENERATED_BODY()
+
+public:
+	FName OwningContentBrowserName;
+	TWeakPtr<SPathView> PathView;
+};
diff --git a/Engine/Source/Editor/ContentBrowser/Private/SAssetView.cpp b/Engine/Source/Editor/ContentBrowser/Private/SAssetView.cpp
--- a/Engine/Source/Editor/ContentBrowser/Private/SAssetView.cpp
+++ b/Engine/Source/Editor/ContentBrowser/Private/SAssetView.cpp
@@ -16,6 +16,7 @@
 #include "ContentBrowserDataSubsystem.h"
 #include "ContentBrowserLog.h"
 #include "ContentBrowserMenuContexts.h"
+#include "ContentBrowserMenuUtils.h"
 #include "ContentBrowserModule.h"
 #include "ContentBrowserSingleton.h"
 #include "ContentBrowserUtils.h"
@@ -2630,13 +2631,24 @@ void SAssetView::RegisterGetViewButtonMenu()
 		Menu->bCloseSelfOnly = true;
 		Menu->AddDynamicSection("DynamicContent", FNewToolMenuDelegate::CreateLambda([](UToolMenu* InMenu)
 		{
+			FName OwningContentBrowserName = NAME_None;
+			FFiltersAdditionalParams FilterParams;
 			if (UContentBrowserAssetViewContextMenuContext* Context = InMenu->FindContext<UContentBrowserAssetViewContextMenuContext>())
 			{
 				if (Context->AssetView.IsValid())
 				{
-					Context->AssetView.Pin()->PopulateViewButtonMenu(InMenu);
+					const TSharedPtr<SAssetView> AssetView = Context->AssetView.Pin();
+					AssetView->PopulateViewButtonMenu(InMenu);
+					AssetView->PopulateFilterAdditionalParams(FilterParams);
+				}
+
+				if (Context->OwningContentBrowser.IsValid())
+				{
+					OwningContentBrowserName = Context->OwningContentBrowser.Pin()->GetInstanceName();
 				}
 			}
+
+			ContentBrowserMenuUtils::AddFiltersToMenu(InMenu, OwningContentBrowserName, FilterParams);
 		}));
 	}
 }
@@ -2817,74 +2829,6 @@ void SAssetView::PopulateViewButtonMenu(UToolMenu* Menu)
 		}
 	}
 
-	{
-		FToolMenuSection& Section = Menu->AddSection("Content", LOCTEXT("ContentHeading", "Content"));
-		Section.AddMenuEntry(
-			"ShowCppClasses",
-			LOCTEXT("ShowCppClassesOption", "Show C++ Classes"),
-			LOCTEXT("ShowCppClassesOptionToolTip", "Show C++ classes in the view?"),
-			FSlateIcon(),
-			FUIAction(
-				FExecuteAction::CreateSP( this, &SAssetView::ToggleShowCppContent ),
-				FCanExecuteAction::CreateSP( this, &SAssetView::IsToggleShowCppContentAllowed ),
-				FIsActionChecked::CreateSP( this, &SAssetView::IsShowingCppContent )
-			),
-			EUserInterfaceActionType::ToggleButton
-		);
-
-		Section.AddMenuEntry(
-			"ShowDevelopersContent",
-			LOCTEXT("ShowDevelopersContentOption", "Show Developers Content"),
-			LOCTEXT("ShowDevelopersContentOptionToolTip", "Show developers content in the view?"),
-			FSlateIcon(),
-			FUIAction(
-				FExecuteAction::CreateSP( this, &SAssetView::ToggleShowDevelopersContent ),
-				FCanExecuteAction::CreateSP( this, &SAssetView::IsToggleShowDevelopersContentAllowed ),
-				FIsActionChecked::CreateSP( this, &SAssetView::IsShowingDevelopersContent )
-				),
-			EUserInterfaceActionType::ToggleButton
-		);
-
-		Section.AddMenuEntry(
-			"ShowEngineFolder",
-			LOCTEXT("ShowEngineFolderOption", "Show Engine Content"),
-			LOCTEXT("ShowEngineFolderOptionToolTip", "Show engine content in the view?"),
-			FSlateIcon(),
-			FUIAction(
-				FExecuteAction::CreateSP( this, &SAssetView::ToggleShowEngineContent ),
-				FCanExecuteAction::CreateSP( this, &SAssetView::IsToggleShowEngineContentAllowed),
-				FIsActionChecked::CreateSP( this, &SAssetView::IsShowingEngineContent )
-			),
-			EUserInterfaceActionType::ToggleButton
-		);
-
-		Section.AddMenuEntry(
-			"ShowPluginFolder",
-			LOCTEXT("ShowPluginFolderOption", "Show Plugin Content"),
-			LOCTEXT("ShowPluginFolderOptionToolTip", "Show plugin content in the view?"),
-			FSlateIcon(),
-			FUIAction(
-				FExecuteAction::CreateSP( this, &SAssetView::ToggleShowPluginContent ),
-				FCanExecuteAction::CreateSP(this, &SAssetView::IsToggleShowPluginContentAllowed),
-				FIsActionChecked::CreateSP( this, &SAssetView::IsShowingPluginContent )
-			),
-			EUserInterfaceActionType::ToggleButton
-		);
-
-		Section.AddMenuEntry(
-			"ShowLocalizedContent",
-			LOCTEXT("ShowLocalizedContentOption", "Show Localized Content"),
-			LOCTEXT("ShowLocalizedContentOptionToolTip", "Show localized content in the view?"),
-			FSlateIcon(),
-			FUIAction(
-				FExecuteAction::CreateSP(this, &SAssetView::ToggleShowLocalizedContent),
-				FCanExecuteAction::CreateSP(this, &SAssetView::IsToggleShowLocalizedContentAllowed),
-				FIsActionChecked::CreateSP(this, &SAssetView::IsShowingLocalizedContent)
-				),
-			EUserInterfaceActionType::ToggleButton
-			);
-	}
-
 	{
 		FToolMenuSection& Section = Menu->AddSection("Search", LOCTEXT("SearchHeading", "Search"));
 		Section.AddMenuEntry(
@@ -3010,6 +2954,15 @@ void SAssetView::PopulateViewButtonMenu(UToolMenu* Menu)
 	}
 }
 
+void SAssetView::PopulateFilterAdditionalParams(FFiltersAdditionalParams& OutParams)
+{
+	OutParams.CanShowCPPClasses = FCanExecuteAction::CreateSP(this, &SAssetView::IsToggleShowCppContentAllowed);
+	OutParams.CanShowDevelopersContent = FCanExecuteAction::CreateSP(this, &SAssetView::IsToggleShowDevelopersContentAllowed);
+	OutParams.CanShowEngineFolder = FCanExecuteAction::CreateSP(this, &SAssetView::IsToggleShowEngineContentAllowed);
+	OutParams.CanShowPluginFolder = FCanExecuteAction::CreateSP(this, &SAssetView::IsToggleShowPluginContentAllowed);
+	OutParams.CanShowLocalizedContent = FCanExecuteAction::CreateSP(this, &SAssetView::IsToggleShowLocalizedContentAllowed);
+}
+
 void SAssetView::ToggleShowFolders()
 {
 	check(IsToggleShowFoldersAllowed());
@@ -3341,23 +3294,6 @@ bool SAssetView::HasDockedCollections() const
 	return GetDefault<UContentBrowserSettings>()->GetDockCollections();
 }
 
-void SAssetView::ToggleShowCppContent()
-{
-	check(IsToggleShowCppContentAllowed());
-
-	bool bNewState = !GetDefault<UContentBrowserSettings>()->GetDisplayCppFolders();
-
-	if (FContentBrowserInstanceConfig* Config = GetContentBrowserConfig())
-	{
-		bNewState = !Config->bShowCppFolders;
-		Config->bShowCppFolders = bNewState;
-		UContentBrowserConfig::Get()->SaveEditorConfig();
-	}
-
-	GetMutableDefault<UContentBrowserSettings>()->SetDisplayCppFolders(bNewState);
-	GetMutableDefault<UContentBrowserSettings>()->PostEditChange();
-}
-
 bool SAssetView::IsToggleShowCppContentAllowed() const
 {
 	return bCanShowClasses;
diff --git a/Engine/Source/Editor/ContentBrowser/Private/SContentBrowser.cpp b/Engine/Source/Editor/ContentBrowser/Private/SContentBrowser.cpp
--- a/Engine/Source/Editor/ContentBrowser/Private/SContentBrowser.cpp
+++ b/Engine/Source/Editor/ContentBrowser/Private/SContentBrowser.cpp
@@ -1259,11 +1259,12 @@ TSharedRef<SWidget> SContentBrowser::CreatePathView(const FContentBrowserConfig*
 					.ShowTreeTitle(false)
 					.ShowSeparator(false)
 				   	.ShowRedirectors_Lambda([this]() { return ContentBrowserUtils::ShouldShowRedirectors(FilterListPtr); })
-					.AllowClassesFolder(true)
 					.AddMetaData<FTagMetaData>(FTagMetaData(TEXT("ContentBrowserSources")))
 					.ExternalSearch(SourcesSearch)
 					.PluginPathFilters(PluginPathFilters)
 					.OwningContentBrowserName(InstanceName)
+					.AllowClassesFolder(Config != nullptr ? Config->bCanShowClasses : true)
+					.CanShowDevelopersFolder(Config != nullptr ? Config->bCanShowDevelopersFolder : true)
 				]
 			]
 		];
diff --git a/Engine/Source/Editor/ContentBrowser/Private/SPathPicker.cpp b/Engine/Source/Editor/ContentBrowser/Private/SPathPicker.cpp
--- a/Engine/Source/Editor/ContentBrowser/Private/SPathPicker.cpp
+++ b/Engine/Source/Editor/ContentBrowser/Private/SPathPicker.cpp
@@ -51,12 +51,16 @@ void SPathPicker::Construct( const FArguments& InArgs )
 		.OnGetItemContextMenu(this, &SPathPicker::GetItemContextMenu) // TODO: Allow this to be wholesale overridden via the picker config
 		.FocusSearchBoxWhenOpened(InArgs._PathPickerConfig.bFocusSearchBoxWhenOpened)
 		.AllowContextMenu(InArgs._PathPickerConfig.bAllowContextMenu)
-		.AllowClassesFolder(InArgs._PathPickerConfig.bAllowClassesFolder)
 		.AllowReadOnlyFolders(InArgs._PathPickerConfig.bAllowReadOnlyFolders)
 		.SelectionMode(ESelectionMode::Single)
 		.CustomFolderPermissionList(InArgs._PathPickerConfig.CustomFolderPermissionList)
 		.ShowFavorites(InArgs._PathPickerConfig.bShowFavorites)
 		.DefaultPath(InArgs._PathPickerConfig.bAddDefaultPath ? InArgs._PathPickerConfig.DefaultPath : FString{})
+		.AllowClassesFolder(InArgs._PathPickerConfig.bAllowClassesFolder)
+		.CanShowDevelopersFolder(InArgs._PathPickerConfig.bCanShowDevelopersFolder)
+		.ForceShowEngineContent(InArgs._PathPickerConfig.bForceShowEngineContent)
+		.ForceShowPluginContent(InArgs._PathPickerConfig.bForceShowPluginContent)
+		.ShowViewOptions(InArgs._PathPickerConfig.bShowViewOptions)
 	];
 	// clang-format on
 
diff --git a/Engine/Source/Editor/ContentBrowser/Private/SPathView.cpp b/Engine/Source/Editor/ContentBrowser/Private/SPathView.cpp
--- a/Engine/Source/Editor/ContentBrowser/Private/SPathView.cpp
+++ b/Engine/Source/Editor/ContentBrowser/Private/SPathView.cpp
@@ -13,7 +13,9 @@
 #include "ContentBrowserDataSubsystem.h"
 #include "ContentBrowserItemData.h"
 #include "ContentBrowserLog.h"
+#include "ContentBrowserMenuUtils.h"
 #include "ContentBrowserModule.h"
+#include "ContentBrowserPathViewMenuContexts.h"
 #include "ContentBrowserPluginFilters.h"
 #include "ContentBrowserSingleton.h"
 #include "ContentBrowserUtils.h"
@@ -57,11 +59,13 @@
 #include "Styling/ISlateStyle.h"
 #include "Textures/SlateIcon.h"
 #include "ToolMenu.h"
+#include "ToolMenus.h"
 #include "ToolMenuSection.h"
 #include "Trace/Detail/Channel.h"
 #include "Types/WidgetActiveTimerDelegate.h"
 #include "UObject/UObjectGlobals.h"
 #include "UObject/UnrealNames.h"
+#include "Widgets/Input/SComboButton.h"
 #include "Widgets/Input/SSearchBox.h"
 #include "Widgets/Layout/SBorder.h"
 #include "Widgets/Layout/SBox.h"
@@ -884,6 +888,9 @@ void SPathView::Construct( const FArguments& InArgs )
 	bAllowClassesFolder = InArgs._AllowClassesFolder;
 	bAllowReadOnlyFolders = InArgs._AllowReadOnlyFolders;
 	bShowRedirectors = InArgs._ShowRedirectors;
+	bCanShowDevelopersFolder = InArgs._CanShowDevelopersFolder;
+	bForceShowEngineContent = InArgs._ForceShowEngineContent;
+	bForceShowPluginContent = InArgs._ForceShowPluginContent;
 	bLastShowRedirectors = bShowRedirectors.Get(false);
 	PreventTreeItemChangedDelegateCount = 0;
 	TreeTitle = LOCTEXT("AssetTreeTitle", "Asset Tree");
@@ -969,6 +976,23 @@ void SPathView::Construct( const FArguments& InArgs )
 					SearchPtr->GetWidget()
 				]
 			]
+
+			+SHorizontalBox::Slot()
+			.Padding(4.f, 0.f, 0.f, 0.f)
+			.AutoWidth()
+			[
+				SNew(SComboButton)
+				.Visibility(InArgs._ShowViewOptions ? EVisibility::Visible : EVisibility::Collapsed)
+				.ComboButtonStyle(&FAppStyle::Get().GetWidgetStyle<FComboButtonStyle>("SimpleComboButton"))
+				.OnGetMenuContent(this, &SPathView::GetViewButtonContent)
+				.HasDownArrow(false)
+				.ButtonContent()
+				[
+					SNew(SImage)
+					.ColorAndOpacity(FSlateColor::UseForeground())
+					.Image(FAppStyle::Get().GetBrush("Icons.Settings"))
+				]
+			]
 		);
 	}
 
@@ -1468,9 +1492,9 @@ EContentBrowserItemAttributeFilter SPathView::GetContentBrowserItemAttributeFilt
 	}
 	
 	return EContentBrowserItemAttributeFilter::IncludeProject
-			| (bDisplayEngineContent ? EContentBrowserItemAttributeFilter::IncludeEngine : EContentBrowserItemAttributeFilter::IncludeNone)
-			| (bDisplayPluginContent ? EContentBrowserItemAttributeFilter::IncludePlugins : EContentBrowserItemAttributeFilter::IncludeNone)
-			| (bDisplayDevelopersContent ? EContentBrowserItemAttributeFilter::IncludeDeveloper : EContentBrowserItemAttributeFilter::IncludeNone)
+			| (bDisplayEngineContent || bForceShowEngineContent ? EContentBrowserItemAttributeFilter::IncludeEngine : EContentBrowserItemAttributeFilter::IncludeNone)
+			| (bDisplayPluginContent || bForceShowPluginContent ? EContentBrowserItemAttributeFilter::IncludePlugins : EContentBrowserItemAttributeFilter::IncludeNone)
+			| (bDisplayDevelopersContent && bCanShowDevelopersFolder ? EContentBrowserItemAttributeFilter::IncludeDeveloper : EContentBrowserItemAttributeFilter::IncludeNone)
 			| (bDisplayL10NContent ? EContentBrowserItemAttributeFilter::IncludeLocalized : EContentBrowserItemAttributeFilter::IncludeNone);
 }
 
@@ -2417,6 +2441,86 @@ TSharedRef<SWidget> SPathView::CreateFavoritesView()
 		];
 }
 
+void SPathView::RegisterGetViewButtonMenu()
+{
+	if (!UToolMenus::Get()->IsMenuRegistered("ContentBrowser.PathViewOptions"))
+	{
+		UToolMenu* Menu = UToolMenus::Get()->RegisterMenu("ContentBrowser.PathViewOptions");
+		Menu->bCloseSelfOnly = true;
+		Menu->AddDynamicSection("DynamicContent", FNewToolMenuDelegate::CreateLambda([](UToolMenu* InMenu)
+		{
+			FName ContextOwningContentBrowserName = NAME_None;
+			FFiltersAdditionalParams Params;
+			if (UContentBrowserPathViewContextMenuContext* Context = InMenu->FindContext<UContentBrowserPathViewContextMenuContext>())
+			{
+				if (Context->PathView.IsValid())
+				{
+					TSharedPtr<SPathView> PathView = Context->PathView.Pin();
+					PathView->PopulateFilterAdditionalParams(Params);
+
+					if (!PathView->OwningContentBrowserName.IsNone())
+					{
+						ContextOwningContentBrowserName = PathView->OwningContentBrowserName;
+					}
+				}
+
+				if (ContextOwningContentBrowserName.IsNone() && !Context->OwningContentBrowserName.IsNone())
+				{
+					ContextOwningContentBrowserName = Context->OwningContentBrowserName;
+				}
+
+				ContentBrowserMenuUtils::AddFiltersToMenu(InMenu, ContextOwningContentBrowserName, Params);
+			}
+		}));
+	}
+}
+
+void SPathView::PopulateFilterAdditionalParams(FFiltersAdditionalParams& OutParams)
+{
+	OutParams.CanShowCPPClasses = FCanExecuteAction::CreateSP(this, &SPathView::IsToggleShowCppContentAllowed);
+	OutParams.CanShowDevelopersContent = FCanExecuteAction::CreateSP(this, &SPathView::IsToggleShowDevelopersContentAllowed);
+	OutParams.CanShowEngineFolder = FCanExecuteAction::CreateSP(this, &SPathView::IsToggleShowEngineContentAllowed);
+	OutParams.CanShowPluginFolder = FCanExecuteAction::CreateSP(this, &SPathView::IsToggleShowPluginContentAllowed);
+	OutParams.CanShowLocalizedContent = FCanExecuteAction::CreateSP(this, &SPathView::IsToggleShowLocalizedContentAllowed);
+}
+
+bool SPathView::IsToggleShowCppContentAllowed() const
+{
+	return bAllowClassesFolder;
+}
+
+bool SPathView::IsToggleShowDevelopersContentAllowed() const
+{
+	return bCanShowDevelopersFolder;
+}
+
+bool SPathView::IsToggleShowEngineContentAllowed() const
+{
+	return !bForceShowEngineContent;
+}
+
+bool SPathView::IsToggleShowPluginContentAllowed() const
+{
+	return !bForceShowPluginContent;
+}
+
+bool SPathView::IsToggleShowLocalizedContentAllowed() const
+{
+	return true;
+}
+
+TSharedRef<SWidget> SPathView::GetViewButtonContent()
+{
+	SPathView::RegisterGetViewButtonMenu();
+
+	UContentBrowserPathViewContextMenuContext* Context = NewObject<UContentBrowserPathViewContextMenuContext>();
+	Context->PathView = SharedThis(this);
+	Context->OwningContentBrowserName = OwningContentBrowserName;
+
+	const FToolMenuContext MenuContext(Context);
+					
+	return UToolMenus::Get()->GenerateWidget("ContentBrowser.PathViewOptions", MenuContext);
+}
 
 void SFavoritePathView::Construct(const FArguments& InArgs)
 {
diff --git a/Engine/Source/Editor/ContentBrowser/Private/SPathView.h b/Engine/Source/Editor/ContentBrowser/Private/SPathView.h
--- a/Engine/Source/Editor/ContentBrowser/Private/SPathView.h
+++ b/Engine/Source/Editor/ContentBrowser/Private/SPathView.h
@@ -49,6 +49,7 @@ class ITableRow;
 class SWidget;
 class UToolMenu;
 struct FAssetData;
+struct FFiltersAdditionalParams;
 struct FGeometry;
 struct FHistoryData;
 struct FPathViewConfig;
@@ -77,6 +78,10 @@ public:
 		, _AllowClassesFolder(false)
 		, _AllowReadOnlyFolders(true)
 		, _ShowFavorites(false)
+		, _CanShowDevelopersFolder(false)
+		, _ForceShowEngineContent(false)
+		, _ForceShowPluginContent(false)
+		, _ShowViewOptions(false)
 		, _SelectionMode( ESelectionMode::Multi )
 		{}
 
@@ -116,6 +121,18 @@ public:
 		/** If true, the favorites expander will be displayed */
 		SLATE_ARGUMENT(bool, ShowFavorites);
 
+		/** Indicates if the 'Show Developers' option should be enabled or disabled */
+		SLATE_ARGUMENT(bool, CanShowDevelopersFolder)
+
+		/** Should always show engine content */
+		SLATE_ARGUMENT(bool, ForceShowEngineContent)
+
+		/** Should always show plugin content */
+		SLATE_ARGUMENT(bool, ForceShowPluginContent)
+
+		/** Should show the filter setting button. Note: If ExternalSearch is valid, then view options are not shown regardless of this setting */
+		SLATE_ARGUMENT(bool, ShowViewOptions)
+
 		/** If true, redirectors are taken into consideration when deciding if folders are empty */
 		SLATE_ATTRIBUTE(bool, ShowRedirectors);
 
@@ -361,6 +378,30 @@ private:
 	/** Create a favorites view. */
 	TSharedRef<SWidget> CreateFavoritesView();
 
+	/** Register menu for when the view combo button is clicked */
+	static void RegisterGetViewButtonMenu();
+
+	/** Populate the given params for this PathView */
+	void PopulateFilterAdditionalParams(FFiltersAdditionalParams& OutParams);
+
+	/** Whether or not it's possible to show C++ content */
+	bool IsToggleShowCppContentAllowed() const;
+
+	/** Whether or not it's possible to toggle developers content */
+	bool IsToggleShowDevelopersContentAllowed() const;
+
+	/** Whether or not it's possible to toggle engine content */
+	bool IsToggleShowEngineContentAllowed() const;
+
+	/** Whether or not it's possible to toggle plugin content */
+	bool IsToggleShowPluginContentAllowed() const;
+
+	/** Whether or not it's possible to show localized content */
+	bool IsToggleShowLocalizedContentAllowed() const;
+
+	/** Handler for when the view combo button is clicked */
+	TSharedRef<SWidget> GetViewButtonContent();
+
 protected:
 	/** A helper class to manage PreventTreeItemChangedDelegateCount by incrementing it when constructed (on the stack) and decrementing when destroyed */
 	class FScopedPreventTreeItemChangedDelegate
@@ -436,7 +477,7 @@ protected:
 	TSharedPtr<FPathPermissionList> CustomFolderPermissionList;
 
 	TAttribute<bool> bShowRedirectors;
-	bool bLastShowRedirectors;
+	bool bLastShowRedirectors = false;
 
 private:
 	/** Used to track if the list of last expanded path should be updated */
@@ -452,19 +493,28 @@ private:
 	FOnGetContentBrowserItemContextMenu OnGetItemContextMenu;
 
 	/** If > 0, the selection or expansion changed delegate will not be called. Used to update the tree from an external source or in certain bulk operations. */
-	int32 PreventTreeItemChangedDelegateCount;
+	int32 PreventTreeItemChangedDelegateCount = 0;
 
 	/** Initial set of item categories that this view should show - may be adjusted further by things like AllowClassesFolder */
-	EContentBrowserItemCategoryFilter InitialCategoryFilter;
+	EContentBrowserItemCategoryFilter InitialCategoryFilter = EContentBrowserItemCategoryFilter::IncludeAll;
 
 	/** If false, the context menu will not open when right clicking an item in the tree */
-	bool bAllowContextMenu;
+	bool bAllowContextMenu : 1;
 
 	/** If false, the classes folder will not be added to the tree automatically */
-	bool bAllowClassesFolder;
+	bool bAllowClassesFolder : 1;
 
 	/** If true, read only folders will be displayed */
-	bool bAllowReadOnlyFolders;
+	bool bAllowReadOnlyFolders : 1;
+
+	/** Indicates if the 'Show Developers' option should be enabled or disabled */
+	bool bCanShowDevelopersFolder : 1;
+
+	/** If true, engine content is always shown */
+	bool bForceShowEngineContent : 1;
+
+	/** If true, plugin content is always shown */
+	bool bForceShowPluginContent : 1;
 
 	/** The title of this path view */
 	FText TreeTitle;
diff --git a/Engine/Source/Editor/ContentBrowser/Public/IContentBrowserSingleton.h b/Engine/Source/Editor/ContentBrowser/Public/IContentBrowserSingleton.h
--- a/Engine/Source/Editor/ContentBrowser/Public/IContentBrowserSingleton.h
+++ b/Engine/Source/Editor/ContentBrowser/Public/IContentBrowserSingleton.h
@@ -388,6 +388,18 @@ struct FPathPickerConfig
 	/** Whether to call OnPathSelected during construction for DefaultPath if DefaultPath is allowed */
 	bool bNotifyDefaultPathSelected : 1;
 
+	/** Indicates if the 'Show Developers' option should be enabled or disabled */
+	bool bCanShowDevelopersFolder;
+
+	/** Indicates if engine content should always be shown */
+	bool bForceShowEngineContent;
+
+	/** Indicates if plugin content should always be shown */
+	bool bForceShowPluginContent;
+
+	/** if true, will add the filter setting to the path picker */
+	bool bShowViewOptions;
+
 	FPathPickerConfig()
 		: bFocusSearchBoxWhenOpened(true)
 		, bAllowContextMenu(true)
@@ -397,6 +409,10 @@ struct FPathPickerConfig
 		, bOnPathSelectedPassesVirtualPaths(false)
 		, bShowFavorites(true)
 		, bNotifyDefaultPathSelected(false)
+		, bCanShowDevelopersFolder(true)
+		, bForceShowEngineContent(false)
+		, bForceShowPluginContent(false)
+		, bShowViewOptions(false)
 	{}
 };
 
diff --git a/Engine/Source/Editor/ContentBrowser/Public/SAssetView.h b/Engine/Source/Editor/ContentBrowser/Public/SAssetView.h
--- a/Engine/Source/Editor/ContentBrowser/Public/SAssetView.h
+++ b/Engine/Source/Editor/ContentBrowser/Public/SAssetView.h
@@ -73,6 +73,7 @@ struct FAssetViewInstanceConfig;
 struct FCharacterEvent;
 struct FCollectionNameType;
 struct FContentBrowserInstanceConfig;
+struct FFiltersAdditionalParams;
 struct FFocusEvent;
 struct FGeometry;
 struct FKeyEvent;
@@ -529,6 +530,9 @@ private:
 	/** Fill in menu content for when the view combo button is clicked */
 	void PopulateViewButtonMenu(UToolMenu* Menu);
 
+	/** Populate the given params for this AssetView */
+	void PopulateFilterAdditionalParams(FFiltersAdditionalParams& OutParams);
+
 	/** Toggle whether folders should be shown or not */
 	void ToggleShowFolders();
 
@@ -616,9 +620,6 @@ private:
 	/** @return true when the collections view is docked */
 	bool HasDockedCollections() const;
 
-	/** Toggle whether C++ content should be shown or not */
-	void ToggleShowCppContent();
-
 	/** Whether or not it's possible to show C++ content */
 	bool IsToggleShowCppContentAllowed() const;
 
diff --git a/Engine/Source/Editor/UnrealEd/Private/Dialogs/DlgPickAssetPath.cpp b/Engine/Source/Editor/UnrealEd/Private/Dialogs/DlgPickAssetPath.cpp
--- a/Engine/Source/Editor/UnrealEd/Private/Dialogs/DlgPickAssetPath.cpp
+++ b/Engine/Source/Editor/UnrealEd/Private/Dialogs/DlgPickAssetPath.cpp
@@ -28,6 +28,9 @@ void SDlgPickAssetPath::Construct(const FArguments& InArgs)
 	PathPickerConfig.OnPathSelected = FOnPathSelected::CreateSP(this, &SDlgPickAssetPath::OnPathChange);
 	PathPickerConfig.bAddDefaultPath = true;
 	PathPickerConfig.bAllowReadOnlyFolders = bAllowReadOnlyFolders;
+	PathPickerConfig.bAllowClassesFolder = true;
+	PathPickerConfig.bCanShowDevelopersFolder = true;
+	PathPickerConfig.bShowViewOptions = true;
 
 	FContentBrowserModule& ContentBrowserModule = FModuleManager::LoadModuleChecked<FContentBrowserModule>("ContentBrowser");
 
