Make content browser tooltips for 2D and cube texture arrays show array size

diff --git a/Engine/Source/Editor/ContentBrowser/Private/AssetViewWidgets.cpp b/Engine/Source/Editor/ContentBrowser/Private/AssetViewWidgets.cpp
--- a/Engine/Source/Editor/ContentBrowser/Private/AssetViewWidgets.cpp
+++ b/Engine/Source/Editor/ContentBrowser/Private/AssetViewWidgets.cpp
@@ -1370,34 +1370,60 @@ void SAssetViewItem::CacheDisplayTags()
 				}
 			}
 	
-			// Dimensional tags need to be split into their component numbers, with each component number re-format
+			// Dimensional tags need to be split into their component numbers, with each component number re-formatted
 			if (!bHasSetDisplayValue && AttributeMetaData.AttributeType == UObject::FAssetRegistryTag::TT_Dimensional)
 			{
-				TArray<FString> NumberStrTokens;
-				AttributeValueStr.ParseIntoArray(NumberStrTokens, TEXT("x"), true);
-	
-				if (NumberStrTokens.Num() > 0 && NumberStrTokens.Num() <= 3)
+				// Formats:
+				//   123         (1D)
+				//   123x234     (2D)
+				//   123x234*345 (2D array)
+				//   123x234x345 (3D)
+				int32 FirstXPos;
+				if (AttributeValueStr.FindChar(TEXT('x'), FirstXPos))
 				{
-					bHasSetDisplayValue = true;
-	
-					switch (NumberStrTokens.Num())
+					FString FirstPart = AttributeValueStr.Left(FirstXPos);
+					FString Remainder = AttributeValueStr.Mid(FirstXPos + 1);
+					int32 RemainderSeparatorPos;
+
+					if (Remainder.FindChar(TEXT('*'), RemainderSeparatorPos))
 					{
-					case 1:
-						DisplayValue = ReformatNumberStringForDisplay(NumberStrTokens[0]);
-						break;
-	
-					case 2:
-						DisplayValue = FText::Format(LOCTEXT("DisplayTag2xFmt", "{0} \u00D7 {1}"), ReformatNumberStringForDisplay(NumberStrTokens[0]), ReformatNumberStringForDisplay(NumberStrTokens[1]));
-						break;
-	
-					case 3:
-						DisplayValue = FText::Format(LOCTEXT("DisplayTag3xFmt", "{0} \u00D7 {1} \u00D7 {2}"), ReformatNumberStringForDisplay(NumberStrTokens[0]), ReformatNumberStringForDisplay(NumberStrTokens[1]), ReformatNumberStringForDisplay(NumberStrTokens[2]));
-						break;
-	
-					default:
-						break;
+						// AxB*C form (2D array)
+						FString SecondPart = Remainder.Left(RemainderSeparatorPos);
+						FString ThirdPart = Remainder.Mid(RemainderSeparatorPos + 1);
+
+						bHasSetDisplayValue = true;
+						DisplayValue = FText::Format(LOCTEXT("DisplayTag2xArrayFmt", "{0} \u00D7 {1} ({2} elements)"),
+							ReformatNumberStringForDisplay(FirstPart),
+							ReformatNumberStringForDisplay(SecondPart),
+							ReformatNumberStringForDisplay(ThirdPart));
+					}
+					else if (Remainder.FindChar(TEXT('x'), RemainderSeparatorPos))
+					{
+						// AxBxC form (3D)
+						FString SecondPart = Remainder.Left(RemainderSeparatorPos);
+						FString ThirdPart = Remainder.Mid(RemainderSeparatorPos + 1);
+
+						bHasSetDisplayValue = true;
+						DisplayValue = FText::Format(LOCTEXT("DisplayTag3xFmt", "{0} \u00D7 {1} \u00D7 {2}"),
+							ReformatNumberStringForDisplay(FirstPart),
+							ReformatNumberStringForDisplay(SecondPart),
+							ReformatNumberStringForDisplay(ThirdPart));
+					}
+					else
+					{
+						// 2D form by default
+						bHasSetDisplayValue = true;
+						DisplayValue = FText::Format(LOCTEXT("DisplayTag2xFmt", "{0} \u00D7 {1}"),
+							ReformatNumberStringForDisplay(FirstPart),
+							ReformatNumberStringForDisplay(Remainder));
 					}
 				}
+				else
+				{
+					// No separators, assume 1D
+					bHasSetDisplayValue = true;
+					DisplayValue = ReformatNumberStringForDisplay(AttributeValueStr);
+				}
 			}
 	
 			// Chronological tags need to format the specified timestamp based on the display flags
