Fix filtering out developer folders from favorites view without meaning to.

diff --git a/Engine/Source/Editor/ContentBrowser/Private/SContentBrowser.cpp b/Engine/Source/Editor/ContentBrowser/Private/SContentBrowser.cpp
--- a/Engine/Source/Editor/ContentBrowser/Private/SContentBrowser.cpp
+++ b/Engine/Source/Editor/ContentBrowser/Private/SContentBrowser.cpp
@@ -1179,6 +1179,7 @@ TSharedRef<SWidget> SContentBrowser::CreateFavoritesView(const FContentBrowserCo
 				.ShowTreeTitle(false)
 				.ShowSeparator(false)
 				.AllowClassesFolder(true)
+				.CanShowDevelopersFolder(true)
 				.OwningContentBrowserName(InstanceName)
 				.AddMetaData<FTagMetaData>(FTagMetaData(TEXT("ContentBrowserFavorites")))
 				.ExternalSearch(FavoritesSearch)
diff --git a/Engine/Source/Editor/ContentBrowser/Private/SPathView.cpp b/Engine/Source/Editor/ContentBrowser/Private/SPathView.cpp
--- a/Engine/Source/Editor/ContentBrowser/Private/SPathView.cpp
+++ b/Engine/Source/Editor/ContentBrowser/Private/SPathView.cpp
@@ -2434,6 +2434,7 @@ TSharedRef<SWidget> SPathView::CreateFavoritesView()
 			.ShowTreeTitle(false)
 			.ShowSeparator(false)
 			.AllowClassesFolder(bAllowClassesFolder)
+			.CanShowDevelopersFolder(bCanShowDevelopersFolder)
 			.AllowReadOnlyFolders(bAllowReadOnlyFolders)
 			.AllowContextMenu(bAllowContextMenu)
 			.AddMetaData<FTagMetaData>(FTagMetaData(TEXT("ContentBrowserFavorites")))
