Allow navigation bar to accept more types of paths to navigate.
- Object path
- Export text path
- Disk path to asset file
- Non-virtual package path
Don't force path tree to repopulate when navigating unless navigating to a folder that isn't present in the tree.
Don't query backend for crumbs to navigate to every frame.

diff --git a/Engine/Source/Editor/ContentBrowser/Private/ContentBrowserUtils.cpp b/Engine/Source/Editor/ContentBrowser/Private/ContentBrowserUtils.cpp
--- a/Engine/Source/Editor/ContentBrowser/Private/ContentBrowserUtils.cpp
+++ b/Engine/Source/Editor/ContentBrowser/Private/ContentBrowserUtils.cpp
@@ -1032,4 +1032,83 @@ EContentBrowserItemAttributeFilter ContentBrowserUtils::GetContentBrowserItemAtt
 		 | (bDisplayDevelopersContent ? EContentBrowserItemAttributeFilter::IncludeDeveloper : EContentBrowserItemAttributeFilter::IncludeNone)
 		 | (bDisplayL10NContent ? EContentBrowserItemAttributeFilter::IncludeLocalized : EContentBrowserItemAttributeFilter::IncludeNone);
 }
+
+FContentBrowserItem ContentBrowserUtils::TryGetItemFromUserProvidedPath(FStringView RequestedPathView)
+{
+	// For all types of accepted input we can trim a trailing slash if it exists
+	if (RequestedPathView.EndsWith('/'))
+	{
+		RequestedPathView.LeftChopInline(1);
+	}
+
+	UContentBrowserDataSubsystem* ContentBrowserData = IContentBrowserDataModule::Get().GetSubsystem();
+	FName RequestedPath(RequestedPathView);
+
+	// If the path is already a valid virtual path, go there
+	FContentBrowserItem Item = ContentBrowserData->GetItemAtPath(RequestedPath, EContentBrowserItemTypeFilter::IncludeAll);
+	if (Item.IsValid())
+	{
+		return Item;
+	}
+
+	// If the path is a non-virtual path like /Game/Maps transform it into a virtual path and try and find an item there
+	FName VirtualPath = ContentBrowserData->ConvertInternalPathToVirtual(RequestedPath);
+	if (!VirtualPath.IsNone())
+	{
+		Item = ContentBrowserData->GetItemAtPath(VirtualPath, EContentBrowserItemTypeFilter::IncludeAll);
+		if (Item.IsValid())
+		{
+			return Item;
+		}
+	}
+
+	// If the string is a complete object path (with or without class), sync to that asset
+	FStringView ObjectPathView = RequestedPathView;
+	if (FPackageName::IsValidObjectPath(ObjectPathView) || FPackageName::ParseExportTextPath(RequestedPathView, nullptr, &ObjectPathView))
+	{
+		VirtualPath = ContentBrowserData->ConvertInternalPathToVirtual(FName(ObjectPathView));
+		Item = ContentBrowserData->GetItemAtPath(VirtualPath, EContentBrowserItemTypeFilter::IncludeFiles);
+		if (Item.IsValid())
+		{
+			return Item;
+		}
+	}
+
+	// If the string is a filesystem path to a package, sync to that asset
+	FString PackageName;
+	if (FPackageName::IsValidLongPackageName(PackageName) || FPackageName::TryConvertFilenameToLongPackageName(FString(RequestedPathView), PackageName))
+	{
+		// Packages like /Game/Characters/Knight do not map to virtual paths in data source, assets like /Game/Characters/Knight.Knight do.
+		// See if there's an item if we duplicate the last part of the path 
+		FName InternalPath(TStringBuilder<320>(InPlace, PackageName, TEXTVIEW("."), FPackageName::GetShortName(PackageName)));
+		VirtualPath = ContentBrowserData->ConvertInternalPathToVirtual(InternalPath);
+		Item = ContentBrowserData->GetItemAtPath(RequestedPath, EContentBrowserItemTypeFilter::IncludeFiles);
+		if (Item.IsValid())
+		{
+			return Item;
+		}
+
+		// Otherwise go up to the package path and enumerate items to see if there's an asset with the desired package name 
+		FString PackagePath = FPackageName::GetLongPackagePath(PackageName);
+		VirtualPath = ContentBrowserData->ConvertInternalPathToVirtual(FName(PackagePath));
+		FContentBrowserDataFilter Filter;
+		Filter.bRecursivePaths = false;
+		Filter.ItemTypeFilter = EContentBrowserItemTypeFilter::IncludeFiles;
+		ContentBrowserData->EnumerateItemsUnderPath(VirtualPath, Filter, [&Item, &PackageName](FContentBrowserItem&& InItem) { 
+			FName InternalPath = InItem.GetInternalPath();	
+			if (WriteToString<256>(InternalPath).ToView().StartsWith(PackageName))
+			{
+				Item = MoveTemp(InItem);
+				return false; 
+			}
+			return true;
+		});
+		if (Item.IsValid())
+		{
+			return Item;
+		}
+	}
+
+	return FContentBrowserItem();
+}
 #undef LOCTEXT_NAMESPACE
diff --git a/Engine/Source/Editor/ContentBrowser/Private/ContentBrowserUtils.h b/Engine/Source/Editor/ContentBrowser/Private/ContentBrowserUtils.h
--- a/Engine/Source/Editor/ContentBrowser/Private/ContentBrowserUtils.h
+++ b/Engine/Source/Editor/ContentBrowser/Private/ContentBrowserUtils.h
@@ -168,4 +168,16 @@ namespace ContentBrowserUtils
 	 * An instance name of None gives global settings.
 	 */
 	EContentBrowserItemAttributeFilter GetContentBrowserItemAttributeFilter(FName InstanceName);
+
+	/** 
+	 * Try and find a valid content browser item from a path provided from user text input.
+	 * Examples of paths:
+	 * 	- A full virtual folder or item path such as /All/Game/Maps /All/Game/Maps/Arena.Arena
+	 *  - An internal asset path such as /Game/Maps/Arena.Arena
+	 *  - An internal package name such as /Game/Maps/Arena 
+	 * 	- An internal package path such as /Game/Maps
+	 * 	- An 'export text' path / asset reference such as /Script/Engine.World'/Game/Maps/Arena.Arena'
+	 *  - A filesystem path to a uasset or umap file.
+	 */
+	FContentBrowserItem TryGetItemFromUserProvidedPath(FStringView RequestedPathView);
 }
diff --git a/Engine/Source/Editor/ContentBrowser/Private/SContentBrowser.cpp b/Engine/Source/Editor/ContentBrowser/Private/SContentBrowser.cpp
--- a/Engine/Source/Editor/ContentBrowser/Private/SContentBrowser.cpp
+++ b/Engine/Source/Editor/ContentBrowser/Private/SContentBrowser.cpp
@@ -168,10 +168,6 @@ namespace ContentBrowserSourcesWidgetSwitcherIndex
 
 namespace ContentBrowser
 {
-	// It is useful when debugging enumeration issues to shut this system off so that breakpoints set will not be triggered every tick
-	bool bCrumbsEnumerate = true;
-	FAutoConsoleVariableRef CrumbsEnumerateCVar(TEXT("ContentBrowser.Debug.CrumbsEnumerate"), bCrumbsEnumerate, TEXT("Enumerate crumbs"), ECVF_Default);
-	
 	const FContentBrowserInstanceConfig* GetConstInstanceConfig(FName ForInstance) 
 	{
 		if (ForInstance.IsNone())
@@ -509,7 +505,6 @@ void SContentBrowser::Construct( const FArguments& InArgs, const FName& InInstan
 				[
 					SAssignNew(NavigationBar, SNavigationBar)
 					.OnPathClicked(this, &SContentBrowser::OnPathClicked)
-					.HasPathMenuContent(this, &SContentBrowser::OnHasCrumbDelimiterContent)
 					.GetPathMenuContent(this, &SContentBrowser::OnGetCrumbDelimiterContent)
 					.GetComboOptions(this, &SContentBrowser::GetRecentPaths)
 					.OnNavigateToPath(this, &SContentBrowser::OnNavigateToPath)
@@ -1634,7 +1629,7 @@ void SContentBrowser::PrepareToSyncItems(TArrayView<const FContentBrowserItem> I
 	if (!bRepopulate)
 	{
 		bRepopulate = Algo::AnyOf(ItemsToSync, [this](const FContentBrowserItem& Item) {
-			return !PathViewPtr->DoesItemExist(Item.GetVirtualPath());
+			return Item.IsFolder() && !PathViewPtr->DoesItemExist(Item.GetVirtualPath());
 		});
 	}
 	
@@ -2504,15 +2499,11 @@ void SContentBrowser::EditPathCommand()
 
 void SContentBrowser::OnNavigateToPath(const FString& NewPath)
 {
-	// If NewPath is not a valid location, the content browser will stay in its current location.
-	TArray<FName> VirtualPaths;
-	FStringView PathView = NewPath;
-	if (PathView.EndsWith('/'))
+	FContentBrowserItem Item = ContentBrowserUtils::TryGetItemFromUserProvidedPath(NewPath);
+	if (Item.IsValid())
 	{
-		PathView.LeftChopInline(1);
+		SyncToItems(MakeArrayView(&Item, 1));
 	}
-	VirtualPaths.Add(FName(PathView)); 
-	SyncToVirtualPaths(VirtualPaths);
 }
 
 void SContentBrowser::OnPathClicked( const FString& CrumbData )
@@ -2614,62 +2605,6 @@ TArray<FString> SContentBrowser::OnCompletePathPrefix(const FString& Prefix) con
 	return Results;
 }
 
-bool SContentBrowser::OnHasCrumbDelimiterContent(const FString& CrumbData) const
-{
-	const FSourcesData SourcesData = AssetViewPtr->GetSourcesData();
-	if (SourcesData.HasCollections())
-	{
-		TOptional<FCollectionNameType> CollectionClicked;
-		{
-			FString CollectionName;
-			FString CollectionTypeString;
-			if (CrumbData.Split(TEXT("?"), &CollectionName, &CollectionTypeString))
-			{
-				const int32 CollectionType = FCString::Atoi(*CollectionTypeString);
-				if (CollectionType >= 0 && CollectionType < ECollectionShareType::CST_All)
-				{
-					CollectionClicked = FCollectionNameType(FName(*CollectionName), ECollectionShareType::Type(CollectionType));
-				}
-			}
-		}
-
-		TArray<FCollectionNameType> ChildCollections;
-		if (CollectionClicked.IsSet())
-		{
-			FCollectionManagerModule& CollectionManagerModule = FCollectionManagerModule::GetModule();
-			CollectionManagerModule.Get().GetChildCollections(CollectionClicked->Name, CollectionClicked->Type, ChildCollections);
-		}
-
-		return (ChildCollections.Num() > 0);
-	}
-	else if (SourcesData.HasVirtualPaths())
-	{
-		// The subsystems may get Deinitialize before the slate windows are close when exiting the editor
-		if (UContentBrowserDataSubsystem* ContentBrowserData = IContentBrowserDataModule::Get().GetSubsystem())
-		{
-			FContentBrowserDataFilter SubItemsFilter;
-			SubItemsFilter.ItemTypeFilter = EContentBrowserItemTypeFilter::IncludeFolders;
-			SubItemsFilter.bRecursivePaths = false;
-			SubItemsFilter.ItemCategoryFilter = PathViewPtr->GetContentBrowserItemCategoryFilter();
-			SubItemsFilter.ItemAttributeFilter = PathViewPtr->GetContentBrowserItemAttributeFilter();
-
-			bool bHasSubItems = false;
-			if (ContentBrowser::bCrumbsEnumerate)
-			{
-				ContentBrowserData->EnumerateItemsUnderPath(*CrumbData, SubItemsFilter, [&bHasSubItems](FContentBrowserItemData&& InSubItem)
-					{
-						bHasSubItems = true;
-				return false;
-					});
-			}
-
-			return bHasSubItems;
-		}
-	}
-
-	return false;
-}
-
 TSharedRef<SWidget> SContentBrowser::OnGetCrumbDelimiterContent(const FString& CrumbData) const
 {
 	FSourcesData SourcesData = AssetViewPtr->GetSourcesData();
@@ -3701,13 +3636,26 @@ void SContentBrowser::UpdatePath()
 		TArray<FString> Crumbs;
 		SourcesData.VirtualPaths[0].ToString().ParseIntoArray(Crumbs, TEXT("/"), true);
 
+		FContentBrowserDataFilter SubItemsFilter;
+		SubItemsFilter.ItemTypeFilter = EContentBrowserItemTypeFilter::IncludeFolders;
+		SubItemsFilter.bRecursivePaths = false;
+		SubItemsFilter.ItemCategoryFilter = PathViewPtr->GetContentBrowserItemCategoryFilter();
+		SubItemsFilter.ItemAttributeFilter = PathViewPtr->GetContentBrowserItemAttributeFilter();
+
 		FString CrumbPath = TEXT("/");
 		for (const FString& Crumb : Crumbs)
 		{
 			CrumbPath += Crumb;
 
+			bool bHasSubItems = false;
+			ContentBrowserData->EnumerateItemsUnderPath(*CrumbPath, SubItemsFilter, [&bHasSubItems](FContentBrowserItemData&& InSubItem)
+				{
+					bHasSubItems = true;
+					return false;
+				});
+
 			const FContentBrowserItem CrumbFolderItem = ContentBrowserData->GetItemAtPath(*CrumbPath, EContentBrowserItemTypeFilter::IncludeFolders);
-			NavigationBar->PushPath(CrumbFolderItem.IsValid() ? CrumbFolderItem.GetDisplayName() : FText::FromString(Crumb), CrumbPath);
+			NavigationBar->PushPath(CrumbFolderItem.IsValid() ? CrumbFolderItem.GetDisplayName() : FText::FromString(Crumb), CrumbPath, bHasSubItems);
 
 			CrumbPath += TEXT("/");
 		}
@@ -3732,17 +3680,21 @@ void SContentBrowser::UpdatePath()
 		for (const FCollectionNameType& CollectionPathItem : CollectionPathItems)
 		{
 			const FString CrumbData = FString::Printf(TEXT("%s?%s"), *CollectionPathItem.Name.ToString(), *FString::FromInt(CollectionPathItem.Type));
+			
+			TArray<FCollectionNameType> ChildCollections;
+			CollectionManagerModule.Get().GetChildCollections(CollectionPathItem.Name, CollectionPathItem.Type, ChildCollections);
+			const bool bHasChildren = ChildCollections.Num() > 0;
 
 			FFormatNamedArguments Args;
 			Args.Add(TEXT("CollectionName"), FText::FromName(CollectionPathItem.Name));
 			const FText DisplayName = FText::Format(LOCTEXT("CollectionPathIndicator", "{CollectionName} (Collection)"), Args);
 
-			NavigationBar->PushPath(DisplayName, CrumbData);
+			NavigationBar->PushPath(DisplayName, CrumbData, bHasChildren);
 		}
 	}
 	else
 	{
-		NavigationBar->PushPath(LOCTEXT("AllAssets", "All Assets"), TEXT(""));
+		NavigationBar->PushPath(LOCTEXT("AllAssets", "All Assets"), TEXT(""), true); 
 	}
 
 	if (ActiveSourcesWidgetIndex != NewSourcesWidgetIndex)
diff --git a/Engine/Source/Editor/ContentBrowser/Private/SContentBrowser.h b/Engine/Source/Editor/ContentBrowser/Private/SContentBrowser.h
--- a/Engine/Source/Editor/ContentBrowser/Private/SContentBrowser.h
+++ b/Engine/Source/Editor/ContentBrowser/Private/SContentBrowser.h
@@ -318,9 +318,6 @@ private:
 	
 	/** Called when item in the path delimiter arrow menu is clicked */
 	void OnPathMenuItemClicked(FString ClickedPath);
-	
-	/** Called to query whether the crumb menu would contain any content (see also OnGetCrumbDelimiterContent) */
-	bool OnHasCrumbDelimiterContent(const FString& CrumbData) const;
 
 	/** 
 	 * Populates the delimiter arrow with a menu of directories under the current directory that can be navigated to
diff --git a/Engine/Source/Editor/ContentBrowser/Private/SNavigationBar.cpp b/Engine/Source/Editor/ContentBrowser/Private/SNavigationBar.cpp
--- a/Engine/Source/Editor/ContentBrowser/Private/SNavigationBar.cpp
+++ b/Engine/Source/Editor/ContentBrowser/Private/SNavigationBar.cpp
@@ -14,6 +14,12 @@
 
 #define LOCTEXT_NAMESPACE "ContentBrowser"
 
+struct FNavigationCrumb
+{
+	FString Path;
+	bool bHasChildren; // Whether to show the > button next to this location allowing navigation to children
+};
+
 // Category of list view item to choose an icon
 enum class ELocationSource
 {
@@ -104,6 +110,8 @@ void SNavigationBar::Construct(const FArguments& InArgs)
 	OnCompletePrefix = InArgs._OnCompletePrefix;
 	OnNavigateToPath = InArgs._OnNavigateToPath;
 	OnCanEditPathAsText = InArgs._OnCanEditPathAsText;
+	OnPathClicked = InArgs._OnPathClicked;
+	OnGetPathMenuContent = InArgs._GetPathMenuContent;
 
 	ComboBoxStyle = InArgs._ComboBoxStyle;
 	TextBoxStyle = InArgs._TextBoxStyle;
@@ -139,15 +147,15 @@ void SNavigationBar::Construct(const FArguments& InArgs)
 				.HAlign(HAlign_Left)
 				.VAlign(VAlign_Fill)
 				[
-					SAssignNew(BreadcrumbBar, SBreadcrumbTrail<FString>)
+					SAssignNew(BreadcrumbBar, SBreadcrumbTrail<FNavigationCrumb>)
 					.Visibility(this, &SNavigationBar::GetNonEditVisibility)
 					.TextStyle(InArgs._BreadcrumbTextStyle)
 					.ButtonStyle(InArgs._BreadcrumbButtonStyle)
 					.ButtonContentPadding(InArgs._BreadcrumbButtonContentPadding)
 					.DelimiterImage(InArgs._BreadcrumbDelimiterImage)
-					.OnCrumbClicked(InArgs._OnPathClicked)
-					.HasCrumbMenuContent(InArgs._HasPathMenuContent)
-					.GetCrumbMenuContent(InArgs._GetPathMenuContent)
+					.OnCrumbClicked(this, &SNavigationBar::HandleCrumbClicked)
+					.HasCrumbMenuContent(this, &SNavigationBar::HandleHasCrumbMenuContent)
+					.GetCrumbMenuContent(this, &SNavigationBar::HandleGetCrumbMenuContent)
 				]
 				// Editable text box taking up all space
 				+ SOverlay::Slot()
@@ -216,9 +224,28 @@ void SNavigationBar::ClearPaths()
 	BreadcrumbBar->ClearCrumbs();
 }
 
-void SNavigationBar::PushPath(const FText& ElementText, const FString& FullPath)
+void SNavigationBar::PushPath(const FText& ElementText, const FString& FullPath, bool bHasChildren)
+{
+	BreadcrumbBar->PushCrumb(ElementText, FNavigationCrumb{FullPath, bHasChildren});	
+}
+
+void SNavigationBar::HandleCrumbClicked(const FNavigationCrumb& Crumb)
 {
-	BreadcrumbBar->PushCrumb(ElementText, FullPath);	
+	OnPathClicked.ExecuteIfBound(Crumb.Path);
+}
+
+bool SNavigationBar::HandleHasCrumbMenuContent(const FNavigationCrumb& Crumb)
+{
+	return Crumb.bHasChildren;
+}
+
+TSharedRef<SWidget> SNavigationBar::HandleGetCrumbMenuContent(const FNavigationCrumb& Crumb)
+{
+	if (OnGetPathMenuContent.IsBound())
+	{
+		return OnGetPathMenuContent.Execute(Crumb.Path);
+	}
+	return SNullWidget::NullWidget;
 }
 
 bool SNavigationBar::SupportsKeyboardFocus() const 
@@ -390,10 +417,10 @@ void SNavigationBar::StartEditingPath()
 	FText Text = FText::GetEmpty();
 	if (BreadcrumbBar->HasCrumbs())
 	{
-		FString Path = BreadcrumbBar->PeekCrumb();
-		if (OnCanEditPathAsText.IsBound() && OnCanEditPathAsText.Execute(Path))
+		FNavigationCrumb Crumb = BreadcrumbBar->PeekCrumb();
+		if (OnCanEditPathAsText.IsBound() && OnCanEditPathAsText.Execute(Crumb.Path))
 		{
-			Text = FText::FromString(Path);
+			Text = FText::FromString(Crumb.Path);
 		}
 	}
 	EditableText->SetText(Text);
@@ -412,7 +439,6 @@ void SNavigationBar::HandleTextChanged(const FText& NewText)
 		CompletionTimerHandle = RegisterActiveTimer(0.1f, FWidgetActiveTimerDelegate::CreateSP(this, &SNavigationBar::HandleUpdateCompletionOptions));
 	}
 }
-
 	
 void SNavigationBar::HandleTextCommitted(const FText& InText, ETextCommit::Type CommitType)
 {
diff --git a/Engine/Source/Editor/ContentBrowser/Private/SNavigationBar.h b/Engine/Source/Editor/ContentBrowser/Private/SNavigationBar.h
--- a/Engine/Source/Editor/ContentBrowser/Private/SNavigationBar.h
+++ b/Engine/Source/Editor/ContentBrowser/Private/SNavigationBar.h
@@ -11,6 +11,8 @@
 #include "Widgets/Layout/SBorder.h"
 
 class FText;
+// Type of individual elements in breadcrumb bar, defines the current location or one of its parents.
+struct FNavigationCrumb;
 class ITableRow;
 template<typename ItemType>
 class SBreadcrumbTrail;
@@ -21,16 +23,16 @@ class STableViewBase;
 class USlateWidgetStyleAsset;
 
 using FOnPathClicked = TDelegate<void(const FString&)>;
-using FHasPathMenuContent = TDelegate<bool(const FString&)>;
 using FGetPathMenuContent = TDelegate<TSharedRef<SWidget>(const FString&)>;
 using FGetComboOptions = TDelegate<TArray<FString>(void)>;
 using FOnCompletePrefix = TDelegate<TArray<FString>(const FString&)>;
 using FOnNavigateToPath = TDelegate<void(const FString&)>;
 using FCanEditPathAsText = TDelegate<bool(const FString&)>;
 
-// Private class for internal list view
+// Private class for internal list view, stores a full possible location to nagivate to.
 struct FLocationItem;
 
+
 class SNavigationBar : public SComboButton
 {
 public:
@@ -73,9 +75,6 @@ public:
 		/** Called when an invididual button is clicked, after the later crumbs were popped */
 		SLATE_EVENT(FOnPathClicked , OnPathClicked)
 
-		/** Called to check whether there are locations after a particular crumb to populate a dropdown*/
-		SLATE_EVENT(FHasPathMenuContent, HasPathMenuContent)
-
 		/** Called to get dropdown options for a crumb separator's menu */
 		SLATE_EVENT(FGetPathMenuContent, GetPathMenuContent)
 		
@@ -98,7 +97,7 @@ public:
 	// Remove all paths from the breadcrumb bar
 	void ClearPaths();
 	// Add a new segment to the breadcrumb bar
-	void PushPath(const FText& SegmentDisplayText, const FString& FullLocation);
+	void PushPath(const FText& SegmentDisplayText, const FString& FullLocation, bool bHasChildren);
 	// Replace the breadcrumb bar with an editable text box and focus it
 	void StartEditingPath();
 
@@ -117,6 +116,9 @@ private:
 	FReply HandleEditableTextKeyDown(const FGeometry& MyGeometry, const FKeyEvent& KeyEvent);
 	FReply HandleComboKeyDown(const FGeometry& MyGeometry, const FKeyEvent& InKeyEvent);
 	FReply HandleComboKeyChar(const FGeometry& MyGeometry, const FCharacterEvent& InCharEvent);
+	void HandleCrumbClicked(const FNavigationCrumb& Crumb);
+	bool HandleHasCrumbMenuContent(const FNavigationCrumb& Crumb);
+	TSharedRef<SWidget> HandleGetCrumbMenuContent(const FNavigationCrumb& Crumb);
 	
 	EActiveTimerReturnType HandleUpdateCompletionOptions(double InCurrentTime, float InDeltaTime);
 
@@ -140,7 +142,7 @@ private:
 	TSharedPtr<FActiveTimerHandle> CompletionTimerHandle;
 
 	TSharedPtr<SEditableText> EditableText;
-    TSharedPtr<SBreadcrumbTrail<FString>> BreadcrumbBar;
+    TSharedPtr<SBreadcrumbTrail<FNavigationCrumb>> BreadcrumbBar;
 	TSharedPtr<SLocationListView> ComboListView;
 	
 	FText PopupHeading;
@@ -151,4 +153,6 @@ private:
 	FGetComboOptions OnGetComboOptions;
 	FOnCompletePrefix OnCompletePrefix;	
 	FCanEditPathAsText OnCanEditPathAsText;
+	FOnPathClicked OnPathClicked;
+	FGetPathMenuContent OnGetPathMenuContent;
 };
\ No newline at end of file
