Additional path handling from user input
Fix long package name case
Try and handle virtual path to asset missing trailing .AssetName
As a last effort, check ancestors of given path string as virtual paths

diff --git a/Engine/Source/Editor/ContentBrowser/Private/ContentBrowserUtils.cpp b/Engine/Source/Editor/ContentBrowser/Private/ContentBrowserUtils.cpp
--- a/Engine/Source/Editor/ContentBrowser/Private/ContentBrowserUtils.cpp
+++ b/Engine/Source/Editor/ContentBrowser/Private/ContentBrowserUtils.cpp
@@ -1074,15 +1074,12 @@ FContentBrowserItem ContentBrowserUtils::TryGetItemFromUserProvidedPath(FStringV
 		}
 	}
 
-	// If the string is a filesystem path to a package, sync to that asset
-	FString PackageName;
-	if (FPackageName::IsValidLongPackageName(PackageName) || FPackageName::TryConvertFilenameToLongPackageName(FString(RequestedPathView), PackageName))
-	{
+	auto GetItemFromPackageName = [ContentBrowserData](const FString& PackageName) -> FContentBrowserItem {
 		// Packages like /Game/Characters/Knight do not map to virtual paths in data source, assets like /Game/Characters/Knight.Knight do.
 		// See if there's an item if we duplicate the last part of the path 
 		FName InternalPath(TStringBuilder<320>(InPlace, PackageName, TEXTVIEW("."), FPackageName::GetShortName(PackageName)));
-		VirtualPath = ContentBrowserData->ConvertInternalPathToVirtual(InternalPath);
-		Item = ContentBrowserData->GetItemAtPath(RequestedPath, EContentBrowserItemTypeFilter::IncludeFiles);
+		FName VirtualPath = ContentBrowserData->ConvertInternalPathToVirtual(InternalPath);
+		FContentBrowserItem Item = ContentBrowserData->GetItemAtPath(VirtualPath, EContentBrowserItemTypeFilter::IncludeFiles);
 		if (Item.IsValid())
 		{
 			return Item;
@@ -1107,6 +1104,57 @@ FContentBrowserItem ContentBrowserUtils::TryGetItemFromUserProvidedPath(FStringV
 		{
 			return Item;
 		}
+		return FContentBrowserItem();
+	};
+
+	// If the string is an incomplete virtual path that looks more like a package name 
+	// e.g. /All/Game/Maps/Arena rather than /Game/Maps/Arena or /All/Game/Maps/Arena.Arena
+	// try and convert it to an internal path, then try and use it as a package name 
+	{
+		FName ConvertedPath;
+		FString PackageName;
+		if (ContentBrowserData->TryConvertVirtualPath(RequestedPath, ConvertedPath) == EContentBrowserPathType::Internal)
+		{
+			if (FPackageName::IsValidLongPackageName(WriteToString<256>(ConvertedPath)))
+			{
+				PackageName = ConvertedPath.ToString();
+				Item = GetItemFromPackageName(PackageName);
+				if (Item.IsValid())
+				{
+					return Item;
+				}
+			}
+		}
+	}
+
+	// If the string is a filesystem path to a package, sync to that asset
+	FString PackageName;
+	if (FPackageName::IsValidLongPackageName(PackageName) || FPackageName::TryConvertFilenameToLongPackageName(FString(RequestedPathView), PackageName))
+	{
+		Item = GetItemFromPackageName(PackageName);
+		if (Item.IsValid())
+		{
+			return Item;
+		}
+	}
+
+	// Try and remove elements from the end of the path until it's a valid virtual path 
+	FPathViews::IterateAncestors(RequestedPathView, [RequestedPathView, ContentBrowserData, &Item](FStringView InAncestor){
+		if (RequestedPathView == InAncestor)
+		{
+			return true;
+		}
+		FName AncestorName(InAncestor);
+		Item = ContentBrowserData->GetItemAtPath(AncestorName, EContentBrowserItemTypeFilter::IncludeFolders);
+		if (Item.IsValid())
+		{
+			return false;
+		}
+		return true;
+	});
+	if (Item.IsValid())
+	{
+		return Item;
 	}
 
 	return FContentBrowserItem();
