Only register OnShouldFilterItem callback when necessary to avoid extra cost in filtering.

diff --git a/Engine/Source/Editor/ContentBrowser/Private/SAssetView.cpp b/Engine/Source/Editor/ContentBrowser/Private/SAssetView.cpp
--- a/Engine/Source/Editor/ContentBrowser/Private/SAssetView.cpp
+++ b/Engine/Source/Editor/ContentBrowser/Private/SAssetView.cpp
@@ -5042,6 +5042,12 @@ void SAssetView::SetFilterBar(TSharedPtr<SFilterList> InFilterBar)
 	FilterBar = InFilterBar;
 }
 
+void SAssetView::SetShouldFilterItem(FOnShouldFilterItem InCallback)
+{
+	OnShouldFilterItem = MoveTemp(InCallback);
+	RequestQuickFrontendListRefresh();
+}
+
 void SAssetView::OnCompleteFiltering(double InAmortizeDuration)
 {
 	CurrentFrontendFilterTelemetry.AmortizeDuration = InAmortizeDuration;
diff --git a/Engine/Source/Editor/ContentBrowser/Private/SContentBrowser.cpp b/Engine/Source/Editor/ContentBrowser/Private/SContentBrowser.cpp
--- a/Engine/Source/Editor/ContentBrowser/Private/SContentBrowser.cpp
+++ b/Engine/Source/Editor/ContentBrowser/Private/SContentBrowser.cpp
@@ -281,6 +281,8 @@ namespace ContentBrowser
 
 SContentBrowser::~SContentBrowser()
 {
+	IConsoleManager::Get().UnregisterConsoleVariableSink_Handle(CVarSinkHandle);
+
 	// Remove the listener for when view settings are changed
 	UContentBrowserSettings::OnSettingChanged().RemoveAll( this );
 
@@ -372,9 +374,13 @@ void SContentBrowser::Construct( const FArguments& InArgs, const FName& InInstan
 	BindCommands();
 	UContentBrowserSettings::OnSettingChanged().AddSP(this, &SContentBrowser::OnContentBrowserSettingsChanged);
 
-	// Currently this controls the asset count 
+	// Currently this controls the asset count
 	const bool bShowBottomToolbar = Config != nullptr ? Config->bShowBottomToolbar : true;
 
+	// Register console variable sink for private content setting changing
+	CVarSinkHandle = IConsoleManager::Get().RegisterConsoleVariableSink_Handle(FConsoleCommandDelegate::CreateSP(this, &SContentBrowser::OnConsoleVariableChanged));
+	static const IConsoleVariable* EnablePublicAssetFeatureCVar = IConsoleManager::Get().FindConsoleVariable(TEXT("AssetTools.EnablePublicAssetFeature"));
+	const bool bEnablePrivateContentFeature = EnablePublicAssetFeatureCVar && EnablePublicAssetFeatureCVar->GetBool();
 	AssetViewPtr = SNew(SAssetView)
 					   .ThumbnailLabel(Config != nullptr ? Config->ThumbnailLabel : EThumbnailLabel::ClassName)
 					   //.ThumbnailScale(Config != nullptr ? Config->ThumbnailScale : 0.18f)
@@ -384,7 +390,7 @@ void SContentBrowser::Construct( const FArguments& InArgs, const FName& InInstan
 					   .OnItemsActivated(this, &SContentBrowser::OnItemsActivated)
 					   .OnGetItemContextMenu(this, &SContentBrowser::GetItemContextMenu, EContentBrowserViewContext::AssetView)
 					   .OnItemRenameCommitted(this, &SContentBrowser::OnItemRenameCommitted)
-					   .OnShouldFilterItem(this, &SContentBrowser::HandlePrivateContentFilter)
+					   .OnShouldFilterItem(bEnablePrivateContentFeature ? FOnShouldFilterItem::CreateSP(this, &SContentBrowser::HandlePrivateContentFilter) : FOnShouldFilterItem())
 					   .FrontendFilters(FrontendFilters)
 					   .ShowRedirectors_Lambda([this]() { return ContentBrowserUtils::ShouldShowRedirectors(FilterListPtr); })
 					   .HighlightedText(this, &SContentBrowser::GetHighlightedText)
@@ -3234,6 +3240,13 @@ void SContentBrowser::OnContentBrowserSettingsChanged(FName PropertyName)
 	}
 }
 
+void SContentBrowser::OnConsoleVariableChanged()
+{
+	static const IConsoleVariable* EnablePublicAssetFeatureCVar = IConsoleManager::Get().FindConsoleVariable(TEXT("AssetTools.EnablePublicAssetFeature"));
+	const bool bEnablePrivateContentFeature = EnablePublicAssetFeatureCVar && EnablePublicAssetFeatureCVar->GetBool();
+	AssetViewPtr->SetShouldFilterItem(bEnablePrivateContentFeature ? FOnShouldFilterItem::CreateSP(this, &SContentBrowser::HandlePrivateContentFilter) : FOnShouldFilterItem());
+}
+
 FReply SContentBrowser::BackClicked()
 {
 	HistoryManager.GoBack();
diff --git a/Engine/Source/Editor/ContentBrowser/Private/SContentBrowser.h b/Engine/Source/Editor/ContentBrowser/Private/SContentBrowser.h
--- a/Engine/Source/Editor/ContentBrowser/Private/SContentBrowser.h
+++ b/Engine/Source/Editor/ContentBrowser/Private/SContentBrowser.h
@@ -78,6 +78,7 @@ enum class EContentBrowserViewContext : uint8
 class SContentBrowser
 	: public SCompoundWidget
 {
+	using Super = SCompoundWidget;
 public:
 
 	SLATE_BEGIN_ARGS( SContentBrowser )
@@ -410,6 +411,12 @@ private:
 
 	/** Called to handle the Content Browser settings changing */
 	void OnContentBrowserSettingsChanged(FName PropertyName);
+	
+	/**
+	 * Called when engine Cvars change to update content browser settings
+	 * e.g. private content feature
+	 */
+	void OnConsoleVariableChanged();
 
 	/** Handler for clicking the history back button */
 	FReply BackClicked();
@@ -752,6 +759,8 @@ private:
 	/** When viewing a dynamic collection, the active search query will be stashed in this variable so that it can be restored again later */
 	TOptional<FText> StashedSearchBoxText;
 
+	/** Handle to a callback for when CVars change */
+	FConsoleVariableSinkHandle CVarSinkHandle;
 public: 
 
 	/** The section of EditorPerProjectUserSettings in which to save content browser settings */
diff --git a/Engine/Source/Editor/ContentBrowser/Public/SAssetView.h b/Engine/Source/Editor/ContentBrowser/Public/SAssetView.h
--- a/Engine/Source/Editor/ContentBrowser/Public/SAssetView.h
+++ b/Engine/Source/Editor/ContentBrowser/Public/SAssetView.h
@@ -448,6 +448,8 @@ public:
 	/** Set the filter list attached to this asset view - allows toggling of the the filter bar layout from the view options */
 	void SetFilterBar(TSharedPtr<SFilterList> InFilterBar);
 
+	/** Change the delegate bound via the widget argument OnShouldFilterAsset after construction. */
+	void SetShouldFilterItem(FOnShouldFilterItem InCallback);
 private:
 
 	/** Sets the pending selection to the current selection (used when changing views or refreshing the view). */
