Remove selected paths from non-per-workspace editor config.

diff --git a/Engine/Source/Editor/ContentBrowser/Private/ContentBrowserConfig.h b/Engine/Source/Editor/ContentBrowser/Private/ContentBrowserConfig.h
--- a/Engine/Source/Editor/ContentBrowser/Private/ContentBrowserConfig.h
+++ b/Engine/Source/Editor/ContentBrowser/Private/ContentBrowserConfig.h
@@ -35,9 +35,6 @@ struct FPathViewConfig
 	UPROPERTY()
 	bool bExpanded = false;
 
-	UPROPERTY()
-	TArray<FName> SelectedPaths;
-
 	UPROPERTY()
 	TArray<FString> PluginFilters;
 };
diff --git a/Engine/Source/Editor/ContentBrowser/Private/SContentBrowser.cpp b/Engine/Source/Editor/ContentBrowser/Private/SContentBrowser.cpp
--- a/Engine/Source/Editor/ContentBrowser/Private/SContentBrowser.cpp
+++ b/Engine/Source/Editor/ContentBrowser/Private/SContentBrowser.cpp
@@ -2228,6 +2228,8 @@ void SContentBrowser::PathSelected(const FString& FolderPath)
 	TArray<FCollectionNameType> SelectedCollections;
 	SourcesChanged(SelectedPaths, SelectedCollections);
 
+	PathViewPtr->SaveSettings(GEditorPerProjectIni, SettingsIniSection, InstanceName.ToString());
+
 	// Notify 'asset path changed' delegate
 	FContentBrowserModule& ContentBrowserModule = FModuleManager::GetModuleChecked<FContentBrowserModule>( TEXT("ContentBrowser") );
 	FContentBrowserModule::FOnAssetPathChanged& PathChangedDelegate = ContentBrowserModule.GetOnAssetPathChanged();
@@ -2253,6 +2255,8 @@ void SContentBrowser::FavoritePathSelected(const FString& FolderPath)
 	TArray<FCollectionNameType> SelectedCollections;
 	SourcesChanged(SelectedPaths, SelectedCollections);
 
+	PathViewPtr->SaveSettings(GEditorPerProjectIni, SettingsIniSection, InstanceName.ToString());
+
 	// Notify 'asset path changed' delegate
 	FContentBrowserModule& ContentBrowserModule = FModuleManager::GetModuleChecked<FContentBrowserModule>(TEXT("ContentBrowser"));
 	FContentBrowserModule::FOnAssetPathChanged& PathChangedDelegate = ContentBrowserModule.GetOnAssetPathChanged();
diff --git a/Engine/Source/Editor/ContentBrowser/Private/SPathView.cpp b/Engine/Source/Editor/ContentBrowser/Private/SPathView.cpp
--- a/Engine/Source/Editor/ContentBrowser/Private/SPathView.cpp
+++ b/Engine/Source/Editor/ContentBrowser/Private/SPathView.cpp
@@ -1290,13 +1290,6 @@ void SPathView::SetSelectedPaths(const TArray<FString>& Paths)
 			TreeViewPtr->RequestScrollIntoView(BestItem);
 		}
 	}
-
-	if (FPathViewConfig* PathViewConfig = GetPathViewConfig())
-	{
-		PathViewConfig->SelectedPaths = LastSelectedPaths.Array();
-
-		UContentBrowserConfig::Get()->SaveEditorConfig();
-	}
 }
 
 void SPathView::ClearSelection()
@@ -1703,11 +1696,6 @@ void SPathView::LoadSettings(const FString& IniFilename, const FString& IniSecti
 {
 	// Selected Paths
 	TArray<FName> NewSelectedPaths;
-	if (FPathViewConfig* PathViewConfig = GetPathViewConfig())
-	{
-		NewSelectedPaths = PathViewConfig->SelectedPaths;
-	}
-	else 
 	{
 		FString SelectedPathsString;
 		if (GConfig->GetString(*IniSection, *(SettingsString + TEXT(".SelectedPaths")), SelectedPathsString, IniFilename))
@@ -1958,13 +1946,6 @@ void SPathView::TreeSelectionChanged( TSharedPtr< FTreeItem > TreeItem, ESelectI
 				OnItemSelectionChanged.Execute(FContentBrowserItem(), SelectInfo);
 			}
 		}
-
-		if (FPathViewConfig* PathViewConfig = GetPathViewConfig())
-		{
-			PathViewConfig->SelectedPaths = LastSelectedPaths.Array();
-
-			UContentBrowserConfig::Get()->SaveEditorConfig();
-		}
 	}
 
 	if (TreeItem.IsValid())
