Disable plugin path filters when trying to sync the content browser into a folder they hide

diff --git a/Engine/Source/Editor/ContentBrowser/Private/SContentBrowser.cpp b/Engine/Source/Editor/ContentBrowser/Private/SContentBrowser.cpp
--- a/Engine/Source/Editor/ContentBrowser/Private/SContentBrowser.cpp
+++ b/Engine/Source/Editor/ContentBrowser/Private/SContentBrowser.cpp
@@ -1631,6 +1631,14 @@ void SContentBrowser::PrepareToSyncItems(TArrayView<const FContentBrowserItem> I
 		}
 	}
 
+	// Disable any plugin filters which hide the path we're navigating too in the path tree
+	bool bSomePluginPathFiltersChanged = false;
+	if (PathViewPtr->DisablePluginPathFiltersThatHideItems(ItemsToSync))
+	{
+		bSomePluginPathFiltersChanged = true;
+		bRepopulate = true;
+	}
+
 	// Check to see if any item paths don't exist (this can happen if we haven't ticked since the path was created)
 	if (!bRepopulate)
 	{
@@ -1650,7 +1658,7 @@ void SContentBrowser::PrepareToSyncItems(TArrayView<const FContentBrowserItem> I
 	if (bRepopulate)
 	{
 		// let the user know if one of their settings is being changed to be able to show the sync targets
-		if (bDisplayDevChanged || bDisplayEngineChanged || bDisplayPluginsChanged || bDisplayLocalizedChanged)
+		if (bDisplayDevChanged || bDisplayEngineChanged || bDisplayPluginsChanged || bDisplayLocalizedChanged || bSomePluginPathFiltersChanged)
 		{
 			TArray<FText> SettingsText;
 			if (bDisplayDevChanged)
@@ -1669,9 +1677,13 @@ void SContentBrowser::PrepareToSyncItems(TArrayView<const FContentBrowserItem> I
 			{
 				SettingsText.Add(LOCTEXT("ShowLocalizedContent", "Show Localized Content"));
 			}
+			if (bSomePluginPathFiltersChanged)
+			{
+				SettingsText.Add(LOCTEXT("SomePluginPathFilters", "Some Plugin Filters"));
+			}
 			FTextBuilder NotificationBuilder;
 			const FText NotificationPrefix = FText::Format(
-				LOCTEXT("AssetRequiresFilterChanges", "To show {0}|plural(one=this asset,other=these assets), the following {1}|plural(one=setting has,other=settings have) been enabled for the active Content Browser:\n"),
+				LOCTEXT("AssetRequiresFilterChanges", "To show {0}|plural(one=this asset,other=these assets), the following {1}|plural(one=setting has,other=settings have) been changed for the active Content Browser:\n"),
 				ItemsToSync.Num(),
 				SettingsText.Num());
 			NotificationBuilder.AppendLine(NotificationPrefix);
diff --git a/Engine/Source/Editor/ContentBrowser/Private/SPathView.cpp b/Engine/Source/Editor/ContentBrowser/Private/SPathView.cpp
--- a/Engine/Source/Editor/ContentBrowser/Private/SPathView.cpp
+++ b/Engine/Source/Editor/ContentBrowser/Private/SPathView.cpp
@@ -1179,11 +1179,14 @@ void SPathView::PluginPathFilterClicked(TSharedRef<FContentBrowserPluginFilter>
 
 bool SPathView::IsPluginPathFilterInUse(TSharedRef<FContentBrowserPluginFilter> Filter) const
 {
-	for (int32 i=0; i < PluginPathFilters->Num(); ++i)
+	if (PluginPathFilters.IsValid())
 	{
-		if (PluginPathFilters->GetFilterAtIndex(i) == Filter)
+		for (int32 i=0; i < PluginPathFilters->Num(); ++i)
 		{
-			return true;
+			if (PluginPathFilters->GetFilterAtIndex(i) == Filter)
+			{
+				return true;
+			}
 		}
 	}
 
@@ -1200,6 +1203,46 @@ void SPathView::ResetPluginPathFilters()
 	Populate();
 }
 
+bool SPathView::DisablePluginPathFiltersThatHideItems(TConstArrayView<FContentBrowserItem> Items)
+{
+	if (!PluginPathFilters.IsValid())
+	{
+		return false;
+	}
+
+	TSet<TSharedRef<IPlugin>> RelevantPlugins;
+
+	for (const FContentBrowserItem& Item : Items)
+	{
+		FName InternalPath = Item.GetInternalPath();
+		if (InternalPath.IsNone())
+		{
+			continue;
+		}
+		TStringBuilder<256> PathBuffer(InPlace, InternalPath);
+		const FStringView MountPoint = FPathViews::GetMountPointNameFromPath(PathBuffer);
+		if (TSharedPtr<IPlugin> Plugin = IPluginManager::Get().FindPlugin(MountPoint))
+		{
+			RelevantPlugins.Add(Plugin.ToSharedRef());
+		}
+	}
+
+	bool bAnyChanges = false;
+	for (const TSharedRef<FContentBrowserPluginFilter>& Filter : AllPluginPathFilters)
+	{
+		if (Algo::AnyOf(RelevantPlugins, [&Filter](const TSharedRef<IPlugin>& Plugin) { return !Filter->PassesFilter(Plugin); }))
+		{
+			// Whether the filter is inverse or not, we don't want it in the list
+			if (IsPluginPathFilterInUse(Filter))
+			{
+				SetPluginPathFilterActive(Filter, Filter->IsInverseFilter());
+				bAnyChanges = true;
+			}
+		}
+	}
+	return bAnyChanges;
+}
+
 void SPathView::SetPluginPathFilterActive(const TSharedRef<FContentBrowserPluginFilter>& Filter, bool bActive)
 {
 	if (Filter->IsInverseFilter())
diff --git a/Engine/Source/Editor/ContentBrowser/Private/SPathView.h b/Engine/Source/Editor/ContentBrowser/Private/SPathView.h
--- a/Engine/Source/Editor/ContentBrowser/Private/SPathView.h
+++ b/Engine/Source/Editor/ContentBrowser/Private/SPathView.h
@@ -240,6 +240,12 @@ public:
 	*/
 	bool InternalPathPassesBlockLists(const FStringView InInternalPath, const int32 InAlreadyCheckedDepth = 0) const;
 
+	/** 
+	 * Disable any filters which would prevent us from syncing to the given items in the content browser.
+	 * Returns true if any filtering changed. 
+	 */
+	bool DisablePluginPathFiltersThatHideItems(TConstArrayView<FContentBrowserItem> Items);
+
 	/** Populates the tree with all folders that are not filtered out */
 	virtual void Populate(const bool bIsRefreshingFilter = false);
 
