Compile in verbose logging for content browser.
Add log for plugin filter activation/deactivation.

diff --git a/Engine/Source/Editor/ContentBrowser/Private/ContentBrowserLog.h b/Engine/Source/Editor/ContentBrowser/Private/ContentBrowserLog.h
--- a/Engine/Source/Editor/ContentBrowser/Private/ContentBrowserLog.h
+++ b/Engine/Source/Editor/ContentBrowser/Private/ContentBrowserLog.h
@@ -4,4 +4,4 @@
 
 #include "CoreMinimal.h"
 
-DECLARE_LOG_CATEGORY_EXTERN(LogContentBrowser, Log, Log);
+DECLARE_LOG_CATEGORY_EXTERN(LogContentBrowser, Log, Verbose);
diff --git a/Engine/Source/Editor/ContentBrowser/Private/SPathView.cpp b/Engine/Source/Editor/ContentBrowser/Private/SPathView.cpp
--- a/Engine/Source/Editor/ContentBrowser/Private/SPathView.cpp
+++ b/Engine/Source/Editor/ContentBrowser/Private/SPathView.cpp
@@ -1251,6 +1251,11 @@ void SPathView::SetPluginPathFilterActive(const TSharedRef<FContentBrowserPlugin
 		bActive = !bActive;
 	}
 
+	UE_LOG(LogContentBrowser, Verbose, TEXT("[%s] Setting%s plugin filter %s to %s"), 
+		*WriteToString<64>(OwningContentBrowserName), 
+		Filter->IsInverseFilter() ? TEXT(" inverse") : TEXT(""),
+		*Filter->GetName(), bActive ? TEXT("Active") : TEXT("Inactive"));
+
 	Filter->ActiveStateChanged(bActive);
 
 	if (bActive)
