Fix selection of default path for SPathPicker

diff --git a/Engine/Source/Editor/ContentBrowser/Private/SPathPicker.cpp b/Engine/Source/Editor/ContentBrowser/Private/SPathPicker.cpp
--- a/Engine/Source/Editor/ContentBrowser/Private/SPathPicker.cpp
+++ b/Engine/Source/Editor/ContentBrowser/Private/SPathPicker.cpp
@@ -55,7 +55,8 @@ void SPathPicker::Construct( const FArguments& InArgs )
 		.SelectionMode(ESelectionMode::Single)
 		.CustomFolderPermissionList(InArgs._PathPickerConfig.CustomFolderPermissionList)
 		.ShowFavorites(InArgs._PathPickerConfig.bShowFavorites)
-		.DefaultPath(InArgs._PathPickerConfig.bAddDefaultPath ? InArgs._PathPickerConfig.DefaultPath : FString{})
+		.DefaultPath(InArgs._PathPickerConfig.DefaultPath)
+		.CreateDefaultPath(InArgs._PathPickerConfig.bAddDefaultPath)
 		.AllowClassesFolder(InArgs._PathPickerConfig.bAllowClassesFolder)
 		.CanShowDevelopersFolder(InArgs._PathPickerConfig.bCanShowDevelopersFolder)
 		.ForceShowEngineContent(InArgs._PathPickerConfig.bForceShowEngineContent)
diff --git a/Engine/Source/Editor/ContentBrowser/Private/SPathView.cpp b/Engine/Source/Editor/ContentBrowser/Private/SPathView.cpp
--- a/Engine/Source/Editor/ContentBrowser/Private/SPathView.cpp
+++ b/Engine/Source/Editor/ContentBrowser/Private/SPathView.cpp
@@ -1098,7 +1098,7 @@ void SPathView::Construct( const FArguments& InArgs )
 		{
 			InternalPath = FName();
 		}
-		if (!TreeData->FindTreeItem(VirtualPath))
+		if (InArgs._CreateDefaultPath && !TreeData->FindTreeItem(VirtualPath))
 		{
 			const FString DefaultPathLeafName = FPaths::GetPathLeaf(VirtualPath.ToString());
 			TreeData->AddFolderItem(FContentBrowserItemData(nullptr,
diff --git a/Engine/Source/Editor/ContentBrowser/Private/SPathView.h b/Engine/Source/Editor/ContentBrowser/Private/SPathView.h
--- a/Engine/Source/Editor/ContentBrowser/Private/SPathView.h
+++ b/Engine/Source/Editor/ContentBrowser/Private/SPathView.h
@@ -151,9 +151,12 @@ public:
 		/** The instance name of the owning content browser. */
 		SLATE_ARGUMENT( FName, OwningContentBrowserName )
 
-		/** Default path to create even if it doesn't exist for use by path picker */
+		/** Default path to select use by path picker */
 		SLATE_ARGUMENT(FString, DefaultPath)
 
+		/** If DefaultPath doesn't exist, create it */
+		SLATE_ARGUMENT(bool, CreateDefaultPath)
+
 	SLATE_END_ARGS()
 
 	/** Destructor */
