Content browser navigation bar updates:
- Ensure some empty space is available to click to edit
- Select all text at first when editing

diff --git a/Engine/Source/Editor/ContentBrowser/Private/SNavigationBar.cpp b/Engine/Source/Editor/ContentBrowser/Private/SNavigationBar.cpp
--- a/Engine/Source/Editor/ContentBrowser/Private/SNavigationBar.cpp
+++ b/Engine/Source/Editor/ContentBrowser/Private/SNavigationBar.cpp
@@ -133,29 +133,43 @@ void SNavigationBar::Construct(const FArguments& InArgs)
 			.HAlign(HAlign_Fill)
 			[
 				SNew(SOverlay)
-				// Invisible button under visible controls to handle clicking in blank space
+				// Breadcrumb trail aligned to left in combo box
 				+ SOverlay::Slot()
 				.HAlign(HAlign_Fill)
 				.VAlign(VAlign_Fill)
 				[
+					// Invisible button under visible controls to handle clicking in blank space
 					SNew(SButton)
 					.OnClicked(this, &SNavigationBar::HandleBlankSpaceClicked)	
 					.ButtonStyle( FAppStyle::Get(), "NoBorder" )
-				]
-				// Breadcrumb trail aligned to left in combo box
-				+ SOverlay::Slot()
-				.HAlign(HAlign_Left)
-				.VAlign(VAlign_Fill)
-				[
-					SAssignNew(BreadcrumbBar, SBreadcrumbTrail<FNavigationCrumb>)
-					.Visibility(this, &SNavigationBar::GetNonEditVisibility)
-					.TextStyle(InArgs._BreadcrumbTextStyle)
-					.ButtonStyle(InArgs._BreadcrumbButtonStyle)
-					.ButtonContentPadding(InArgs._BreadcrumbButtonContentPadding)
-					.DelimiterImage(InArgs._BreadcrumbDelimiterImage)
-					.OnCrumbClicked(this, &SNavigationBar::HandleCrumbClicked)
-					.HasCrumbMenuContent(this, &SNavigationBar::HandleHasCrumbMenuContent)
-					.GetCrumbMenuContent(this, &SNavigationBar::HandleGetCrumbMenuContent)
+					[
+						SNew(SHorizontalBox)
+						+ SHorizontalBox::Slot()
+						.HAlign(HAlign_Fill)
+						.VAlign(VAlign_Fill)
+						.FillWidth(1.0f)
+						[
+							SAssignNew(BreadcrumbBar, SBreadcrumbTrail<FNavigationCrumb>)
+							.Visibility(this, &SNavigationBar::GetNonEditVisibility)
+							.TextStyle(InArgs._BreadcrumbTextStyle)
+							.ButtonStyle(InArgs._BreadcrumbButtonStyle)
+							.ButtonContentPadding(InArgs._BreadcrumbButtonContentPadding)
+							.DelimiterImage(InArgs._BreadcrumbDelimiterImage)
+							.OnCrumbClicked(this, &SNavigationBar::HandleCrumbClicked)
+							.HasCrumbMenuContent(this, &SNavigationBar::HandleHasCrumbMenuContent)
+							.GetCrumbMenuContent(this, &SNavigationBar::HandleGetCrumbMenuContent)
+						]
+						// Spacer to allow some blank space to enter edit mode always
+						+ SHorizontalBox::Slot()
+						.HAlign(HAlign_Right)
+						.VAlign(VAlign_Fill)
+						.AutoWidth()
+						[
+							SNew(SImage)
+							.Image(FAppStyle::GetNoBrush())
+							.DesiredSizeOverride(FVector2D(10.0f, 1.0f))
+						]
+					]
 				]
 				// Editable text box taking up all space
 				+ SOverlay::Slot()
@@ -166,7 +180,7 @@ void SNavigationBar::Construct(const FArguments& InArgs)
 					.Visibility(this, &SNavigationBar::GetEditTextVisibility)
 					.OnTextCommitted(this, &SNavigationBar::HandleTextCommitted)
 					.OnTextChanged(this, &SNavigationBar::HandleTextChanged)
-					.SelectAllTextWhenFocused(false)
+					.SelectAllTextWhenFocused(true) 
 					.ClearKeyboardFocusOnCommit(true)
 					.OnKeyDownHandler(this, &SNavigationBar::HandleEditableTextKeyDown)
 				]
@@ -214,6 +228,7 @@ void SNavigationBar::Construct(const FArguments& InArgs)
 		]
 	);
 
+	BreadcrumbBar->ScrollToEnd();
 	ComboListView->SetBackgroundBrush(FStyleDefaults::GetNoBrush());
 	EditableText->SetTextBlockStyle(&TextBoxStyle->TextStyle);
 	SetMenuContentWidgetToFocus(ComboListView);
@@ -227,6 +242,7 @@ void SNavigationBar::ClearPaths()
 void SNavigationBar::PushPath(const FText& ElementText, const FString& FullPath, bool bHasChildren)
 {
 	BreadcrumbBar->PushCrumb(ElementText, FNavigationCrumb{FullPath, bHasChildren});	
+	BreadcrumbBar->ScrollToEnd();
 }
 
 void SNavigationBar::HandleCrumbClicked(const FNavigationCrumb& Crumb)
@@ -316,6 +332,7 @@ FReply SNavigationBar::HandleComboKeyChar(const FGeometry& MyGeometry, const FCh
 		FString NewText = SelectedItems[0]->VirtualPath;
 		NewText += InCharEvent.GetCharacter();
 		EditableText->SetText(FText::FromString(NewText));
+		EditableText->SetSelectAllTextWhenFocused(false);
 
 		return FReply::Handled().SetUserFocus(EditableText.ToSharedRef());
 	}
@@ -360,6 +377,7 @@ FReply SNavigationBar::HandleComboKeyDown(const FGeometry& MyGeometry, const FKe
 		{
 			FString NewText = SelectedItems[0]->VirtualPath + TEXT("/"); // TODO: Ideally only add slash if this prefix has children itself 
 			EditableText->SetText(FText::FromString(NewText));
+			EditableText->SetSelectAllTextWhenFocused(false);
 			return FReply::Handled().SetUserFocus(EditableText.ToSharedRef(), EFocusCause::SetDirectly);
 		}
 		return FReply::Handled(); // Swallow event to avoid confusion with items being selected or not 
@@ -423,6 +441,7 @@ void SNavigationBar::StartEditingPath()
 			Text = FText::FromString(Crumb.Path);
 		}
 	}
+	EditableText->SetSelectAllTextWhenFocused(true);
 	EditableText->SetText(Text);
 	FSlateApplication::Get().SetKeyboardFocus(EditableText.ToSharedRef(), EFocusCause::SetDirectly);
 }
@@ -451,6 +470,7 @@ void SNavigationBar::HandleTextCommitted(const FText& InText, ETextCommit::Type
 	}
 
 	// Stop editing and navigate to new path
+	EditableText->SetSelectAllTextWhenFocused(true);
 	EditTextVisibility = EVisibility::Hidden;
 	SetIsOpen(false, false);
 	OnNavigateToPath.ExecuteIfBound(InText.ToString());
diff --git a/Engine/Source/Runtime/Slate/Private/Widgets/Navigation/SBreadcrumbTrail.cpp b/Engine/Source/Runtime/Slate/Private/Widgets/Navigation/SBreadcrumbTrail.cpp
new file mode 100644
--- /dev/null
+++ b/Engine/Source/Runtime/Slate/Private/Widgets/Navigation/SBreadcrumbTrail.cpp
@@ -0,0 +1,13 @@
+// Copyright Epic Games, Inc. All Rights Reserved.
+
+#include "Widgets/Navigation/SBreadcrumbTrail.h"
+
+void SBreadcrumbTrailBase::ScrollToStart()
+{
+    CrumbBox->ScrollToStart();
+}
+
+void SBreadcrumbTrailBase::ScrollToEnd()
+{
+    CrumbBox->ScrollToEnd();
+}
\ No newline at end of file
diff --git a/Engine/Source/Runtime/Slate/Public/Widgets/Navigation/SBreadcrumbTrail.h b/Engine/Source/Runtime/Slate/Public/Widgets/Navigation/SBreadcrumbTrail.h
--- a/Engine/Source/Runtime/Slate/Public/Widgets/Navigation/SBreadcrumbTrail.h
+++ b/Engine/Source/Runtime/Slate/Public/Widgets/Navigation/SBreadcrumbTrail.h
@@ -20,11 +20,23 @@
 #include "Widgets/Input/SButton.h"
 #include "Widgets/Layout/SScrollBox.h"
 
+// Base class for breadcrumb trail holding methods which are not dependent on the crumb data type.
+class SBreadcrumbTrailBase : public SCompoundWidget
+{
+public:
+	SLATE_API void ScrollToStart();
+	SLATE_API void ScrollToEnd();
+
+protected:
+	/** The horizontal box which contains all the breadcrumbs */
+	TSharedPtr<SScrollBox> CrumbBox;
+};
+
 /**
  * A breadcrumb trail. Allows the user to see their currently selected path and navigate upwards.
  */
 template <typename ItemType>
-class SBreadcrumbTrail : public SCompoundWidget
+class SBreadcrumbTrail : public SBreadcrumbTrailBase
 {
 private:
 	/** A container for data associated with a single crumb in the trail. */
@@ -420,9 +432,6 @@ private:
 
 private:
 
-	/** The horizontal box which contains all the breadcrumbs */
-	TSharedPtr<SScrollBox> CrumbBox;
-
 	/** The list of crumbs and their data */
 	TArray<FCrumbItem> CrumbList;
 
