Add additional logging for path view.
Add console command PathView.Repopulate to force rebuild of path view.

diff --git a/Engine/Source/Editor/ContentBrowser/Private/ContentBrowserLog.h b/Engine/Source/Editor/ContentBrowser/Private/ContentBrowserLog.h
--- a/Engine/Source/Editor/ContentBrowser/Private/ContentBrowserLog.h
+++ b/Engine/Source/Editor/ContentBrowser/Private/ContentBrowserLog.h
@@ -4,4 +4,4 @@
 
 #include "CoreMinimal.h"
 
-DECLARE_LOG_CATEGORY_EXTERN(LogContentBrowser, Log, Verbose);
+DECLARE_LOG_CATEGORY_EXTERN(LogContentBrowser, Log, All);
diff --git a/Engine/Source/Editor/ContentBrowser/Private/SPathView.cpp b/Engine/Source/Editor/ContentBrowser/Private/SPathView.cpp
--- a/Engine/Source/Editor/ContentBrowser/Private/SPathView.cpp
+++ b/Engine/Source/Editor/ContentBrowser/Private/SPathView.cpp
@@ -87,6 +87,20 @@ struct FGeometry;
 
 namespace UE::PathView
 {
+TArray<TWeakPtr<SPathView>> AllPathViews;
+
+FAutoConsoleCommand RepopulateAllPathViewsCommand(TEXT("PathView.Repopulate"),
+	TEXT("Repopulate all path views to expose bugs with caching/data updates"),
+	FConsoleCommandDelegate::CreateLambda([]() {
+		for (TWeakPtr<SPathView> WeakView : AllPathViews)
+		{
+			if (TSharedPtr<SPathView> View = WeakView.Pin())
+			{
+				View->Populate();
+			}
+		}
+	}));
+
 TSharedRef<FTreeItem> CreateOrReuseNode(FContentBrowserItemData&& InData,
 	TMap<FName, TSharedPtr<FTreeItem>>* OldItemsByInvariantPath)
 {
@@ -360,6 +374,7 @@ protected:
 	// Mapping of full virtual path such as '/All/Game/Maps/Arena' to items
 	TMap<FName, TSharedPtr<FTreeItem>> VirtualPathToItem;
 	// Mapping of path that doesn't change based on display settings (e.g. '/MyPlugin/MyAsset') to item
+	// Used to reuse node objects when changing path view settings and rebuilding the tree 
 	TMap<FName, TSharedPtr<FTreeItem>> InvariantPathToItem;
 
 	// Used for retrieving saved settings per content browser instance
@@ -413,10 +428,8 @@ void FPathViewData::PopulateFullFolderTree(const FContentBrowserDataCompiledFilt
 			UContentBrowserDataSource* Source = InItemData.GetOwnerDataSource();
 			if (Source && !Source->IsFolderVisible(InItemData.GetVirtualPath(), EmptyFilter.FolderFlags, EmptyFilter.FolderFilter))
 			{
-				UE_LOG(LogContentBrowser,
-					VeryVerbose,
-					TEXT("Hiding folder %s that fails current pre-text filtering"),
-					*WriteToString<256>(InItemData.GetVirtualPath()));
+				UE_LOG(LogContentBrowser, VeryVerbose, TEXT("Hiding folder %s on source %s that fails current pre-text filtering"),
+					*WriteToString<256>(InItemData.GetVirtualPath()), *WriteToString<256>(Source->GetFName()));
 				return true; // continue enumerating
 			}
 
@@ -706,6 +719,7 @@ TSharedRef<FTreeItem> FPathViewData::AddFolderItemInternal(FContentBrowserItemDa
 		return LeafItem.ToSharedRef();
 	}
 
+	UContentBrowserDataSubsystem* ContentBrowserData = IContentBrowserDataModule::Get().GetSubsystem();
 	UContentBrowserDataSource* OriginalDataSource = InItemData.GetOwnerDataSource();
 
 	FName ItemInvariantPath = InItemData.GetInvariantPath();
@@ -726,7 +740,7 @@ TSharedRef<FTreeItem> FPathViewData::AddFolderItemInternal(FContentBrowserItemDa
 
 	// Work backwards from the leaf path of the requested item until we encounter an item that already existed
 	FPathViews::IterateAncestors(PathBuffer.ToView(),
-		[this, &PathBuffer, &PreviousItem, OriginalDataSource, OldItemsByInvariantPath](FStringView PathView) {
+		[this, &PathBuffer, &PreviousItem, OriginalDataSource, OldItemsByInvariantPath, &ContentBrowserData](FStringView PathView) {
 			if (PathView.Len() == PathBuffer.Len())
 			{
 				// This is the item returned by the data source, we already added it
@@ -743,20 +757,28 @@ TSharedRef<FTreeItem> FPathViewData::AddFolderItemInternal(FContentBrowserItemDa
 			bool bContinue = false;
 			if (!ParentItem.IsValid())
 			{
-				// TODO: If another data source provides this path in future, can that data source become the 'primary'
-				// ?
+				UE_LOG(LogContentBrowser, VeryVerbose, TEXT("[%s] Creating placeholder or virtual parent %.*s"), 
+					*WriteToString<256>(OwningContentBrowserName),
+					PathView.Len(), PathView.GetData());
+				// TODO: If another data source provides this path in future, can that data source become the 'primary'?
+				FName ItemName(FPathViews::GetPathLeaf(PathView)); 
+				FName InternalPath;
+				if (ContentBrowserData->TryConvertVirtualPath(ParentVirtualPath, InternalPath) != EContentBrowserPathType::Internal)
+				{
+					InternalPath = FName(); // Assuming this is a virtual path with no internal path
+				}
 				ParentItem = UE::PathView::CreateOrReuseNode(FContentBrowserItemData(OriginalDataSource,
 																 EContentBrowserItemFlags::Type_Folder,
 																 ParentVirtualPath,
-																 FName(FPathViews::GetPathLeaf(PathView)),
+																 ItemName,
 																 FText(),
 																 nullptr, 
-																 FName() // Assuming this is a virtual path with no internal path
+																 InternalPath
 																 ),
 					OldItemsByInvariantPath);
 				VirtualPathToItem.Add(ParentVirtualPath, ParentItem);
-				InvariantPathToItem.Add(
-					ParentItem->GetItem().GetInvariantPath()); // TODO: Do fully virtual paths have an invariant path?
+				// TODO: Do fully virtual paths have an invariant path?
+				InvariantPathToItem.Add(ParentItem->GetItem().GetInvariantPath()); 
 				bContinue = true;
 			}
 			ParentItem->AddChild(PreviousItem);
@@ -863,6 +885,8 @@ TSet<FName> SPathView::FScopedSelectionChangedEvent::GetSelectionSet() const
 
 SPathView::~SPathView()
 {
+	UE::PathView::AllPathViews.RemoveAllSwap([this](const TWeakPtr<SPathView> Weak) { return Weak.Pin().Get() == this; });
+
 	if (IContentBrowserDataModule* ContentBrowserDataModule = IContentBrowserDataModule::GetPtr())
 	{
 		if (UContentBrowserDataSubsystem* ContentBrowserData = ContentBrowserDataModule->GetSubsystem())
@@ -878,6 +902,8 @@ SPathView::~SPathView()
 
 void SPathView::Construct( const FArguments& InArgs )
 {
+	UE::PathView::AllPathViews.Add(SharedThis(this));
+
 	OwningContentBrowserName = InArgs._OwningContentBrowserName;
 	OnItemSelectionChanged = InArgs._OnItemSelectionChanged;
 	bAllowContextMenu = InArgs._AllowContextMenu;
@@ -1166,7 +1192,7 @@ void SPathView::PopulatePathViewFiltersMenu(UToolMenu* Menu)
 				FUIAction(
 					FExecuteAction::CreateSP(this, &SPathView::PluginPathFilterClicked, Filter),
 					FCanExecuteAction(),
-					FIsActionChecked::CreateSP(this, &SPathView::IsPluginPathFilterInUse, Filter)
+					FIsActionChecked::CreateSP(this, &SPathView::IsPluginPathFilterChecked, Filter)
 				),
 				EUserInterfaceActionType::ToggleButton
 			);
@@ -1180,6 +1206,15 @@ void SPathView::PluginPathFilterClicked(TSharedRef<FContentBrowserPluginFilter>
 	Populate();
 }
 
+bool SPathView::IsPluginPathFilterChecked(TSharedRef<FContentBrowserPluginFilter> Filter) const
+{
+	if (IsPluginPathFilterInUse(Filter))
+	{
+		return !Filter->IsInverseFilter();
+	}
+	return Filter->IsInverseFilter();
+}
+
 bool SPathView::IsPluginPathFilterInUse(TSharedRef<FContentBrowserPluginFilter> Filter) const
 {
 	if (PluginPathFilters.IsValid())
@@ -1426,6 +1461,8 @@ void SPathView::RenameFolderItem(const FContentBrowserItem& InItem)
 
 FContentBrowserDataCompiledFilter SPathView::CreateCompiledFolderFilter() const
 {
+	UE_LOG(LogContentBrowser, VeryVerbose, TEXT("[%s] Creating folder filter"), *WriteToString<256>(OwningContentBrowserName));
+
 	const UContentBrowserSettings* ContentBrowserSettings = GetDefault<UContentBrowserSettings>();
 	bool bDisplayPluginFolders = ContentBrowserSettings->GetDisplayPluginFolders();
 	// check to see if we have an instance config that overrides the default in UContentBrowserSettings
@@ -1440,6 +1477,9 @@ FContentBrowserDataCompiledFilter SPathView::CreateCompiledFolderFilter() const
 	DataFilter.ItemCategoryFilter = GetContentBrowserItemCategoryFilter();
 	DataFilter.ItemAttributeFilter = GetContentBrowserItemAttributeFilter();
 
+	UE_LOG(LogContentBrowser, VeryVerbose, TEXT("[%s] bDisplayPluginFolders:%d ItemCategoryFilter:%d ItemAttributeFilter:%d"), 
+		*WriteToString<256>(OwningContentBrowserName), bDisplayPluginFolders, DataFilter.ItemCategoryFilter, DataFilter.ItemAttributeFilter);
+
 	TSharedPtr<FPathPermissionList> CombinedFolderPermissionList = ContentBrowserUtils::GetCombinedFolderPermissionList(FolderPermissionList, bAllowReadOnlyFolders ? nullptr : WritableFolderPermissionList);
 
 	if (CustomFolderPermissionList.IsValid())
@@ -1453,6 +1493,19 @@ FContentBrowserDataCompiledFilter SPathView::CreateCompiledFolderFilter() const
 
 	if (PluginPathFilters.IsValid() && PluginPathFilters->Num() > 0 && bDisplayPluginFolders)
 	{
+		UE_SUPPRESS(LogContentBrowser, VeryVerbose, {
+			FString PluginFiltersString;
+			for (int32 i=0; i < PluginPathFilters->Num(); ++i)
+			{
+				if (i != 0)
+				{
+					PluginFiltersString += TEXT(", ");
+				}
+				PluginFiltersString += PluginPathFilters->GetFilterAtIndex(i)->GetName();
+			}
+			UE_LOG(LogContentBrowser, VeryVerbose, TEXT("[%s] Active plugin filters: %s"), 
+				*WriteToString<256>(OwningContentBrowserName), *PluginFiltersString);
+		});
 		TArray<TSharedRef<IPlugin>> Plugins = IPluginManager::Get().GetEnabledPluginsWithContent();
 		for (const TSharedRef<IPlugin>& Plugin : Plugins)
 		{
@@ -1470,10 +1523,7 @@ FContentBrowserDataCompiledFilter SPathView::CreateCompiledFolderFilter() const
 		}
 	}
 
-	UE_LOG(LogContentBrowser,
-		VeryVerbose,
-		TEXT("Compiled folder permission list: %s"),
-		*CombinedFolderPermissionList->ToString());
+	UE_LOG(LogContentBrowser, VeryVerbose, TEXT("Compiled folder permission list: %s"), CombinedFolderPermissionList.IsValid() ? *CombinedFolderPermissionList->ToString() : TEXT("null"));
 
 	ContentBrowserUtils::AppendAssetFilterToContentBrowserFilter(FARFilter(), nullptr, CombinedFolderPermissionList, DataFilter);
 
@@ -1764,6 +1814,9 @@ void SPathView::LoadSettings(const FString& IniFilename, const FString& IniSecti
 		IContentBrowserDataModule::Get().GetSubsystem()->ConvertInternalPathToVirtual(Path, Path);
 	}
 
+	UE_LOG(LogContentBrowser, Verbose, TEXT("[%s] LoadSettings: SelectedPaths: %s"), 
+		*WriteToString<256>(OwningContentBrowserName), *FString::JoinBy(NewSelectedPaths, TEXT(", "), UE_PROJECTION_MEMBER(FName, ToString)));
+
 	{
 		// Batch the selection changed event
 		FScopedSelectionChangedEvent ScopedSelectionChangedEvent(SharedThis(this));
@@ -1808,6 +1861,8 @@ void SPathView::LoadSettings(const FString& IniFilename, const FString& IniSecti
 		TArray<FString> NewSelectedFilters;
 		if (FPathViewConfig* PathViewConfig = GetPathViewConfig())
 		{
+			UE_LOG(LogContentBrowser, Verbose, TEXT("[%s] LoadSettings: Loading plugin filters from editor config: %s"), 
+				*WriteToString<256>(OwningContentBrowserName), *FString::Join(NewSelectedFilters, TEXT(", ")));
 			NewSelectedFilters = PathViewConfig->PluginFilters;
 		}
 		else
@@ -1815,6 +1870,8 @@ void SPathView::LoadSettings(const FString& IniFilename, const FString& IniSecti
 			FString PluginFiltersString;
 			if (GConfig->GetString(*IniSection, *(SettingsString + TEXT(".PluginFilters")), PluginFiltersString, IniFilename))
 			{
+				UE_LOG(LogContentBrowser, Verbose, TEXT("[%s] LoadSettings: Loading plugin filters from ini: %s"), 
+					*WriteToString<256>(OwningContentBrowserName), *PluginFiltersString);
 				PluginFiltersString.ParseIntoArray(NewSelectedFilters, TEXT(","), /*bCullEmpty*/ true);
 			}
 		}
diff --git a/Engine/Source/Editor/ContentBrowser/Private/SPathView.h b/Engine/Source/Editor/ContentBrowser/Private/SPathView.h
--- a/Engine/Source/Editor/ContentBrowser/Private/SPathView.h
+++ b/Engine/Source/Editor/ContentBrowser/Private/SPathView.h
@@ -373,7 +373,13 @@ private:
 	/** Toggle plugin filter. */
 	void PluginPathFilterClicked(TSharedRef<FContentBrowserPluginFilter> Filter);
 
-	/** Returns true if filter is being used. */
+	/** Return whether the given filter should be shown enabled in the UI. */
+	bool IsPluginPathFilterChecked(TSharedRef<FContentBrowserPluginFilter> Filter) const;
+
+	/** 
+	 * Returns true if filter is being used. Note that because some filters are 'inverse' this is not the same as whether
+	 * the filter should be visibly checked in the UI to be shown as 'enabled' to the user. 
+	 */
 	bool IsPluginPathFilterInUse(TSharedRef<FContentBrowserPluginFilter> Filter) const;
 
 	TArray<FName> GetDefaultPathsToExpand() const;
