New "Show Other Developers" filter implemented as modification to backend query to avoid processing each item in time-sliced frontend filter.
Use path permission lists to hide developer folders other than the current user.

diff --git a/Engine/Source/Developer/ToolWidgets/Public/Filters/FilterBase.h b/Engine/Source/Developer/ToolWidgets/Public/Filters/FilterBase.h
--- a/Engine/Source/Developer/ToolWidgets/Public/Filters/FilterBase.h
+++ b/Engine/Source/Developer/ToolWidgets/Public/Filters/FilterBase.h
@@ -2,6 +2,7 @@
 
 #pragma once
 
+#include "Misc/FilterCollection.h"
 #include "Misc/IFilter.h"
 #include "Internationalization/Text.h"
 #include "Math/Color.h"
@@ -66,7 +67,11 @@ public:
 	/** Can be overriden for custom FilterBar subclasses to load settings, currently not implemented in any gneeric Filter Bar */
 	virtual void LoadSettings(const FString& IniFilename, const FString& IniSection, const FString& SettingsString) = 0;
 
-	/** Set this filter as checked, and activates or deactivates it */
+	/** 
+	 * Set this filter as checked, and activates or deactivates it.
+	 * @note this only works for filters registered by a widget such as SBasicFilterBar that manages them and 
+	 *  should be called by external code rather than by the managing widget itself.
+	 */
 	void SetActive(bool bActive) { SetActiveEvent.Broadcast(bActive); }
 
 	/** Checks whether the filter is checked/pinned and active/enabled */
@@ -80,6 +85,29 @@ public:
 		return false;
 	}
 
+	/** 
+	 * Manage the state of a filter in a collection taking into account whether it is an inverse filter or not.
+	 * Adds or removes the filter from the collcetion and calls ActiveStateChanged
+	 */
+	static void SetActiveInCollection(TSharedRef<FFilterBase<FilterType>> This, bool bActive, TFilterCollection<FilterType>& Collection)
+	{
+		if(This->IsInverseFilter())
+		{
+			//Inverse filters are active when they are "disabled"
+			bActive = !bActive;
+		}
+		This->ActiveStateChanged(bActive);
+
+		if (bActive)
+		{
+			Collection.Add(This);
+		}
+		else
+		{
+			Collection.Remove(This);
+		}
+	}
+
 	// IFilter implementation
 	DECLARE_DERIVED_EVENT( FFilterBase, IFilter<FilterType>::FChangedEvent, FChangedEvent );
 	virtual FChangedEvent& OnChanged() override { return ChangedEvent; }
diff --git a/Engine/Source/Developer/ToolWidgets/Public/Filters/SBasicFilterBar.h b/Engine/Source/Developer/ToolWidgets/Public/Filters/SBasicFilterBar.h
--- a/Engine/Source/Developer/ToolWidgets/Public/Filters/SBasicFilterBar.h
+++ b/Engine/Source/Developer/ToolWidgets/Public/Filters/SBasicFilterBar.h
@@ -967,21 +967,7 @@ protected:
  	/** Sets the active state of a frontend filter. */
  	void SetFrontendFilterActive(const TSharedRef<FFilterBase<FilterType>>& Filter, bool bActive)
 	{
-		if(Filter->IsInverseFilter())
-		{
-			//Inverse filters are active when they are "disabled"
-			bActive = !bActive;
-		}
-		Filter->ActiveStateChanged(bActive);
-
-		if ( bActive )
-		{
-			ActiveFilters->Add(Filter);
-		}
-		else
-		{
-			ActiveFilters->Remove(Filter);
-		}
+		Filter->SetActiveInCollection(Filter, bActive, *ActiveFilters);
 	}
 
 	/* 'Activate' A filter by adding it to the filter bar, does not turn it on */
diff --git a/Engine/Source/Editor/ContentBrowser/Private/Filters.cpp b/Engine/Source/Editor/ContentBrowser/Private/Filters.cpp
--- a/Engine/Source/Editor/ContentBrowser/Private/Filters.cpp
+++ b/Engine/Source/Editor/ContentBrowser/Private/Filters.cpp
@@ -2,10 +2,19 @@
 
 #include "Filters.h"
 
+#include "Algo/AnyOf.h"
+#include "AssetRegistry/AssetRegistryModule.h"
+#include "ContentBrowserDataSubsystem.h"
+#include "ContentBrowserLog.h"
 #include "FrontendFilterBase.h"
+#include "IContentBrowserDataModule.h"
 
 #define LOCTEXT_NAMESPACE "ContentBrowser"
 
+/** 
+ * FFilter_ShowRedirectors 
+ */
+
 FFilter_ShowRedirectors::FFilter_ShowRedirectors(TSharedPtr<FFrontendFilterCategory> InCategory)
 	: FFrontendFilter(InCategory)
 {
@@ -47,9 +56,130 @@ void FFilter_ShowRedirectors::LoadSettings(const FString& IniFilename, const FSt
 
 }
 
-bool FFilter_ShowRedirectors::PassesFilter(FAssetFilterType InItem) const
+/** 
+ * FFilter_OtherDevelopers 
+ */
+
+FFilter_ShowOtherDevelopers::FFilter_ShowOtherDevelopers(TSharedPtr<FFrontendFilterCategory> InCategory, FName InFilterBarIdentifier)
+	: FFrontendFilter(InCategory)
+	, FilterBarIdentifier(InFilterBarIdentifier)
+	, PathPermissionList(MakeShared<FPathPermissionList>())
+{
+	UContentBrowserDataSubsystem* ContentBrowserData = IContentBrowserDataModule::Get().GetSubsystem();
+	ItemDataUpdatedHandle = ContentBrowserData->OnItemDataUpdated().AddRaw(this, &FFilter_ShowOtherDevelopers::HandleItemDataUpdated);
+	ItemDataRefreshedHandle = ContentBrowserData->OnItemDataRefreshed().AddRaw(this, &FFilter_ShowOtherDevelopers::HandleItemDataRefreshed);
+
+	BuildFilter();
+}
+
+FFilter_ShowOtherDevelopers::~FFilter_ShowOtherDevelopers()
+{
+	UContentBrowserDataSubsystem* ContentBrowserData = IContentBrowserDataModule::Get().GetSubsystem();
+	ContentBrowserData->OnItemDataUpdated().Remove(ItemDataUpdatedHandle);
+	ContentBrowserData->OnItemDataRefreshed().Remove(ItemDataRefreshedHandle);
+}
+
+TSharedRef<const FPathPermissionList> FFilter_ShowOtherDevelopers::GetPathPermissionList() 
+{
+	return PathPermissionList;
+}
+
+void FFilter_ShowOtherDevelopers::BuildFilter()
+{
+	static const FName NAME_OtherDevelopers{"OtherDevelopers"};
+	// Update list of other developer folders and put into permission list 
+	FName BaseDevelopersPath{TEXTVIEW("/Game/Developers")};
+	FName UserDeveloperFolder(TStringBuilder<256>{ InPlace, BaseDevelopersPath, TEXTVIEW("/"), FPaths::GameUserDeveloperFolderName() });
+
+	FAssetRegistryModule& AssetRegistryModule = FModuleManager::LoadModuleChecked<FAssetRegistryModule>("AssetRegistry");
+	IAssetRegistry& AssetRegistry = AssetRegistryModule.Get();
+	TArray<FName> ChildPaths;
+	AssetRegistry.GetSubPaths(BaseDevelopersPath, ChildPaths, false);
+	TSet<FName> PreviousPaths = MoveTemp(OtherDeveloperFolders);
+	OtherDeveloperFolders.Reset();
+	OtherDeveloperFolders.Append(TSet<FName>(ChildPaths));
+	OtherDeveloperFolders.Remove(UserDeveloperFolder);
+	if (OtherDeveloperFolders.Num() != PreviousPaths.Num() || OtherDeveloperFolders.Difference(PreviousPaths).Num() != 0
+		|| PreviousPaths.Difference(OtherDeveloperFolders).Num() != 0)
+	{
+		UE_LOG(LogContentBrowser, Verbose, TEXT("[%s] FFilterShowOtherDevelopers rebuilt exclusion list: %s"), *FilterBarIdentifier.ToString(), 
+			*FString::JoinBy(OtherDeveloperFolders, TEXT(","), UE_PROJECTION_MEMBER(FName, ToString)));
+		PathPermissionList = MakeShared<FPathPermissionList>();
+		for (FName OtherPath : OtherDeveloperFolders)
+		{
+			PathPermissionList->AddDenyListItem(NAME_OtherDevelopers, WriteToString<256>(OtherPath).ToView());
+		}
+		BroadcastChangedEvent();
+	}
+	else
+	{
+		UE_LOG(LogContentBrowser, Verbose, TEXT("[%s] FFilterShowOtherDevelopers keeping previous exclusion list"), *FilterBarIdentifier.ToString());
+	}
+}
+
+void FFilter_ShowOtherDevelopers::HandleItemDataRefreshed()
+{
+	BuildFilter();
+}
+
+void FFilter_ShowOtherDevelopers::HandleItemDataUpdated(TArrayView<const FContentBrowserItemDataUpdate> InUpdatedItems)
+{
+	const bool bNeedsRebuild = Algo::AnyOf(InUpdatedItems, [](const FContentBrowserItemDataUpdate& Update){
+		FName InternalPath = Update.GetItemData().GetInternalPath();
+		if (!InternalPath.IsNone() && Update.GetItemData().IsFolder() && Update.GetUpdateType() == EContentBrowserItemUpdateType::Added)
+		{
+			if (WriteToString<256>(InternalPath).ToView().StartsWith(TEXTVIEW("/Game/Developers"), ESearchCase::IgnoreCase))
+			{
+				return true;
+			}
+		}
+		return false;
+	});	
+	if (bNeedsRebuild)
+	{
+		BuildFilter();
+	}
+}
+
+/** Returns the human readable name for this filter */
+FText FFilter_ShowOtherDevelopers::GetDisplayName() const
+{
+	return LOCTEXT("FrontendFilter_ShowOtherDevelopers", "Other Developers"); 
+}
+
+/** Returns the tooltip for this filter, shown in the filters menu */
+FText FFilter_ShowOtherDevelopers::GetToolTipText() const
+{ 
+	return LOCTEXT("FrontendFilter_ShowOtherDevelopersTooltip", "Allow display of assets in developer folders that aren't yours."); 
+}
+
+/** Returns the name of the icon to use in menu entries */
+FName FFilter_ShowOtherDevelopers::GetIconName() const
+{
+	return NAME_None;
+}
+
+/** Notification that the filter became active or inactive */
+void FFilter_ShowOtherDevelopers::ActiveStateChanged(bool bActive)
+{
+	if (bActive)
+	{
+		UE_LOG(LogContentBrowser, Verbose, TEXT("[%s] FFilterShowOtherDevelopers active, hiding content from other developers"), *FilterBarIdentifier.ToString());
+	}
+	else
+	{
+		UE_LOG(LogContentBrowser, Verbose, TEXT("[%s] FFilterShowOtherDevelopers inactive, showing content from all developers"), *FilterBarIdentifier.ToString());
+	}
+}
+
+/** Can be overriden for custom FilterBar subclasses to save settings, currently not implemented in any gneeric Filter Bar */
+void FFilter_ShowOtherDevelopers::SaveSettings(const FString& IniFilename, const FString& IniSection, const FString& SettingsString) const
+{
+}
+
+/** Can be overriden for custom FilterBar subclasses to load settings, currently not implemented in any gneeric Filter Bar */
+void FFilter_ShowOtherDevelopers::LoadSettings(const FString& IniFilename, const FString& IniSection, const FString& SettingsString)
 {
-	return true; // All items pass, this filter just communicates with the backend 
 }
 
 #undef LOCTEXT_NAMESPACE
\ No newline at end of file
diff --git a/Engine/Source/Editor/ContentBrowser/Private/Filters.h b/Engine/Source/Editor/ContentBrowser/Private/Filters.h
--- a/Engine/Source/Editor/ContentBrowser/Private/Filters.h
+++ b/Engine/Source/Editor/ContentBrowser/Private/Filters.h
@@ -5,6 +5,7 @@
 // Private filters not exposed to modules that depend on ContentBrowser
 
 #include "IContentBrowserSingleton.h"
+#include "Filters/FilterBase.h"
 #include "Filters/GenericFilter.h"
 #include "FrontendFilterBase.h"
 #include "Templates/SharedPointer.h"
@@ -50,5 +51,62 @@ public:
 	/** Can be overriden for custom FilterBar subclasses to load settings, currently not implemented in any gneeric Filter Bar */
 	virtual void LoadSettings(const FString& IniFilename, const FString& IniSection, const FString& SettingsString) override;
 
-	virtual bool PassesFilter(FAssetFilterType InItem) const override;
+	/** Pass all objects - filter is just used to set backend query state. */
+	virtual bool PassesFilter(FAssetFilterType InItem) const override { return true; }
+};
+
+// Non-frontend filter which modifies content browser backend query to exclude folders belonging to other developers 
+class FFilter_ShowOtherDevelopers : public FFrontendFilter
+{
+public: 
+	FFilter_ShowOtherDevelopers(TSharedPtr<FFrontendFilterCategory> InCategory, FName FilterBarIdentifier);
+	FFilter_ShowOtherDevelopers& operator=(const FFilter_ShowOtherDevelopers&) = delete;
+	virtual ~FFilter_ShowOtherDevelopers();
+
+	/** Get the list of folders to be denied when this filter is active (visually disabled, becuase it's an inverse filter) */
+	TSharedRef<const FPathPermissionList> GetPathPermissionList();
+
+
+	/** Pass all objects - filter is just used to set backend query state. */
+	virtual bool PassesFilter(FAssetFilterType InItem) const override { return true; }
+
+	/** Returns the system name for this filter */
+	virtual FString GetName() const override { return TEXT("ShowOtherDevelopersBackend"); }
+
+	/** Returns the human readable name for this filter */
+	virtual FText GetDisplayName() const override;
+
+	/** Returns the tooltip for this filter, shown in the filters menu */
+	virtual FText GetToolTipText() const override;
+
+	/** Returns the name of the icon to use in menu entries */
+	virtual FName GetIconName() const override;
+
+	/** If true, the filter will be active in the FilterBar when it is inactive in the UI (i.e the filter pill is grayed out) */
+	virtual bool IsInverseFilter() const override { return true; }
+
+	/** Notification that the filter became active or inactive */
+	virtual void ActiveStateChanged(bool bActive) override;
+
+	/** Called when the right-click context menu is being built for this filter */
+	virtual void ModifyContextMenu(FMenuBuilder& MenuBuilder) override {}
+	
+	/** Can be overriden for custom FilterBar subclasses to save settings, currently not implemented in any gneeric Filter Bar */
+	virtual void SaveSettings(const FString& IniFilename, const FString& IniSection, const FString& SettingsString) const override;
+
+	/** Can be overriden for custom FilterBar subclasses to load settings, currently not implemented in any gneeric Filter Bar */
+	virtual void LoadSettings(const FString& IniFilename, const FString& IniSection, const FString& SettingsString) override;
+
+private:
+	void BuildFilter();
+
+	/** Handle new folders being added to rebuild the filter */
+	void HandleItemDataUpdated(TArrayView<const FContentBrowserItemDataUpdate> InUpdatedItems);
+	void HandleItemDataRefreshed();
+
+	FName FilterBarIdentifier;
+	TSet<FName> OtherDeveloperFolders;
+	TSharedRef<FPathPermissionList> PathPermissionList;
+	FDelegateHandle ItemDataUpdatedHandle;
+	FDelegateHandle ItemDataRefreshedHandle;
 };
\ No newline at end of file
diff --git a/Engine/Source/Editor/ContentBrowser/Private/SAssetPicker.cpp b/Engine/Source/Editor/ContentBrowser/Private/SAssetPicker.cpp
--- a/Engine/Source/Editor/ContentBrowser/Private/SAssetPicker.cpp
+++ b/Engine/Source/Editor/ContentBrowser/Private/SAssetPicker.cpp
@@ -271,13 +271,13 @@ void SAssetPicker::Construct( const FArguments& InArgs )
 		];
 
 		// Use the 'other developer' filter from the filter list widget. 
-		OtherDevelopersFilter = StaticCastSharedPtr<FFrontendFilter_ShowOtherDevelopers>(FilterListPtr->GetFrontendFilter(TEXT("ShowOtherDevelopers")));
+		OtherDevelopersFilter = StaticCastSharedPtr<FFilter_ShowOtherDevelopers>(FilterListPtr->GetFrontendFilter(TEXT("ShowOtherDevelopersBackend")));
 	}
 	else
 	{
 		// Filter UI is off, but the 'other developer' filter is a built-in feature.
-		OtherDevelopersFilter = MakeShared<FFrontendFilter_ShowOtherDevelopers>(nullptr);
-		FrontendFilters->Add(OtherDevelopersFilter);
+		OtherDevelopersFilter = MakeShared<FFilter_ShowOtherDevelopers>(nullptr, FName(SaveSettingsName));
+		OtherDevelopersFilter->SetActiveInCollection(OtherDevelopersFilter.ToSharedRef(), false, *FrontendFilters);
 	}
 
 	// Make game-specific filter
@@ -542,16 +542,21 @@ void SAssetPicker::SetNewBackendFilter(const FARFilter& NewFilter)
 void SAssetPicker::OnFilterChanged()
 {
 	FARFilter Filter;
+	TArray<TSharedRef<const FPathPermissionList>> CustomPermissionLists;
 	
 	if ( FilterListPtr.IsValid() )
 	{
-		Filter = FilterListPtr->GetCombinedBackendFilter();
+		Filter = FilterListPtr->GetCombinedBackendFilter(CustomPermissionLists);
+	}
+	else if (!IsShowingOtherDevelopersContent())
+	{
+		CustomPermissionLists.Add(OtherDevelopersFilter->GetPathPermissionList());
 	}
 
 	Filter.Append(CurrentBackendFilter);
 	if (AssetViewPtr.IsValid())
 	{
-		AssetViewPtr->SetBackendFilter( Filter );
+		AssetViewPtr->SetBackendFilter(Filter, &CustomPermissionLists);
 	}
 }
 
@@ -656,10 +661,30 @@ void SAssetPicker::RefreshAssetView(bool bRefreshSources)
 	}
 }
 
+bool SAssetPicker::IsShowingOtherDevelopersContent() const
+{
+	if (FilterListPtr.IsValid())
+	{
+		return FilterListPtr->IsFrontendFilterActive(OtherDevelopersFilter);
+	}
+	else
+	{
+		for (int32 i=0; i < FrontendFilters->Num(); ++i)
+		{
+			if (FrontendFilters->GetFilterAtIndex(i) == OtherDevelopersFilter)
+			{
+				// Inverse filter, so its presence in the list means "hide other developers content"
+				return false;
+			}
+		}
+		return true;
+	}
+}
+
 FText SAssetPicker::GetShowOtherDevelopersToolTip() const
 {
 	// NOTE: This documents the filter effect rather than the button action.
-	if (FilterListPtr ? FilterListPtr->IsFrontendFilterActive(OtherDevelopersFilter) : OtherDevelopersFilter->GetShowOtherDeveloperAssets())
+	if (IsShowingOtherDevelopersContent())
 	{
 		return LOCTEXT("ShowingOtherDevelopersFilterTooltipText", "Showing Other Developers Assets");
 	}
@@ -678,7 +703,14 @@ void SAssetPicker::HandleShowOtherDevelopersCheckStateChanged( ECheckBoxState In
 	}
 	else
 	{
-		OtherDevelopersFilter->SetShowOtherDeveloperAssets(InCheckboxState == ECheckBoxState::Checked); // The checked state matches the active state.
+		if (InCheckboxState == ECheckBoxState::Checked)
+		{
+			OtherDevelopersFilter->SetActiveInCollection(OtherDevelopersFilter.ToSharedRef(), /* do show other developers content*/ true, *FrontendFilters);
+		}
+		else
+		{
+			OtherDevelopersFilter->SetActiveInCollection(OtherDevelopersFilter.ToSharedRef(), /* do not show other developers content*/ false, *FrontendFilters);
+		}
 	}
 }
 
@@ -690,7 +722,14 @@ ECheckBoxState SAssetPicker::GetShowOtherDevelopersCheckState() const
 	}
 	else
 	{
-		return OtherDevelopersFilter->GetShowOtherDeveloperAssets() ? ECheckBoxState::Checked : ECheckBoxState::Unchecked; // The checked state matches the active state.
+		if (IsShowingOtherDevelopersContent())
+		{
+			return ECheckBoxState::Checked;
+		}
+		else
+		{
+			return ECheckBoxState::Unchecked;
+		}
 	}
 }
 
diff --git a/Engine/Source/Editor/ContentBrowser/Private/SAssetPicker.h b/Engine/Source/Editor/ContentBrowser/Private/SAssetPicker.h
--- a/Engine/Source/Editor/ContentBrowser/Private/SAssetPicker.h
+++ b/Engine/Source/Editor/ContentBrowser/Private/SAssetPicker.h
@@ -19,7 +19,7 @@
 #include "Widgets/DeclarativeSyntaxSupport.h"
 #include "Widgets/SCompoundWidget.h"
 
-class FFrontendFilter_ShowOtherDevelopers;
+class FFilter_ShowOtherDevelopers;
 class FFrontendFilter_Text;
 class FUICommandList;
 class SAssetSearchBox;
@@ -121,6 +121,9 @@ private:
 	/** Gets if showing other developers assets */
 	ECheckBoxState GetShowOtherDevelopersCheckState() const;
 
+	/** Returns whether the ShowOtherDevelopersFilter is active i.e. if we should show content belonging to other developers */
+	bool IsShowingOtherDevelopersContent() const;
+
 	/** Called upon the Rename UICommand being executed, sends rename request to the asset view */
 	void OnRenameRequested() const;
 
@@ -188,7 +191,7 @@ private:
 	/** Filters needed for filtering the assets */
 	TSharedPtr< FAssetFilterCollectionType > FilterCollection;
 	TSharedPtr< FFrontendFilter_Text > TextFilter;
-	TSharedPtr< FFrontendFilter_ShowOtherDevelopers > OtherDevelopersFilter;
+	TSharedPtr<FFilter_ShowOtherDevelopers> OtherDevelopersFilter;
 
 	EAssetTypeCategories::Type DefaultFilterMenuExpansion;
 
diff --git a/Engine/Source/Editor/ContentBrowser/Private/SAssetView.cpp b/Engine/Source/Editor/ContentBrowser/Private/SAssetView.cpp
--- a/Engine/Source/Editor/ContentBrowser/Private/SAssetView.cpp
+++ b/Engine/Source/Editor/ContentBrowser/Private/SAssetView.cpp
@@ -2,6 +2,7 @@
 
 #include "SAssetView.h"
 
+#include "Algo/AnyOf.h"
 #include "Algo/Transform.h"
 #include "AssetRegistry/AssetRegistryState.h"
 #include "AssetSelection.h"
@@ -595,10 +596,18 @@ bool SAssetView::IsAssetPathSelected() const
 	return NumAssetPaths > 0 && NumClassPaths == 0;
 }
 
-void SAssetView::SetBackendFilter(const FARFilter& InBackendFilter)
+void SAssetView::SetBackendFilter(const FARFilter& InBackendFilter, TArray<TSharedRef<const FPathPermissionList>>* InCustomPermissionLists)
 {
 	// Update the path and collection lists
 	BackendFilter = InBackendFilter;
+	if (InCustomPermissionLists)
+	{
+		BackendCustomPathFilters = *InCustomPermissionLists;
+	}
+	else
+	{
+		BackendCustomPathFilters.Reset();
+	}
 	RequestSlowFullListRefresh();
 }
 
@@ -1937,6 +1946,43 @@ FContentBrowserDataFilter SAssetView::CreateBackendDataFilter(bool bInvalidateCa
 		| (IsShowingLocalizedContent() ? EContentBrowserItemAttributeFilter::IncludeLocalized : EContentBrowserItemAttributeFilter::IncludeNone);
 
 	TSharedPtr<FPathPermissionList> CombinedFolderPermissionList = ContentBrowserUtils::GetCombinedFolderPermissionList(FolderPermissionList, IsShowingReadOnlyFolders() ? nullptr : WritableFolderPermissionList);
+	
+	UContentBrowserDataSubsystem* CBData = IContentBrowserDataModule::Get().GetSubsystem();
+	if (BackendCustomPathFilters.Num())
+	{
+		if (!CombinedFolderPermissionList.IsValid())
+		{
+			CombinedFolderPermissionList = MakeShared<FPathPermissionList>();
+		}
+
+		if (!CombinedFolderPermissionList->HasAllowListEntries() 
+		&& Algo::AnyOf(BackendCustomPathFilters, UE_PROJECTION_MEMBER(FPathPermissionList, HasAllowListEntries)))
+		{
+			// Need to add an explicit allow-root to the combined list before combining so that the allow list entries don't take everything away
+			CombinedFolderPermissionList->AddAllowListItem("AssetView", TEXTVIEW("/"));
+		}
+
+		TArray<FName> SelectedPaths;
+		SelectedPaths.Reserve(SourcesData.VirtualPaths.Num());
+		// Convert paths to internal if possible
+		for (FName VirtualPath : SourcesData.VirtualPaths)
+		{
+			FName ConvertedPath;
+			CBData->TryConvertVirtualPath(VirtualPath, ConvertedPath);
+			SelectedPaths.Add(ConvertedPath);
+		}
+		// If a filter list explicitly denies a folder we have selected, ignore that filter.
+		for (const TSharedRef<const FPathPermissionList>& CustomList : BackendCustomPathFilters)
+		{
+			const bool bFiltersExplicitSelection = !bRecurse && Algo::AnyOf(SelectedPaths, [&CustomList](FName SelectedPath) {
+				return !CustomList->PassesStartsWithFilter(WriteToString<256>(SelectedPath));
+			});
+			if (!bFiltersExplicitSelection)
+			{
+				CombinedFolderPermissionList = MakeShared<FPathPermissionList>(CombinedFolderPermissionList->CombinePathFilters(*CustomList));
+			}
+		}
+	}
 
 	if (bShowDisallowedAssetClassAsUnsupportedItems && AssetClassPermissionList && AssetClassPermissionList->HasFiltering())
 	{
diff --git a/Engine/Source/Editor/ContentBrowser/Private/SContentBrowser.cpp b/Engine/Source/Editor/ContentBrowser/Private/SContentBrowser.cpp
--- a/Engine/Source/Editor/ContentBrowser/Private/SContentBrowser.cpp
+++ b/Engine/Source/Editor/ContentBrowser/Private/SContentBrowser.cpp
@@ -3737,8 +3737,9 @@ void SContentBrowser::UpdatePath()
 
 void SContentBrowser::OnFilterChanged()
 {
-	FARFilter Filter = FilterListPtr->GetCombinedBackendFilter();
-	AssetViewPtr->SetBackendFilter( Filter );
+	TArray<TSharedRef<const FPathPermissionList>> CustomPermissionLists;
+	FARFilter Filter = FilterListPtr->GetCombinedBackendFilter(CustomPermissionLists);
+	AssetViewPtr->SetBackendFilter(Filter, &CustomPermissionLists);
 
 	// Notify 'filter changed' delegate
 	FContentBrowserModule& ContentBrowserModule = FModuleManager::GetModuleChecked<FContentBrowserModule>( TEXT("ContentBrowser") );
diff --git a/Engine/Source/Editor/ContentBrowser/Private/SFilterList.cpp b/Engine/Source/Editor/ContentBrowser/Private/SFilterList.cpp
--- a/Engine/Source/Editor/ContentBrowser/Private/SFilterList.cpp
+++ b/Engine/Source/Editor/ContentBrowser/Private/SFilterList.cpp
@@ -71,11 +71,15 @@ void SFilterList::Construct( const FArguments& InArgs )
 
 	TSharedPtr<FFrontendFilterCategory> DefaultCategory = MakeShareable( new FFrontendFilterCategory(LOCTEXT("FrontendFiltersCategory", "Other Filters"), LOCTEXT("FrontendFiltersCategoryTooltip", "Filter assets by all filters in this category.")) );
 	
+	TSharedRef<FFilter_ShowOtherDevelopers> OtherDevelopersFilter = MakeShared<FFilter_ShowOtherDevelopers>(DefaultCategory, InArgs._FilterBarIdentifier);
+	// This filter affecst the backend query so we must perform a full refresh when it changes
+	OtherDevelopersFilter->OnChanged().Add(this->OnFilterChanged);
+
 	// Add all built-in frontend filters here
 	AllFrontendFilters_Internal.Add( MakeShareable(new FFrontendFilter_CheckedOut(DefaultCategory)) );
 	AllFrontendFilters_Internal.Add( MakeShareable(new FFrontendFilter_Modified(DefaultCategory)) );
 	AllFrontendFilters_Internal.Add( MakeShareable(new FFrontendFilter_Writable(DefaultCategory)) );
-	AllFrontendFilters_Internal.Add( MakeShareable(new FFrontendFilter_ShowOtherDevelopers(DefaultCategory)) );
+	AllFrontendFilters_Internal.Add(OtherDevelopersFilter);
 	AllFrontendFilters_Internal.Add( MakeShareable(new FFrontendFilter_ReplicatedBlueprint(DefaultCategory)) );
 	AllFrontendFilters_Internal.Add(MakeShared<FFilter_ShowRedirectors>(DefaultCategory));
 	AllFrontendFilters_Internal.Add( MakeShareable(new FFrontendFilter_InUseByLoadedLevels(DefaultCategory)) );
@@ -183,6 +187,16 @@ TSharedRef<SWidget> SFilterList::ExternalMakeAddFilterMenu()
 	return SAssetFilterBar<FAssetFilterType>::MakeAddFilterMenu();
 }
 
+FARFilter SFilterList::GetCombinedBackendFilter(TArray<TSharedRef<const FPathPermissionList>>& OutPermissionLists) const
+{
+	TSharedPtr<FFilter_ShowOtherDevelopers> OtherDevelopersFilter = StaticCastSharedPtr<FFilter_ShowOtherDevelopers>(GetFrontendFilter("ShowOtherDevelopersBackend"));
+	if (!OtherDevelopersFilter->IsActive()) // Filter when inactive, because it's an inverse filter 
+	{
+		OutPermissionLists.Add(OtherDevelopersFilter->GetPathPermissionList());
+	}
+	return Super::GetCombinedBackendFilter();
+}
+
 void SFilterList::DisableFiltersThatHideItems(TArrayView<const FContentBrowserItem> ItemList)
 {
 	if (HasAnyFilters() && ItemList.Num() > 0)
@@ -198,7 +212,8 @@ void SFilterList::DisableFiltersThatHideItems(TArrayView<const FContentBrowserIt
 
 				FContentBrowserDataFilter DataFilter;
 				DataFilter.bRecursivePaths = true;
-				ContentBrowserUtils::AppendAssetFilterToContentBrowserFilter(GetCombinedBackendFilter(), nullptr, nullptr, DataFilter);
+				TArray<TSharedRef<const FPathPermissionList>> UnusedPermissionLists;
+				ContentBrowserUtils::AppendAssetFilterToContentBrowserFilter(GetCombinedBackendFilter(UnusedPermissionLists), nullptr, nullptr, DataFilter);
 
 				ContentBrowserData->CompileFilter(RootPath, DataFilter, CompiledDataFilter);
 			}
@@ -254,6 +269,30 @@ void SFilterList::DisableFiltersThatHideItems(TArrayView<const FContentBrowserIt
 			}
 		}
 
+		auto AddAndActivateInverseFilter = [this, &ExecuteOnFilterChanged](const TSharedRef<FFilterBase<FAssetFilterType>>& InFilter) 
+		{
+			int32 ExistingIndex = Filters.IndexOfByPredicate([InFilter](TSharedPtr<SFilter> Filter) { return Filter->GetFrontendFilter() == InFilter; });
+			TSharedRef<SFilter> FilterWidget = ExistingIndex == INDEX_NONE ? AddFilterToBar(InFilter) : Filters[ExistingIndex];
+			FilterWidget->SetEnabled(true, false);
+			SetFrontendFilterActive(InFilter, true);
+			ExecuteOnFilterChanged = true;
+		};
+		
+		TSharedPtr<FFilter_ShowOtherDevelopers> OtherDevelopersFilter = StaticCastSharedPtr<FFilter_ShowOtherDevelopers>(GetFrontendFilter("ShowOtherDevelopersBackend"));
+		// Special case: if item is hidden because of "other developers" filter, disable it
+		if (OtherDevelopersFilter.IsValid() && OtherDevelopersFilter->IsActive())
+		{
+			TSharedRef<const FPathPermissionList> PermissionList = OtherDevelopersFilter->GetPathPermissionList();	
+			for (const FContentBrowserItem& Item : ItemList)
+			{
+				if (PermissionList->PassesStartsWithFilter(WriteToString<256>(Item.GetInternalPath())))
+				{
+					AddAndActivateInverseFilter(OtherDevelopersFilter.ToSharedRef());
+					break;
+				}
+			}
+		}
+
 		// Special case: if the object is a redirector then enable the 'show redirectors' filter - this will also prevent
 		// folders that contain only redirectors from being hidden with the "hide empty folders" setting
 		FString RedirectorClassPath = UObjectRedirector::StaticClass()->GetPathName();
@@ -266,11 +305,7 @@ void SFilterList::DisableFiltersThatHideItems(TArrayView<const FContentBrowserIt
 			TSharedPtr<FFilter_ShowRedirectors> RedirectorFilter = StaticCastSharedPtr<FFilter_ShowRedirectors>(GetFrontendFilter("ShowRedirectorsBackend"));
 			if (RedirectorFilter.IsValid())
 			{
-				int32 ExistingIndex = Filters.IndexOfByPredicate([RedirectorFilter](TSharedPtr<SFilter> Filter) { return Filter->GetFrontendFilter() == RedirectorFilter; });
-				TSharedRef<SFilter> FilterWidget = ExistingIndex == INDEX_NONE ? AddFilterToBar(RedirectorFilter.ToSharedRef()) : Filters[ExistingIndex];
-				FilterWidget->SetEnabled(true, false);
-				SetFrontendFilterActive(RedirectorFilter.ToSharedRef(), true);
-				ExecuteOnFilterChanged = true;
+				AddAndActivateInverseFilter(RedirectorFilter.ToSharedRef());
 			}
 		}
 
diff --git a/Engine/Source/Editor/ContentBrowser/Private/SFilterList.h b/Engine/Source/Editor/ContentBrowser/Private/SFilterList.h
--- a/Engine/Source/Editor/ContentBrowser/Private/SFilterList.h
+++ b/Engine/Source/Editor/ContentBrowser/Private/SFilterList.h
@@ -23,6 +23,7 @@
 #include "UObject/NameTypes.h"
 #include "Widgets/DeclarativeSyntaxSupport.h"
 
+class FFilter_ShowOtherDevelopers;
 class FFrontendFilter;
 class SWidget;
 class UClass;
@@ -40,6 +41,7 @@ enum class ECheckBoxState : uint8;
  */
 class SFilterList : public SAssetFilterBar<FAssetFilterType>
 {
+	using Super = SAssetFilterBar<FAssetFilterType>; 
 public:
 	DECLARE_DELEGATE_OneParam(FOnFilterBarLayoutChanging, EFilterBarLayout /* NewLayout */)
 	/**
@@ -119,6 +121,14 @@ public:
 	/** Retrieve a specific frontend filter */
 	TSharedPtr<FFrontendFilter> GetFrontendFilter(const FString& InName) const;
 	
+	/** 
+	 * Replaces super version of GetCombinedBackendFilter to allow returning FPathPermissionLists as well.
+	 * Permission lists will be combined into a content browser data filter after a delay. Filters may modify them 
+	 * via the shared ref but should not do so off the game thread.
+	 */
+	FARFilter GetCombinedBackendFilter(TArray<TSharedRef<const FPathPermissionList>>& OutPermissionLists) const;
+	FARFilter GetCombinedBackendFilter() const = delete;
+	
 	/** Handler for when the floating add filter button was clicked */
 	TSharedRef<SWidget> ExternalMakeAddFilterMenu();
 
diff --git a/Engine/Source/Editor/ContentBrowser/Public/FrontendFilters.h b/Engine/Source/Editor/ContentBrowser/Public/FrontendFilters.h
--- a/Engine/Source/Editor/ContentBrowser/Public/FrontendFilters.h
+++ b/Engine/Source/Editor/ContentBrowser/Public/FrontendFilters.h
@@ -233,7 +233,8 @@ public:
 };
 
 /** An inverse filter that allows display of content in developer folders that are not the current user's */
-class CONTENTBROWSER_API FFrontendFilter_ShowOtherDevelopers : public FFrontendFilter
+class UE_DEPRECATED(5.5, "This frontend filter has been deprecated and replaced with backend filtering. see FFilter_ShowOtherDevelopers.")
+CONTENTBROWSER_API FFrontendFilter_ShowOtherDevelopers : public FFrontendFilter
 {
 public:
 	/** Constructor */
diff --git a/Engine/Source/Editor/ContentBrowser/Public/SAssetView.h b/Engine/Source/Editor/ContentBrowser/Public/SAssetView.h
--- a/Engine/Source/Editor/ContentBrowser/Public/SAssetView.h
+++ b/Engine/Source/Editor/ContentBrowser/Public/SAssetView.h
@@ -308,8 +308,14 @@ public:
 	/** Returns true if a real asset path is selected (i.e \Engine\* or \Game\*) */
 	bool IsAssetPathSelected() const;
 
-	/** Notifies the asset view that the filter-list filter has changed */
-	void SetBackendFilter(const FARFilter& InBackendFilter);
+	/**
+	 * @brief Provide a backend filter for the asset view and invalidate current source items.
+	 * 
+	 * @param InBackendFilter Asset registry filter for uassets
+	 * @param InCustomPermissionLists Optional permission lists to allow/deny specific folders.
+	 * 	All folders will be allowed by default, so only denials will have any effect.
+	 */
+	void SetBackendFilter(const FARFilter& InBackendFilter, TArray<TSharedRef<const FPathPermissionList>>* InCustomPermissionLists = nullptr);
 
 	/** Get the current backend filter */
 	const FARFilter& GetBackendFilter() const { return BackendFilter; }
@@ -933,6 +939,9 @@ private:
 	TSharedPtr<FPathPermissionList> AssetClassPermissionList;
 	TSharedPtr<FPathPermissionList> FolderPermissionList;
 	TSharedPtr<FPathPermissionList> WritableFolderPermissionList;
+	// Paths which should be filtered out based on current filters the user has selected 
+	//  - not 'permissions' so may be ignore if e.g. the user explicitly selects a filtered folder
+	TArray<TSharedRef<const FPathPermissionList>> BackendCustomPathFilters;
 	TSharedPtr<FAssetFilterCollectionType> FrontendFilters;
 
 	TAttribute<bool> bShowRedirectors;
diff --git a/Engine/Source/Runtime/Core/Public/Misc/NamePermissionList.h b/Engine/Source/Runtime/Core/Public/Misc/NamePermissionList.h
--- a/Engine/Source/Runtime/Core/Public/Misc/NamePermissionList.h
+++ b/Engine/Source/Runtime/Core/Public/Misc/NamePermissionList.h
@@ -357,7 +357,7 @@ public:
 	 * Result will contain AllowList paths that pass both filters.
 	 * @return new combined filter.
 	 */
-	CORE_API FPathPermissionList CombinePathFilters(const FPathPermissionList& OtherFilter) const;
+	CORE_API [[nodiscard]] FPathPermissionList CombinePathFilters(const FPathPermissionList& OtherFilter) const;
 
 	/**
 	* Unregisters specified owners then adds specified filters in one operation (to avoid multiple filters changed events).
