Additional path view logging.

diff --git a/Engine/Source/Editor/ContentBrowser/Private/SPathView.cpp b/Engine/Source/Editor/ContentBrowser/Private/SPathView.cpp
--- a/Engine/Source/Editor/ContentBrowser/Private/SPathView.cpp
+++ b/Engine/Source/Editor/ContentBrowser/Private/SPathView.cpp
@@ -83,6 +83,8 @@ class SWidget;
 struct FAssetData;
 struct FGeometry;
 
+DEFINE_LOG_CATEGORY_STATIC(LogPathView, Log, Log);
+
 #define LOCTEXT_NAMESPACE "ContentBrowser"
 
 namespace UE::PathView
@@ -428,11 +430,18 @@ void FPathViewData::PopulateFullFolderTree(const FContentBrowserDataCompiledFilt
 			UContentBrowserDataSource* Source = InItemData.GetOwnerDataSource();
 			if (Source && !Source->IsFolderVisible(InItemData.GetVirtualPath(), EmptyFilter.FolderFlags, EmptyFilter.FolderFilter))
 			{
-				UE_LOG(LogContentBrowser, VeryVerbose, TEXT("Hiding folder %s on source %s that fails current pre-text filtering"),
-					*WriteToString<256>(InItemData.GetVirtualPath()), *WriteToString<256>(Source->GetFName()));
+				UE_LOG(LogPathView, VeryVerbose, TEXT("[%s] Populate: skipping folder %s:%s that fails current pre-text filtering"),
+				 	*WriteToString<256>(OwningContentBrowserName),
+					*WriteToString<256>(Source->GetFName()),
+					*WriteToString<256>(InItemData.GetVirtualPath()));
 				return true; // continue enumerating
 			}
 
+			UE_LOG(LogPathView, VeryVerbose, TEXT("[%s] Populate: adding folder %s:%s"),
+				*WriteToString<256>(OwningContentBrowserName),
+				*WriteToString<256>(Source->GetFName()),
+				*WriteToString<256>(InItemData.GetVirtualPath()));
+
 			AddFolderItemInternal(MoveTemp(InItemData), &OldItemsByInvariantPath);
 			return true;
 		});
@@ -464,7 +473,7 @@ void FPathViewData::PopulateWithFavorites(const FContentBrowserDataCompiledFilte
 				UContentBrowserDataSource* ItemDataSource = InItemData.GetOwnerDataSource();
 				if (!ItemDataSource->IsFolderVisible(InItemData.GetVirtualPath(), EmptyFilter.FolderFlags, EmptyFilter.FolderFilter))
 				{
-					UE_LOG(LogContentBrowser,
+					UE_LOG(LogPathView,
 						VeryVerbose,
 						TEXT("Hiding folder %s that fails current pre-text filtering"),
 						*WriteToString<256>(InItemData.GetVirtualPath()));
@@ -492,15 +501,15 @@ void FPathViewData::ProcessDataUpdates(TConstArrayView<FContentBrowserItemDataUp
 		UContentBrowserDataSource* ItemDataSource = InItemData.GetOwnerDataSource();
 		if (!ItemDataSource->DoesItemPassFilter(InItemData, CompiledDataFilter))
 		{
+			UE_LOG(LogPathView, VeryVerbose, TEXT("[%s] Fails compiled data filter"),
+				*WriteToString<256>(OwningContentBrowserName));
 			return false;
 		}
 
 		if (!ContentBrowserData->IsFolderVisible(InItemData.GetVirtualPath(), EmptyFilter.FolderFlags, EmptyFilter.FolderFilter))
 		{
-			UE_LOG(LogContentBrowser,
-				VeryVerbose,
-				TEXT("Hiding folder %s that fails broad visibility filter"),
-				*WriteToString<256>(InItemData.GetVirtualPath()));
+			UE_LOG(LogPathView, VeryVerbose, TEXT("[%s] Fails folder visibility filter"),
+				*WriteToString<256>(OwningContentBrowserName));
 			return false;
 		}
 
@@ -523,6 +532,9 @@ void FPathViewData::ProcessDataUpdates(TConstArrayView<FContentBrowserItemDataUp
 		switch (ItemDataUpdate.GetUpdateType())
 		{
 			case EContentBrowserItemUpdateType::Added:
+				UE_LOG(LogPathView, VeryVerbose, TEXT("[%s] Added item %s:%s"), *WriteToString<256>(OwningContentBrowserName), 
+					*WriteToString<256>(ItemData.GetOwnerDataSource()->GetFName()),
+					*WriteToString<256>(ItemData.GetVirtualPath()));
 				if (DoesItemPassFilter(ItemData))
 				{
 					NewItems.Emplace(AddFolderItemInternal(MoveTemp(ItemData), nullptr));
@@ -530,6 +542,9 @@ void FPathViewData::ProcessDataUpdates(TConstArrayView<FContentBrowserItemDataUp
 				break;
 
 			case EContentBrowserItemUpdateType::Modified:
+				UE_LOG(LogPathView, VeryVerbose, TEXT("[%s] Modified item %s:%s"), *WriteToString<256>(OwningContentBrowserName), 
+					*WriteToString<256>(ItemData.GetOwnerDataSource()->GetFName()),
+					*WriteToString<256>(ItemData.GetVirtualPath()));
 				if (DoesItemPassFilter(ItemData))
 				{
 					NewItems.Emplace(AddFolderItemInternal(MoveTemp(ItemData), nullptr));
@@ -546,6 +561,10 @@ void FPathViewData::ProcessDataUpdates(TConstArrayView<FContentBrowserItemDataUp
 
 			case EContentBrowserItemUpdateType::Moved:
 			{
+				UE_LOG(LogPathView, VeryVerbose, TEXT("[%s] Moved item %s:%s->%s"), *WriteToString<256>(OwningContentBrowserName), 
+					*WriteToString<256>(ItemData.GetOwnerDataSource()->GetFName()),
+					*WriteToString<256>(ItemDataUpdate.GetPreviousVirtualPath()),
+					*WriteToString<256>(ItemData.GetVirtualPath()));
 				const FContentBrowserMinimalItemData OldItemKey(ItemData.GetItemType(), ItemDataUpdate.GetPreviousVirtualPath(), ItemData.GetOwnerDataSource());
 				TSharedPtr<FTreeItem> Parent = TryRemoveFolderItemInternal(OldItemKey);
 				if (DoesItemPassFilter(ItemData))
@@ -560,6 +579,9 @@ void FPathViewData::ProcessDataUpdates(TConstArrayView<FContentBrowserItemDataUp
 			break;
 
 			case EContentBrowserItemUpdateType::Removed:
+				UE_LOG(LogPathView, VeryVerbose, TEXT("[%s] Removed item %s:%s"), *WriteToString<256>(OwningContentBrowserName), 
+					*WriteToString<256>(ItemData.GetOwnerDataSource()->GetFName()),
+					*WriteToString<256>(ItemData.GetVirtualPath()));
 				TryRemoveFolderItemInternal(ItemData);
 				break;
 
@@ -709,18 +731,27 @@ TSharedRef<FTreeItem> FPathViewData::AddFolderItem(FContentBrowserItemData&& InI
 TSharedRef<FTreeItem> FPathViewData::AddFolderItemInternal(FContentBrowserItemData&& InItemData,
 	TMap<FName, TSharedPtr<FTreeItem>>* OldItemsByInvariantPath)
 {
+	UContentBrowserDataSubsystem* ContentBrowserData = IContentBrowserDataModule::Get().GetSubsystem();
+	UContentBrowserDataSource* OriginalDataSource = InItemData.GetOwnerDataSource();
 	FName ItemVirtualPath = InItemData.GetVirtualPath();
 	TSharedPtr<FTreeItem> LeafItem = VirtualPathToItem.FindRef(ItemVirtualPath);
 	if (LeafItem.IsValid())
 	{
+		UE_LOG(LogPathView, VeryVerbose, TEXT("[%s] Appending data to existing item %s:%s"), 
+			*WriteToString<256>(OwningContentBrowserName),
+			*WriteToString<256>(OriginalDataSource->GetFName()),
+			*WriteToString<256>(ItemVirtualPath));
+
 		// Item already existed - duplicate item returned by multiple data sources, merge data and move on.
 		// We will have already created all the parent items.
 		LeafItem->AppendItemData(InItemData);
 		return LeafItem.ToSharedRef();
 	}
 
-	UContentBrowserDataSubsystem* ContentBrowserData = IContentBrowserDataModule::Get().GetSubsystem();
-	UContentBrowserDataSource* OriginalDataSource = InItemData.GetOwnerDataSource();
+	UE_LOG(LogPathView, VeryVerbose, TEXT("[%s] Creating new tree item for %s:%s"), 
+		*WriteToString<256>(OwningContentBrowserName),
+		*WriteToString<256>(OriginalDataSource->GetFName()),
+		*WriteToString<256>(ItemVirtualPath));
 
 	FName ItemInvariantPath = InItemData.GetInvariantPath();
 	TStringBuilder<FName::StringBufferSize> PathBuffer(InPlace, ItemVirtualPath);
@@ -748,6 +779,11 @@ TSharedRef<FTreeItem> FPathViewData::AddFolderItemInternal(FContentBrowserItemDa
 			}
 			if (PathView == TEXTVIEW("/"))
 			{
+				UE_LOG(LogPathView, VeryVerbose, TEXT("[%s] Adding root item %s:%.*s"), 
+					*WriteToString<256>(OwningContentBrowserName),
+					*WriteToString<256>(OriginalDataSource->GetFName()),
+					PathView.Len(), PathView.GetData());
+
 				// PreviousItem must have been new, add it to the set of root items
 				RootItems.Add(PreviousItem);
 				return false;
@@ -757,8 +793,9 @@ TSharedRef<FTreeItem> FPathViewData::AddFolderItemInternal(FContentBrowserItemDa
 			bool bContinue = false;
 			if (!ParentItem.IsValid())
 			{
-				UE_LOG(LogContentBrowser, VeryVerbose, TEXT("[%s] Creating placeholder or virtual parent %.*s"), 
+				UE_LOG(LogPathView, VeryVerbose, TEXT("[%s] Creating placeholder or virtual parent %s:%.*s"), 
 					*WriteToString<256>(OwningContentBrowserName),
+					*WriteToString<256>(OriginalDataSource->GetFName()),
 					PathView.Len(), PathView.GetData());
 				// TODO: If another data source provides this path in future, can that data source become the 'primary'?
 				FName ItemName(FPathViews::GetPathLeaf(PathView)); 
@@ -781,6 +818,12 @@ TSharedRef<FTreeItem> FPathViewData::AddFolderItemInternal(FContentBrowserItemDa
 				InvariantPathToItem.Add(ParentItem->GetItem().GetInvariantPath()); 
 				bContinue = true;
 			}
+			else
+			{
+				UE_LOG(LogPathView, VeryVerbose, TEXT("[%s] Found existing parent %.*s"), 
+					*WriteToString<256>(OwningContentBrowserName),
+					PathView.Len(), PathView.GetData());
+			}
 			ParentItem->AddChild(PreviousItem);
 			PreviousItem = ParentItem.ToSharedRef();
 			return bContinue; // If we made a node here, keep checking if we need to make more parent nodes
@@ -1159,7 +1202,7 @@ void SPathView::Tick(const FGeometry& AllottedGeometry, const double InCurrentTi
 	const bool bNewShowRedirectors = bShowRedirectors.Get(false);
 	if (bNewShowRedirectors != bLastShowRedirectors)
 	{
-		UE_LOG(LogContentBrowser, Verbose, TEXT("PathView bShowRedirectors changed to %d"), bNewShowRedirectors);
+		UE_LOG(LogPathView, Verbose, TEXT("PathView bShowRedirectors changed to %d"), bNewShowRedirectors);
 		bLastShowRedirectors = bNewShowRedirectors;
 		HandleSettingChanged("ShowRedirectors");
 	}
@@ -1288,7 +1331,7 @@ void SPathView::SetPluginPathFilterActive(const TSharedRef<FContentBrowserPlugin
 		bActive = !bActive;
 	}
 
-	UE_LOG(LogContentBrowser, Verbose, TEXT("[%s] Setting%s plugin filter %s to %s"), 
+	UE_LOG(LogPathView, Verbose, TEXT("[%s] Setting%s plugin filter %s to %s"), 
 		*WriteToString<64>(OwningContentBrowserName), 
 		Filter->IsInverseFilter() ? TEXT(" inverse") : TEXT(""),
 		*Filter->GetName(), bActive ? TEXT("Active") : TEXT("Inactive"));
@@ -1460,7 +1503,7 @@ void SPathView::RenameFolderItem(const FContentBrowserItem& InItem)
 
 FContentBrowserDataCompiledFilter SPathView::CreateCompiledFolderFilter() const
 {
-	UE_LOG(LogContentBrowser, VeryVerbose, TEXT("[%s] Creating folder filter"), *WriteToString<256>(OwningContentBrowserName));
+	UE_LOG(LogPathView, VeryVerbose, TEXT("[%s] Creating folder filter"), *WriteToString<256>(OwningContentBrowserName));
 
 	const UContentBrowserSettings* ContentBrowserSettings = GetDefault<UContentBrowserSettings>();
 	bool bDisplayPluginFolders = ContentBrowserSettings->GetDisplayPluginFolders();
@@ -1476,7 +1519,7 @@ FContentBrowserDataCompiledFilter SPathView::CreateCompiledFolderFilter() const
 	DataFilter.ItemCategoryFilter = GetContentBrowserItemCategoryFilter();
 	DataFilter.ItemAttributeFilter = GetContentBrowserItemAttributeFilter();
 
-	UE_LOG(LogContentBrowser, VeryVerbose, TEXT("[%s] bDisplayPluginFolders:%d ItemCategoryFilter:%d ItemAttributeFilter:%d"), 
+	UE_LOG(LogPathView, VeryVerbose, TEXT("[%s] bDisplayPluginFolders:%d ItemCategoryFilter:%d ItemAttributeFilter:%d"), 
 		*WriteToString<256>(OwningContentBrowserName), bDisplayPluginFolders, DataFilter.ItemCategoryFilter, DataFilter.ItemAttributeFilter);
 
 	TSharedPtr<FPathPermissionList> CombinedFolderPermissionList = ContentBrowserUtils::GetCombinedFolderPermissionList(FolderPermissionList, bAllowReadOnlyFolders ? nullptr : WritableFolderPermissionList);
@@ -1492,7 +1535,7 @@ FContentBrowserDataCompiledFilter SPathView::CreateCompiledFolderFilter() const
 
 	if (PluginPathFilters.IsValid() && PluginPathFilters->Num() > 0 && bDisplayPluginFolders)
 	{
-		UE_SUPPRESS(LogContentBrowser, VeryVerbose, {
+		UE_SUPPRESS(LogPathView, VeryVerbose, {
 			FString PluginFiltersString;
 			for (int32 i=0; i < PluginPathFilters->Num(); ++i)
 			{
@@ -1502,7 +1545,7 @@ FContentBrowserDataCompiledFilter SPathView::CreateCompiledFolderFilter() const
 				}
 				PluginFiltersString += PluginPathFilters->GetFilterAtIndex(i)->GetName();
 			}
-			UE_LOG(LogContentBrowser, VeryVerbose, TEXT("[%s] Active plugin filters: %s"), 
+			UE_LOG(LogPathView, VeryVerbose, TEXT("[%s] Active plugin filters: %s"), 
 				*WriteToString<256>(OwningContentBrowserName), *PluginFiltersString);
 		});
 		TArray<TSharedRef<IPlugin>> Plugins = IPluginManager::Get().GetEnabledPluginsWithContent();
@@ -1522,7 +1565,7 @@ FContentBrowserDataCompiledFilter SPathView::CreateCompiledFolderFilter() const
 		}
 	}
 
-	UE_LOG(LogContentBrowser, VeryVerbose, TEXT("Compiled folder permission list: %s"), CombinedFolderPermissionList.IsValid() ? *CombinedFolderPermissionList->ToString() : TEXT("null"));
+	UE_LOG(LogPathView, VeryVerbose, TEXT("Compiled folder permission list: %s"), CombinedFolderPermissionList.IsValid() ? *CombinedFolderPermissionList->ToString() : TEXT("null"));
 
 	ContentBrowserUtils::AppendAssetFilterToContentBrowserFilter(FARFilter(), nullptr, CombinedFolderPermissionList, DataFilter);
 
@@ -1813,7 +1856,7 @@ void SPathView::LoadSettings(const FString& IniFilename, const FString& IniSecti
 		IContentBrowserDataModule::Get().GetSubsystem()->ConvertInternalPathToVirtual(Path, Path);
 	}
 
-	UE_LOG(LogContentBrowser, Verbose, TEXT("[%s] LoadSettings: SelectedPaths: %s"), 
+	UE_LOG(LogPathView, Verbose, TEXT("[%s] LoadSettings: SelectedPaths: %s"), 
 		*WriteToString<256>(OwningContentBrowserName), *FString::JoinBy(NewSelectedPaths, TEXT(", "), UE_PROJECTION_MEMBER(FName, ToString)));
 
 	{
@@ -1860,7 +1903,7 @@ void SPathView::LoadSettings(const FString& IniFilename, const FString& IniSecti
 		TArray<FString> NewSelectedFilters;
 		if (FPathViewConfig* PathViewConfig = GetPathViewConfig())
 		{
-			UE_LOG(LogContentBrowser, Verbose, TEXT("[%s] LoadSettings: Loading plugin filters from editor config: %s"), 
+			UE_LOG(LogPathView, Verbose, TEXT("[%s] LoadSettings: Loading plugin filters from editor config: %s"), 
 				*WriteToString<256>(OwningContentBrowserName), *FString::Join(NewSelectedFilters, TEXT(", ")));
 			NewSelectedFilters = PathViewConfig->PluginFilters;
 		}
@@ -1869,7 +1912,7 @@ void SPathView::LoadSettings(const FString& IniFilename, const FString& IniSecti
 			FString PluginFiltersString;
 			if (GConfig->GetString(*IniSection, *(SettingsString + TEXT(".PluginFilters")), PluginFiltersString, IniFilename))
 			{
-				UE_LOG(LogContentBrowser, Verbose, TEXT("[%s] LoadSettings: Loading plugin filters from ini: %s"), 
+				UE_LOG(LogPathView, Verbose, TEXT("[%s] LoadSettings: Loading plugin filters from ini: %s"), 
 					*WriteToString<256>(OwningContentBrowserName), *PluginFiltersString);
 				PluginFiltersString.ParseIntoArray(NewSelectedFilters, TEXT(","), /*bCullEmpty*/ true);
 			}
@@ -2161,7 +2204,7 @@ FText SPathView::GetHighlightText() const
 void SPathView::Populate(const bool bIsRefreshingFilter)
 {
 	TRACE_CPUPROFILER_EVENT_SCOPE(SPathView::Populate);
-	UE_LOG(LogContentBrowser, Verbose, TEXT("Repopulating path view"));
+	UE_LOG(LogPathView, Verbose, TEXT("Repopulating path view"));
 
 	const bool bFilteringByText = !TreeData->GetFolderPathTextFilter().GetRawFilterText().IsEmpty();
 
@@ -2178,6 +2221,7 @@ void SPathView::Populate(const bool bIsRefreshingFilter)
 		return ExplicitlyAddPathToSelection(VirtualPath);
 	}))
 	{
+		UE_LOG(LogPathView, VeryVerbose, TEXT("[%s] Resetting pending initial paths as all are selected"), *WriteToString<256>(OwningContentBrowserName));
 		PendingInitialPaths.Reset();
 	}
 }
@@ -2324,6 +2368,7 @@ void SPathView::HandleItemDataUpdated(TArrayView<const FContentBrowserItemDataUp
 	if (!Algo::AnyOf(InUpdatedItems,
 			[](const FContentBrowserItemDataUpdate& Update) { return Update.GetItemData().IsFolder(); }))
 	{
+		UE_LOG(LogPathView, VeryVerbose, TEXT("[%s] Skipping item data update because there were no folders present"), *WriteToString<256>(OwningContentBrowserName));
 		return;
 	}
 
@@ -2335,9 +2380,8 @@ void SPathView::HandleItemDataUpdated(TArrayView<const FContentBrowserItemDataUp
 
 	const double HandleItemDataUpdatedStartTime = FPlatformTime::Seconds();
 	TreeData->ProcessDataUpdates(InUpdatedItems, CreateCompiledFolderFilter());
-	UE_LOG(LogContentBrowser,
-		VeryVerbose,
-		TEXT("PathView - HandleItemDataUpdated completed in %0.4f seconds for %d items"),
+	UE_LOG(LogPathView, VeryVerbose, TEXT("[%s] PathView - HandleItemDataUpdated completed in %0.4f seconds for %d items"),
+		*WriteToString<256>(OwningContentBrowserName),
 		FPlatformTime::Seconds() - HandleItemDataUpdatedStartTime,
 		InUpdatedItems.Num());
 
@@ -2346,6 +2390,7 @@ void SPathView::HandleItemDataUpdated(TArrayView<const FContentBrowserItemDataUp
 		return ExplicitlyAddPathToSelection(VirtualPath);
 	}))
 	{
+		UE_LOG(LogPathView, VeryVerbose, TEXT("[%s] Resetting pending initial paths as all are selected"), *WriteToString<256>(OwningContentBrowserName));
 		PendingInitialPaths.Reset();
 	}
 }
@@ -2365,6 +2410,7 @@ void SPathView::HandleItemDataRefreshed()
 void SPathView::HandleItemDataDiscoveryComplete()
 {
 	// If there were any more initial paths, they no longer exist so clear them now.
+	UE_LOG(LogPathView, VeryVerbose, TEXT("[%s] Resetting pending initial paths at end of asset data discovery"), *WriteToString<256>(OwningContentBrowserName));
 	PendingInitialPaths.Empty();
 }
 
@@ -2380,7 +2426,7 @@ void SPathView::HandleSettingChanged(FName PropertyName)
 		|| (PropertyName == GET_MEMBER_NAME_CHECKED(UContentBrowserSettings, bDisplayFriendlyNameForPluginFolders))
 		|| (PropertyName == NAME_None)) // @todo: Needed if PostEditChange was called manually, for now
 	{
-		UE_LOG(LogContentBrowser,
+		UE_LOG(LogPathView,
 			Log,
 			TEXT("[%s][PathView] HandleSettingChanged %s"),
 			*WriteToString<256>(OwningContentBrowserName),
@@ -2771,6 +2817,7 @@ void SFavoritePathView::HandleItemDataUpdated(TArrayView<const FContentBrowserIt
 	}
 	if (FavoritePaths.Num() == 0)
 	{
+		UE_LOG(LogPathView, VeryVerbose, TEXT("[%s] Skipping item data update because there were no favorites present"), *WriteToString<256>(OwningContentBrowserName));
 		return;
 	}
 
@@ -2817,7 +2864,7 @@ void SFavoritePathView::HandleItemDataUpdated(TArrayView<const FContentBrowserIt
 		}
 	}
 
-	UE_LOG(LogContentBrowser, VeryVerbose, TEXT("FavoritePathView - HandleItemDataUpdated completed in %0.4f seconds for %d items"), FPlatformTime::Seconds() - HandleItemDataUpdatedStartTime, InUpdatedItems.Num());
+	UE_LOG(LogPathView, VeryVerbose, TEXT("FavoritePathView - HandleItemDataUpdated completed in %0.4f seconds for %d items"), FPlatformTime::Seconds() - HandleItemDataUpdatedStartTime, InUpdatedItems.Num());
 }
 
 void SFavoritePathView::FixupFavoritesFromExternalChange(TArrayView<const AssetViewUtils::FMovedContentFolder> MovedFolders)
