Add comment

diff --git a/Engine/Source/Editor/ContentBrowser/Private/Filters.cpp b/Engine/Source/Editor/ContentBrowser/Private/Filters.cpp
--- a/Engine/Source/Editor/ContentBrowser/Private/Filters.cpp
+++ b/Engine/Source/Editor/ContentBrowser/Private/Filters.cpp
@@ -104,6 +104,8 @@ void FFilter_ShowOtherDevelopers::BuildFilter()
 	{
 		UE_LOG(LogContentBrowser, Verbose, TEXT("[%s] FFilterShowOtherDevelopers rebuilt exclusion list: %s"), *FilterBarIdentifier.ToString(), 
 			*FString::JoinBy(OtherDeveloperFolders, TEXT(","), UE_PROJECTION_MEMBER(FName, ToString)));
+		
+		// Recreate the permission list so that the content browser can do pointer comparison to tell that the list has changed rather than tbinding the delegate
 		PathPermissionList = MakeShared<FPathPermissionList>();
 		for (FName OtherPath : OtherDeveloperFolders)
 		{
