Reimplement text filtering in the content browser to be performed asynchronously in parallel.
- FFrontendFilter_Text is deprecated in favor of FAssetTextFilter.
- FAssetTextFilter can be compiled and then the compiled filter cloned for use in parallel.
- Change FFrontendFilter_CustomText to work by merging its query with FAssetTextFilter instead.
- Deprecate UTextFilter*Handler classes and replace with an explicitly registered interface for thread safety.
Refactor item storage in SAssetView:
- FAssetViewItemCollection stores items and handles threading/updates
- Items are stored linearly in an array with an external lookup for simpler filter progress handling
- FHashTable used for lookup can be built & deduplicated in parallel
Other changes and optimizations
- Move AssetViewTypes.h header to private directory as it was had no DLL exported members anyway.
- When SAsset view receives a new backend filter, check if it's different before requesting a full refresh.

diff --git a/Engine/Plugins/Experimental/MeshModelingToolsetExp/Source/ModelingEditorUI/Private/ModelingWidgets/SToolInputAssetPicker.cpp b/Engine/Plugins/Experimental/MeshModelingToolsetExp/Source/ModelingEditorUI/Private/ModelingWidgets/SToolInputAssetPicker.cpp
--- a/Engine/Plugins/Experimental/MeshModelingToolsetExp/Source/ModelingEditorUI/Private/ModelingWidgets/SToolInputAssetPicker.cpp
+++ b/Engine/Plugins/Experimental/MeshModelingToolsetExp/Source/ModelingEditorUI/Private/ModelingWidgets/SToolInputAssetPicker.cpp
@@ -19,6 +19,7 @@
 #include "Framework/Commands/GenericCommands.h"
 #include "Editor.h"
 #include "PropertyHandle.h"
+#include "AssetTextFilter.h"
 
 #define LOCTEXT_NAMESPACE "SToolInputAssetPicker"
 
@@ -209,7 +210,7 @@ void SToolInputAssetPicker::Construct( const FArguments& InArgs )
 
 	if (AssetViewPtr.IsValid() && !InArgs._AssetPickerConfig.bAutohideSearchBar)
 	{
-		TextFilter = MakeShareable(new FFrontendFilter_Text());
+		TextFilter = MakeShared<FAssetTextFilter>();
 		bool bClassNamesProvided = (InArgs._AssetPickerConfig.Filter.ClassPaths.Num() != 1);
 		TextFilter->SetIncludeClassName(bClassNamesProvided || AssetViewPtr->IsIncludingClassNames());
 		TextFilter->SetIncludeAssetPath(AssetViewPtr->IsIncludingAssetPaths());
@@ -295,12 +296,10 @@ void SToolInputAssetPicker::SetSearchBoxText(const FText& InSearchText)
 		TextFilter->SetRawFilterText(InSearchText);
 		if (InSearchText.IsEmpty())
 		{
-			FrontendFilters->Remove(TextFilter);
 			AssetViewPtr->SetUserSearching(false);
 		}
 		else
 		{
-			FrontendFilters->Add(TextFilter);
 			AssetViewPtr->SetUserSearching(true);
 		}
 	}
diff --git a/Engine/Plugins/Experimental/MeshModelingToolsetExp/Source/ModelingEditorUI/Public/ModelingWidgets/SToolInputAssetPicker.h b/Engine/Plugins/Experimental/MeshModelingToolsetExp/Source/ModelingEditorUI/Public/ModelingWidgets/SToolInputAssetPicker.h
--- a/Engine/Plugins/Experimental/MeshModelingToolsetExp/Source/ModelingEditorUI/Public/ModelingWidgets/SToolInputAssetPicker.h
+++ b/Engine/Plugins/Experimental/MeshModelingToolsetExp/Source/ModelingEditorUI/Public/ModelingWidgets/SToolInputAssetPicker.h
@@ -12,7 +12,7 @@
 #include "IContentBrowserSingleton.h"
 #include "SourcesData.h"
 
-class FFrontendFilter_Text;
+class FAssetTextFilter;
 class FUICommandList;
 class SAssetSearchBox;
 class SAssetView;
@@ -118,7 +118,7 @@ private:
 
 	/** Filters needed for filtering the assets */
 	TSharedPtr< FAssetFilterCollectionType > FilterCollection;
-	TSharedPtr< FFrontendFilter_Text > TextFilter;
+	TSharedPtr<FAssetTextFilter> TextFilter;
 
 	EAssetTypeCategories::Type DefaultFilterMenuExpansion;
 
diff --git a/Engine/Source/Editor/ContentBrowser/Private/AssetTextFilter.cpp b/Engine/Source/Editor/ContentBrowser/Private/AssetTextFilter.cpp
new file mode 100644
--- /dev/null
+++ b/Engine/Source/Editor/ContentBrowser/Private/AssetTextFilter.cpp
@@ -0,0 +1,528 @@
+// Copyright Epic Games, Inc. All Rights Reserved.
+
+#include "AssetTextFilter.h"
+
+#include "Algo/Compare.h"
+#include "CollectionManagerModule.h"
+#include "CollectionManagerTypes.h"
+#include "Filters.h"
+#include "ICollectionManager.h"
+#include "String/ParseTokens.h"
+
+namespace AssetTextFilter
+{
+/** Keys used by TestComplexExpression */
+const FName NameKeyName("Name");
+const FName PathKeyName("Path");
+const FName ClassKeyName("Class");
+const FName TypeKeyName("Type");
+const FName CollectionKeyName("Collection");
+const FName TagKeyName("Tag");
+
+FRWLock CustomHandlerLock;
+TArray<IAssetTextFilterHandler*> CustomHandlers;
+} // namespace AssetTextFilter
+
+FAssetTextFilter::FAssetTextFilter()
+	: ReferencedDynamicCollections()
+	, TextFilterExpressionEvaluator(ETextFilterExpressionEvaluatorMode::Complex)
+{
+	FCollectionManagerModule& CollectionManagerModule = FCollectionManagerModule::GetModule();
+
+	// We need to watch for collection changes so that we can keep ReferencedDynamicCollections up-to-date
+	ICollectionManager& Manager = CollectionManagerModule.Get();
+	OnCollectionCreatedHandle = Manager.OnCollectionCreated().AddRaw(this, &FAssetTextFilter::HandleCollectionCreated);
+	OnCollectionDestroyedHandle = Manager.OnCollectionDestroyed().AddRaw(this, &FAssetTextFilter::HandleCollectionDestroyed);
+	OnCollectionRenamedHandle = Manager.OnCollectionRenamed().AddRaw(this, &FAssetTextFilter::HandleCollectionRenamed);
+	OnCollectionUpdatedHandle = Manager.OnCollectionUpdated().AddRaw(this, &FAssetTextFilter::HandleCollectionUpdated);
+}
+
+FAssetTextFilter::~FAssetTextFilter()
+{
+	// Check IsModuleAvailable as we might be in the process of shutting down...
+	if (FCollectionManagerModule::IsModuleAvailable())
+	{
+		FCollectionManagerModule& CollectionManagerModule = FCollectionManagerModule::GetModule();
+
+		ICollectionManager& Manager = CollectionManagerModule.Get();
+		Manager.OnCollectionCreated().Remove(OnCollectionCreatedHandle);
+		Manager.OnCollectionDestroyed().Remove(OnCollectionDestroyedHandle);
+		Manager.OnCollectionRenamed().Remove(OnCollectionRenamedHandle);
+		Manager.OnCollectionUpdated().Remove(OnCollectionUpdatedHandle);
+	}
+}
+
+bool FAssetTextFilter::IsEmpty() const
+{
+	return TextFilterExpressionEvaluator.GetFilterText().IsEmpty() && CustomTextFilters.IsEmpty();
+}
+
+TSharedPtr<FCompiledAssetTextFilter> FAssetTextFilter::Compile()
+{
+	TSharedPtr<FTextFilterExpressionEvaluator> CompiledEvaluator;
+	TSharedPtr<const TArray<FCollectionNameType>> SharedReferencedDynamicCollections;
+	if (CustomTextFilters.Num() > 0)
+	{
+		// Combine main filter and custom saved queries with AND semantics
+		TStringBuilder<2048> Builder;
+		if (!TextFilterExpressionEvaluator.GetFilterText().IsEmpty())
+		{
+			Builder << TEXTVIEW("(") << TextFilterExpressionEvaluator.GetFilterText().ToString() << TEXTVIEW(")");
+		}
+		for (const FText& Text : CustomTextFilters)
+		{
+			if (Builder.Len() > 0)
+			{
+				Builder << TEXTVIEW(" AND ");
+			}
+			Builder << TEXTVIEW("(") << Text.ToString() << TEXTVIEW(")");
+		}
+		CompiledEvaluator = MakeShared<FTextFilterExpressionEvaluator>(ETextFilterExpressionEvaluatorMode::Complex);
+		CompiledEvaluator->SetFilterText(FText::FromStringView(Builder.ToView()));
+	
+		TArray<FCollectionNameType> CombinedReferencedDynamicCollections;
+		TextFilterExpressionEvaluator.TestTextFilter(FFrontendFilter_GatherDynamicCollectionsExpressionContext(CombinedReferencedDynamicCollections));
+		if (CombinedReferencedDynamicCollections.Num())
+		{
+			SharedReferencedDynamicCollections = MakeShared<const TArray<FCollectionNameType>>(MoveTemp(CombinedReferencedDynamicCollections));
+		}
+	}
+	else
+	{
+		CompiledEvaluator = MakeShared<FTextFilterExpressionEvaluator>(TextFilterExpressionEvaluator);
+		if (bReferencedDynamicCollectionsDirty)
+		{
+			ReferencedDynamicCollections.Reset();
+			TextFilterExpressionEvaluator.TestTextFilter(FFrontendFilter_GatherDynamicCollectionsExpressionContext(ReferencedDynamicCollections));
+		}
+		if (ReferencedDynamicCollections.Num())
+		{
+			SharedReferencedDynamicCollections = MakeShared<const TArray<FCollectionNameType>>(ReferencedDynamicCollections);
+		}
+	}
+
+	// "Include" flags were always propagated to FFrontendFilter_CustomText, we could in theory save that as part of the data, which would mean we would need separate evaluators.
+	return MakeShared<FCompiledAssetTextFilter>(FCompiledAssetTextFilter::FPrivateToken(), CompiledEvaluator.ToSharedRef(), SharedReferencedDynamicCollections,
+		bIncludeClassName, bIncludeAssetPath, bIncludeCollectionNames);
+}
+
+void FAssetTextFilter::SetCustomTextFilters(TArray<FText> InQueries)
+{
+	if (Algo::Compare(InQueries, CustomTextFilters, [](const FText& A, const FText& B) { return A.EqualTo(B); }) == false)
+	{
+		CustomTextFilters = MoveTemp(InQueries);
+		BroadcastChangedEvent(); // This is likely not strictly necessary as these queries will come from the content browser in the first place 
+	}
+}
+
+FText FAssetTextFilter::GetRawFilterText() const
+{
+	return TextFilterExpressionEvaluator.GetFilterText();
+}
+
+void FAssetTextFilter::SetRawFilterText(const FText& InFilterText)
+{
+	if (TextFilterExpressionEvaluator.SetFilterText(InFilterText))
+	{
+		bReferencedDynamicCollectionsDirty = true;
+
+		// Will trigger a re-filter with the new text
+		BroadcastChangedEvent();
+	}
+}
+
+FText FAssetTextFilter::GetFilterErrorText() const
+{
+	return TextFilterExpressionEvaluator.GetFilterErrorText();
+}
+
+void FAssetTextFilter::SetIncludeClassName(const bool InIncludeClassName)
+{
+	if (bIncludeClassName != InIncludeClassName)
+	{
+		bIncludeClassName = InIncludeClassName;
+
+		// Will trigger a re-filter with the new setting
+		BroadcastChangedEvent();
+	}
+}
+
+void FAssetTextFilter::SetIncludeAssetPath(const bool InIncludeAssetPath)
+{
+	if (bIncludeAssetPath != InIncludeAssetPath)
+	{
+		bIncludeAssetPath = InIncludeAssetPath;
+
+		// Will trigger a re-filter with the new setting
+		BroadcastChangedEvent();
+	}
+}
+
+bool FAssetTextFilter::GetIncludeAssetPath() const
+{
+	return bIncludeAssetPath;
+}
+
+void FAssetTextFilter::SetIncludeCollectionNames(const bool InIncludeCollectionNames)
+{
+	if (bIncludeCollectionNames != InIncludeCollectionNames)
+	{
+		bIncludeCollectionNames = InIncludeCollectionNames;
+
+		// Will trigger a re-filter with the new collections
+		BroadcastChangedEvent();
+	}
+}
+
+bool FAssetTextFilter::GetIncludeCollectionNames() const
+{
+	return bIncludeCollectionNames;
+}
+
+void FAssetTextFilter::HandleCollectionCreated(const FCollectionNameType& Collection)
+{
+	bReferencedDynamicCollectionsDirty = true;
+
+	// Will trigger a re-filter with the new collections
+	BroadcastChangedEvent();
+}
+
+void FAssetTextFilter::HandleCollectionDestroyed(const FCollectionNameType& Collection)
+{
+	if (ReferencedDynamicCollections.Contains(Collection))
+	{
+		bReferencedDynamicCollectionsDirty = true;
+
+		// Will trigger a re-filter with the new collections
+		BroadcastChangedEvent();
+	}
+}
+
+void FAssetTextFilter::HandleCollectionRenamed(
+	const FCollectionNameType& OriginalCollection, const FCollectionNameType& NewCollection)
+{
+	int32 FoundIndex = INDEX_NONE;
+	if (ReferencedDynamicCollections.Find(OriginalCollection, FoundIndex))
+	{
+		ReferencedDynamicCollections[FoundIndex] = NewCollection;
+	}
+}
+
+void FAssetTextFilter::HandleCollectionUpdated(const FCollectionNameType& Collection)
+{
+	bReferencedDynamicCollectionsDirty = true;
+
+	// Will trigger a re-filter with the new collections
+	BroadcastChangedEvent();
+}
+
+FCompiledAssetTextFilter::FCompiledAssetTextFilter(
+	FCompiledAssetTextFilter::FPrivateToken,
+	TSharedRef<const FTextFilterExpressionEvaluator> InSharedEvaluator,
+	TSharedPtr<const TArray<FCollectionNameType>> InSharedReferencedDynamicCollections,
+	bool InIncludeClassName,
+	bool InIncludeAssetPath,
+	bool InIncludeCollectionNames)
+	: Evaluator(MoveTemp(InSharedEvaluator))
+	, ReferencedDynamicCollections(MoveTemp(InSharedReferencedDynamicCollections))
+	, bIncludeClassName(InIncludeClassName)
+	, bIncludeAssetPath(InIncludeAssetPath)
+	, bIncludeCollectionNames(InIncludeCollectionNames)
+{
+	CollectionManager = &FCollectionManagerModule::GetModule().Get();
+	TextBuffer.Reserve(2048);
+}
+
+FCompiledAssetTextFilter FCompiledAssetTextFilter::CloneForThreading()
+{
+	return FCompiledAssetTextFilter(*this);
+}
+
+bool FCompiledAssetTextFilter::PassesFilter(FAssetFilterType InItem)
+{
+	using namespace AssetTextFilter;
+	FReadScopeLock Guard(CustomHandlerLock);
+
+	AssetPtr = &InItem;
+    TextBuffer.Reset();
+	AssetCollectionNames.Reset();
+
+	int32 DisplayNameLen = 0;
+	int32 AssetPathLen = 0;
+	int32 ExportTextPathLen = 0;
+
+    TextBuffer.Append(AssetPtr->GetDisplayName().ToString());
+	DisplayNameLen = TextBuffer.Len();
+
+	if (bIncludeAssetPath)
+	{
+		// Get the full asset path, and also split it so we can compare each part in the filter
+		AssetPtr->GetVirtualPath().AppendString(TextBuffer);
+        
+        // chop off item name if necessary
+		{
+			int32 LastSlashIndex = INDEX_NONE;
+			int32 LastDotIndex = INDEX_NONE;
+			if (TextBuffer.FindLastChar(TEXT('/'), LastSlashIndex) && TextBuffer.FindLastChar(TEXT('.'), LastDotIndex)
+				&& LastDotIndex > LastSlashIndex && LastSlashIndex >= DisplayNameLen)
+			{
+				TextBuffer.LeftInline(LastDotIndex, EAllowShrinking::No);
+			}
+		}
+		AssetPathLen = TextBuffer.Len() - DisplayNameLen;
+
+		if (bIncludeClassName && !AssetPtr->IsFolder())
+		{
+			// Get the full export text path as people sometimes search by copying this (requires class and asset path
+			// search to be enabled in order to match)
+			// TODO: this allocates a temporary FString inside the backends
+			AssetPtr->AppendItemReference(TextBuffer);
+			ExportTextPathLen = TextBuffer.Len() - AssetPathLen;
+		}
+	}
+
+	if (CollectionManager && bIncludeCollectionNames)
+	{
+		FSoftObjectPath ItemCollectionId;
+		if (AssetPtr->TryGetCollectionId(ItemCollectionId))
+		{
+			CollectionManager->GetCollectionsContainingObject(
+				ItemCollectionId,
+				ECollectionShareType::CST_All,
+				AssetCollectionNames,
+				ECollectionRecursionFlags::SelfAndChildren);
+
+			if (ReferencedDynamicCollections.IsValid())
+			{
+				// Test the dynamic collections from the active query against the current asset
+				// We can do this as a flat list since FFrontendFilter_GatherDynamicCollectionsExpressionContext has already
+				// taken care of processing the recursion
+				for (const FCollectionNameType& DynamicCollection : *ReferencedDynamicCollections)
+				{
+					bool bPassesCollectionFilter = false;
+					CollectionManager->TestDynamicQuery(
+						DynamicCollection.Name, DynamicCollection.Type, *this, bPassesCollectionFilter);
+					if (bPassesCollectionFilter)
+					{
+						AssetCollectionNames.AddUnique(DynamicCollection.Name);
+					}
+				}
+			}
+		}
+	}
+    
+    // Convert entire text buffer to uppercase
+    // Assumes that for a TCHAR buffer this won't change number of code points 
+	TextBuffer.ToUpperInline();
+	
+	const TCHAR* Cursor = *TextBuffer;
+	AssetDisplayName = FStringView(Cursor, DisplayNameLen);
+	Cursor += DisplayNameLen;
+	AssetFullPath = FStringView(Cursor, AssetPathLen);
+	Cursor += AssetPathLen;
+	AssetExportTextPath  = FStringView(Cursor, ExportTextPathLen);
+	AssetExportTextPath.TrimStartAndEndInline(); // Backends try to separate export text paths with newlines 
+	Cursor += ExportTextPathLen;
+
+	const bool bMatched = Evaluator->TestTextFilter(*this);
+	AssetPtr = nullptr;
+	return bMatched;
+}
+
+bool FCompiledAssetTextFilter::TestBasicStringExpression(
+	const FTextFilterString& InValue, const ETextFilterTextComparisonMode InTextComparisonMode) const
+{
+	using namespace AssetTextFilter;
+	for (IAssetTextFilterHandler* Handler : CustomHandlers)
+	{
+		bool bIsHandlerMatch = false;
+		if (Handler->HandleTextFilterValue(*AssetPtr, InValue, InTextComparisonMode, bIsHandlerMatch))
+		{
+			return bIsHandlerMatch;
+		}
+	}
+
+	if (InValue.CompareName(AssetPtr->GetItemName(), InTextComparisonMode))
+	{
+		return true;
+	}
+
+	if (InValue.CompareFStringView(AssetDisplayName, InTextComparisonMode))
+	{
+		return true;
+	}
+
+	if (bIncludeAssetPath)
+	{
+		if (InValue.CompareFStringView(AssetFullPath, InTextComparisonMode))
+		{
+			return true;
+		}
+
+		bool bSuccess = false;
+		UE::String::ParseTokens(AssetFullPath, TEXTVIEW("/"), [&bSuccess, &InValue, InTextComparisonMode](FStringView Element){
+			bSuccess = bSuccess || InValue.CompareFStringView(Element, InTextComparisonMode);
+		});
+		if (bSuccess)
+		{
+			return true;
+		}
+	}
+
+	if (bIncludeClassName)
+	{
+		const FContentBrowserItemDataAttributeValue ClassValue = AssetPtr->GetItemAttribute(NAME_Class);
+		if (ClassValue.IsValid() && InValue.CompareName(ClassValue.GetValue<FName>(), InTextComparisonMode))
+		{
+			return true;
+		}
+	}
+
+	if (bIncludeClassName && bIncludeAssetPath)
+	{
+		// Only test this if we're searching the class name and asset path too, as the exported text contains the type
+		// and path in the string
+		if (InValue.CompareFStringView(AssetExportTextPath, InTextComparisonMode))
+		{
+			return true;
+		}
+	}
+
+	for (const FName& AssetCollectionName : AssetCollectionNames)
+	{
+		if (InValue.CompareName(AssetCollectionName, InTextComparisonMode))
+		{
+			return true;
+		}
+	}
+
+	return false;
+}
+
+bool FCompiledAssetTextFilter::TestComplexExpression(
+	const FName& InKey,
+	const FTextFilterString& InValue,
+	const ETextFilterComparisonOperation InComparisonOperation,
+	const ETextFilterTextComparisonMode InTextComparisonMode) const
+{
+	using namespace AssetTextFilter;
+	for (IAssetTextFilterHandler* Handler : CustomHandlers)
+	{
+		bool bIsHandlerMatch = false;
+		if (Handler->HandleTextFilterKeyValue(*AssetPtr, InKey, InValue, InComparisonOperation, InTextComparisonMode, bIsHandlerMatch))
+		{
+			return bIsHandlerMatch;
+		}
+	}
+
+	// Special case for the asset name, as this isn't contained within the asset registry meta-data
+	if (InKey == NameKeyName)
+	{
+		// Names can only work with Equal or NotEqual type tests
+		if (InComparisonOperation != ETextFilterComparisonOperation::Equal
+			&& InComparisonOperation != ETextFilterComparisonOperation::NotEqual)
+		{
+			return false;
+		}
+
+		const bool bIsMatch = TextFilterUtils::TestBasicStringExpression(AssetPtr->GetItemName(), InValue, InTextComparisonMode);
+		return (InComparisonOperation == ETextFilterComparisonOperation::Equal) ? bIsMatch : !bIsMatch;
+	}
+
+	// Special case for the asset path, as this isn't contained within the asset registry meta-data
+	if (InKey == PathKeyName)
+	{
+		// Paths can only work with Equal or NotEqual type tests
+		if (InComparisonOperation != ETextFilterComparisonOperation::Equal
+			&& InComparisonOperation != ETextFilterComparisonOperation::NotEqual)
+		{
+			return false;
+		}
+
+		// If the comparison mode is partial, then we only need to test the ObjectPath as that contains the other two as
+		// sub-strings
+		bool bIsMatch = false;
+		if (InTextComparisonMode == ETextFilterTextComparisonMode::Partial)
+		{
+			bIsMatch =
+				TextFilterUtils::TestBasicStringExpression(AssetPtr->GetVirtualPath(), InValue, InTextComparisonMode);
+		}
+		else
+		{
+			bIsMatch =
+				TextFilterUtils::TestBasicStringExpression(AssetPtr->GetVirtualPath(), InValue, InTextComparisonMode)
+				|| (!AssetFullPath.IsEmpty()
+					&& TextFilterUtils::TestBasicStringExpression(FString(AssetFullPath), InValue, InTextComparisonMode));
+		}
+		return (InComparisonOperation == ETextFilterComparisonOperation::Equal) ? bIsMatch : !bIsMatch;
+	}
+
+	// Special case for the asset type, as this isn't contained within the asset registry meta-data
+	if (InKey == ClassKeyName || InKey == TypeKeyName)
+	{
+		// Class names can only work with Equal or NotEqual type tests
+		if (InComparisonOperation != ETextFilterComparisonOperation::Equal
+			&& InComparisonOperation != ETextFilterComparisonOperation::NotEqual)
+		{
+			return false;
+		}
+
+		const FContentBrowserItemDataAttributeValue ClassValue = AssetPtr->GetItemAttribute(NAME_Class);
+		const bool bIsMatch =
+			ClassValue.IsValid()
+			&& TextFilterUtils::TestBasicStringExpression(ClassValue.GetValue<FName>(), InValue, InTextComparisonMode);
+		return (InComparisonOperation == ETextFilterComparisonOperation::Equal) ? bIsMatch : !bIsMatch;
+	}
+
+	// Special case for collections, as these aren't contained within the asset registry meta-data
+	if (InKey == CollectionKeyName || InKey == TagKeyName)
+	{
+		// Collections can only work with Equal or NotEqual type tests
+		if (InComparisonOperation != ETextFilterComparisonOperation::Equal
+			&& InComparisonOperation != ETextFilterComparisonOperation::NotEqual)
+		{
+			return false;
+		}
+
+		bool bFoundMatch = false;
+		for (const FName& AssetCollectionName : AssetCollectionNames)
+		{
+			if (TextFilterUtils::TestBasicStringExpression(AssetCollectionName, InValue, InTextComparisonMode))
+			{
+				bFoundMatch = true;
+				break;
+			}
+		}
+
+		return (InComparisonOperation == ETextFilterComparisonOperation::Equal) ? bFoundMatch : !bFoundMatch;
+	}
+
+	// Generic handling for anything in the asset meta-data
+	{
+		const FContentBrowserItemDataAttributeValue AttributeValue = AssetPtr->GetItemAttribute(InKey);
+		if (AttributeValue.IsValid())
+		{
+			return TextFilterUtils::TestComplexExpression(
+				AttributeValue.GetValue<FString>(),
+				InValue,
+				InComparisonOperation,
+				InTextComparisonMode);
+		}
+	}
+
+	return false;
+}
+
+void IAssetTextFilterHandler::RegisterHandler()
+{
+	using namespace AssetTextFilter;
+	FWriteScopeLock Guard(CustomHandlerLock);
+	CustomHandlers.Add(this);
+}
+
+void IAssetTextFilterHandler::UnregisterHandler()
+{
+	using namespace AssetTextFilter;
+	FWriteScopeLock Guard(CustomHandlerLock);
+	CustomHandlers.Remove(this);
+}
\ No newline at end of file
diff --git a/Engine/Source/Editor/ContentBrowser/Private/AssetViewTypes.cpp b/Engine/Source/Editor/ContentBrowser/Private/AssetViewTypes.cpp
--- a/Engine/Source/Editor/ContentBrowser/Private/AssetViewTypes.cpp
+++ b/Engine/Source/Editor/ContentBrowser/Private/AssetViewTypes.cpp
@@ -35,16 +35,27 @@ FAssetViewItem::FAssetViewItem(const FContentBrowserItemData& InItemData)
 	checkf(Item.IsValid(), TEXT("FAssetViewItem was constructed from an invalid item!"));
 }
 
+void FAssetViewItem::ResetItemData(FContentBrowserItemData InItemData)
+{
+	Item = FContentBrowserItem{MoveTemp(InItemData)};
+	// Do not broadcast event here, it will be broadcast on the main thread after bulk building/recycling of items
+}
+
+void FAssetViewItem::BroadcastItemDataChanged()
+{
+	ItemDataChangedEvent.Broadcast();
+}
+
 void FAssetViewItem::AppendItemData(const FContentBrowserItem& InItem)
 {
 	Item.Append(InItem);
-	ItemDataChangedEvent.Broadcast();
+	// Do not broadcast event here, caller is responsible for broadcasting in a threadsafe way
 }
 
 void FAssetViewItem::AppendItemData(const FContentBrowserItemData& InItemData)
 {
 	Item.Append(InItemData);
-	ItemDataChangedEvent.Broadcast();
+	// Do not broadcast event here, caller is responsible for broadcasting in a threadsafe way
 }
 
 void FAssetViewItem::RemoveItemData(const FContentBrowserItem& InItem)
@@ -76,12 +87,19 @@ void FAssetViewItem::RemoveItemData(const FContentBrowserMinimalItemData& InItem
 
 void FAssetViewItem::ClearCachedCustomColumns()
 {
+	check(IsInGameThread());
 	CachedCustomColumnData.Reset();
 	CachedCustomColumnDisplayText.Reset();
 }
 
 void FAssetViewItem::CacheCustomColumns(TArrayView<const FAssetViewCustomColumn> CustomColumns, const bool bUpdateSortData, const bool bUpdateDisplayText, const bool bUpdateExisting)
 {
+	check(IsInGameThread());
+	if (bUpdateExisting && CachedCustomColumnData.IsEmpty())
+	{
+		return;
+	}
+
 	for (const FAssetViewCustomColumn& Column : CustomColumns)
 	{
 		FAssetData ItemAssetData;
@@ -195,17 +213,11 @@ FSimpleDelegate& FAssetViewItem::OnRenameCanceled()
 	return RenameCanceledEvent;
 }
 
-bool FAssetViewItem::ShouldRenameWhenScrolledIntoView() const
-{
-	return bRenameWhenScrolledIntoView;
-}
-
-void FAssetViewItem::RenameWhenScrolledIntoView()
+FString FAssetViewItem::ItemToString_Debug(TSharedPtr<FAssetViewItem> AssetItem) 
 {
-	bRenameWhenScrolledIntoView = true;
-}
-
-void FAssetViewItem::ClearRenameWhenScrolledIntoView()
-{
-	bRenameWhenScrolledIntoView = false;
+	if (AssetItem.IsValid())
+	{
+		return AssetItem->GetItem().GetVirtualPath().ToString();
+	}
+	return TEXT("nullptr");
 }
diff --git a/Engine/Source/Editor/ContentBrowser/Private/AssetViewTypes.h b/Engine/Source/Editor/ContentBrowser/Private/AssetViewTypes.h
new file mode 100644
--- /dev/null
+++ b/Engine/Source/Editor/ContentBrowser/Private/AssetViewTypes.h
@@ -0,0 +1,105 @@
+// Copyright Epic Games, Inc. All Rights Reserved.
+
+#pragma once
+
+#include "Containers/ArrayView.h"
+#include "Containers/Map.h"
+#include "ContentBrowserItem.h"
+#include "CoreMinimal.h"
+#include "Delegates/Delegate.h"
+#include "Templates/Tuple.h"
+#include "UObject/NameTypes.h"
+#include "UObject/Object.h"
+
+class FContentBrowserItemData;
+class FString;
+class FText;
+struct FAssetViewCustomColumn;
+
+/** An item (folder or file) displayed in the asset view */
+class FAssetViewItem
+{
+public:
+	FAssetViewItem() = default;
+
+	explicit FAssetViewItem(FContentBrowserItem&& InItem);
+	explicit FAssetViewItem(const FContentBrowserItem& InItem);
+
+	explicit FAssetViewItem(FContentBrowserItemData&& InItemData);
+	explicit FAssetViewItem(const FContentBrowserItemData& InItemData);
+
+	// When recycling an object, clear the item data and replace it with the given data.
+	void ResetItemData(FContentBrowserItemData InItemData);
+
+	void AppendItemData(const FContentBrowserItem& InItem);
+
+	void AppendItemData(const FContentBrowserItemData& InItemData);
+
+	void RemoveItemData(const FContentBrowserItem& InItem);
+
+	void RemoveItemData(const FContentBrowserItemData& InItemData);
+
+	void RemoveItemData(const FContentBrowserMinimalItemData& InItemKey);
+
+	/** Clear cached custom column data */
+	void ClearCachedCustomColumns();
+
+	/**
+	 * Updates cached custom column data (only does something for files) 
+	 * @param bUpdateExisting If true, only updates existing columns, if false only adds missing columns 
+	 */
+	void CacheCustomColumns(TArrayView<const FAssetViewCustomColumn> CustomColumns, const bool bUpdateSortData, const bool bUpdateDisplayText, const bool bUpdateExisting);
+	
+	/** Get the display value of a custom column on this item */
+	bool GetCustomColumnDisplayValue(const FName ColumnName, FText& OutText) const;
+
+	/** Get the value (and optionally also the type) of a custom column on this item */
+	bool GetCustomColumnValue(const FName ColumnName, FString& OutString, UObject::FAssetRegistryTag::ETagType* OutType = nullptr) const;
+
+	/** Get the value (and optionally also the type) of a named tag on this item */
+	bool GetTagValue(const FName Tag, FString& OutString, UObject::FAssetRegistryTag::ETagType* OutType = nullptr) const;
+
+	/** Get the underlying Content Browser item */
+	const FContentBrowserItem& GetItem() const;
+
+	bool IsFolder() const;
+
+	bool IsFile() const;
+
+	bool IsTemporary() const;
+
+	// Called when the view explicitly wants to notify widgets of changes
+	// Not called during bulk rebuilds when the view will be re-populated even if items are being recycled
+	void BroadcastItemDataChanged();
+
+	/** Get the event fired when the data for this item changes */
+	FSimpleMulticastDelegate& OnItemDataChanged();
+
+	/** Get the event fired whenever a rename is requested */
+	FSimpleDelegate& OnRenameRequested();
+
+	/** Get the event fired whenever a rename is canceled */
+	FSimpleDelegate& OnRenameCanceled();
+
+	/** Helper function to turn an item into a string for debugging */
+	static FString ItemToString_Debug(TSharedPtr<FAssetViewItem> AssetItem);
+
+private:
+	/** Underlying Content Browser item data */
+	FContentBrowserItem Item;
+
+	/** An event to fire when the data for this item changes */
+	FSimpleMulticastDelegate ItemDataChangedEvent;
+
+	/** Broadcasts whenever a rename is requested */
+	FSimpleDelegate RenameRequestedEvent;
+
+	/** Broadcasts whenever a rename is canceled */
+	FSimpleDelegate RenameCanceledEvent;
+
+	/** Map of values/types for custom columns */
+	TMap<FName, TTuple<FString, UObject::FAssetRegistryTag::ETagType>> CachedCustomColumnData;
+	
+	/** Map of display text for custom columns */
+	TMap<FName, FText> CachedCustomColumnDisplayText;
+};
diff --git a/Engine/Source/Editor/ContentBrowser/Private/Filters.h b/Engine/Source/Editor/ContentBrowser/Private/Filters.h
--- a/Engine/Source/Editor/ContentBrowser/Private/Filters.h
+++ b/Engine/Source/Editor/ContentBrowser/Private/Filters.h
@@ -8,6 +8,7 @@
 #include "Filters/FilterBase.h"
 #include "Filters/GenericFilter.h"
 #include "FrontendFilterBase.h"
+#include "Misc/TextFilterExpressionEvaluator.h"
 #include "Templates/SharedPointer.h"
 
 /** 
@@ -55,6 +56,48 @@ public:
 	virtual bool PassesFilter(FAssetFilterType InItem) const override { return true; }
 };
 
+/** 
+ * Expression context which gathers up the names of any dynamic collections being referenced by the current query 
+ * Private utility class for FAssetTextFilter and deprecated FFrontendFilter_Text
+ */
+class FFrontendFilter_GatherDynamicCollectionsExpressionContext : public ITextFilterExpressionContext
+{
+public:
+	FFrontendFilter_GatherDynamicCollectionsExpressionContext(TArray<FCollectionNameType>& OutReferencedDynamicCollections);
+	~FFrontendFilter_GatherDynamicCollectionsExpressionContext();
+
+	virtual bool TestBasicStringExpression(const FTextFilterString& InValue, const ETextFilterTextComparisonMode InTextComparisonMode) const override;
+	virtual bool TestComplexExpression(const FName& InKey, const FTextFilterString& InValue, const ETextFilterComparisonOperation InComparisonOperation, const ETextFilterTextComparisonMode InTextComparisonMode) const override;
+
+private:
+	bool TestAgainstAvailableCollections(const FTextFilterString& InValue, const ETextFilterTextComparisonMode InTextComparisonMode) const;
+
+	/** Contains a collection name along with its recursion depth in the dynamic query - used so we can test them depth first */
+	struct FDynamicCollectionNameAndDepth
+	{
+		FDynamicCollectionNameAndDepth(FCollectionNameType InCollection, const int32 InRecursionDepth)
+			: Collection(InCollection)
+			, RecursionDepth(InRecursionDepth)
+		{
+		}
+
+		FCollectionNameType Collection;
+		int32 RecursionDepth;
+	};
+
+	/** The currently available dynamic collections */
+	TArray<FCollectionNameType> AvailableDynamicCollections;
+
+	/** This will be populated with any dynamic collections that are being referenced by the current query - these collections may not all match when tested against the actual asset data */
+	TArray<FCollectionNameType>& ReferencedDynamicCollections;
+
+	/** Dynamic collections that have currently be found as part of the query (or recursive sub-query) */
+	mutable TArray<FDynamicCollectionNameAndDepth> FoundDynamicCollections;
+
+	/** Incremented when we test a sub-query, decremented once we're done */
+	mutable int32 CurrentRecursionDepth;
+};
+
 // Non-frontend filter which modifies content browser backend query to exclude folders belonging to other developers 
 class FFilter_ShowOtherDevelopers : public FFrontendFilter
 {
diff --git a/Engine/Source/Editor/ContentBrowser/Private/FrontendFilters.cpp b/Engine/Source/Editor/ContentBrowser/Private/FrontendFilters.cpp
--- a/Engine/Source/Editor/ContentBrowser/Private/FrontendFilters.cpp
+++ b/Engine/Source/Editor/ContentBrowser/Private/FrontendFilters.cpp
@@ -21,6 +21,7 @@
 #include "Delegates/Delegate.h"
 #include "Editor.h"
 #include "Engine/World.h"
+#include "Filters.h"
 #include "Framework/Commands/UIAction.h"
 #include "Framework/Commands/UICommandInfo.h"
 #include "Framework/MultiBox/MultiBoxBuilder.h"
@@ -108,144 +109,115 @@ namespace FrontendFilterHelper
 /////////////////////////////////////////
 // FFrontendFilter_Text
 /////////////////////////////////////////
+namespace UE::ContentBrowser
+{
+	/** Keys used by FFrontendFilter_GatherDynamicCollectionsExpressionContext::TestComplexExpression */
+	const FName CollectionKeyName("Collection");
+	const FName TagKeyName("Tag");
+}
 
-/** Expression context which gathers up the names of any dynamic collections being referenced by the current query */
-class FFrontendFilter_GatherDynamicCollectionsExpressionContext : public ITextFilterExpressionContext
+FFrontendFilter_GatherDynamicCollectionsExpressionContext::FFrontendFilter_GatherDynamicCollectionsExpressionContext(TArray<FCollectionNameType>& OutReferencedDynamicCollections)
+	: AvailableDynamicCollections()
+	, ReferencedDynamicCollections(OutReferencedDynamicCollections)
+	, CurrentRecursionDepth(0)
 {
-public:
-	FFrontendFilter_GatherDynamicCollectionsExpressionContext(TArray<FCollectionNameType>& OutReferencedDynamicCollections)
-		: AvailableDynamicCollections()
-		, ReferencedDynamicCollections(OutReferencedDynamicCollections)
-		, CurrentRecursionDepth(0)
-		, CollectionKeyName("Collection")
-		, TagKeyName("Tag")
+	if (FCollectionManagerModule::IsModuleAvailable())
 	{
-		if (FCollectionManagerModule::IsModuleAvailable())
-		{
-			FCollectionManagerModule& CollectionManagerModule = FCollectionManagerModule::GetModule();
+		FCollectionManagerModule& CollectionManagerModule = FCollectionManagerModule::GetModule();
 
-			TArray<FCollectionNameType> AvailableCollections;
-			CollectionManagerModule.Get().GetCollections(AvailableCollections);
+		TArray<FCollectionNameType> AvailableCollections;
+		CollectionManagerModule.Get().GetCollections(AvailableCollections);
 
-			for (const FCollectionNameType& AvailableCollection : AvailableCollections)
+		for (const FCollectionNameType& AvailableCollection : AvailableCollections)
+		{
+			// Only care about dynamic collections
+			ECollectionStorageMode::Type StorageMode = ECollectionStorageMode::Static;
+			CollectionManagerModule.Get().GetCollectionStorageMode(AvailableCollection.Name, AvailableCollection.Type, StorageMode);
+			if (StorageMode != ECollectionStorageMode::Dynamic)
 			{
-				// Only care about dynamic collections
-				ECollectionStorageMode::Type StorageMode = ECollectionStorageMode::Static;
-				CollectionManagerModule.Get().GetCollectionStorageMode(AvailableCollection.Name, AvailableCollection.Type, StorageMode);
-				if (StorageMode != ECollectionStorageMode::Dynamic)
-				{
-					continue;
-				}
-
-				AvailableDynamicCollections.Add(AvailableCollection);
+				continue;
 			}
+
+			AvailableDynamicCollections.Add(AvailableCollection);
 		}
 	}
+}
 
-	~FFrontendFilter_GatherDynamicCollectionsExpressionContext()
+FFrontendFilter_GatherDynamicCollectionsExpressionContext::~FFrontendFilter_GatherDynamicCollectionsExpressionContext()
+{
+	// Sort and populate the final list of referenced dynamic collections
+	FoundDynamicCollections.Sort([](const FDynamicCollectionNameAndDepth& A, const FDynamicCollectionNameAndDepth& B)
 	{
-		// Sort and populate the final list of referenced dynamic collections
-		FoundDynamicCollections.Sort([](const FDynamicCollectionNameAndDepth& A, const FDynamicCollectionNameAndDepth& B)
-		{
-			return A.RecursionDepth > B.RecursionDepth;
-		});
-
-		ReferencedDynamicCollections.Reset();
-		ReferencedDynamicCollections.Reserve(FoundDynamicCollections.Num());
-		for (const auto& FoundDynamicCollection : FoundDynamicCollections)
-		{
-			ReferencedDynamicCollections.Add(FoundDynamicCollection.Collection);
-		}
-	}
+		return A.RecursionDepth > B.RecursionDepth;
+	});
 
-	virtual bool TestBasicStringExpression(const FTextFilterString& InValue, const ETextFilterTextComparisonMode InTextComparisonMode) const override
+	ReferencedDynamicCollections.Reset();
+	ReferencedDynamicCollections.Reserve(FoundDynamicCollections.Num());
+	for (const auto& FoundDynamicCollection : FoundDynamicCollections)
 	{
-		TestAgainstAvailableCollections(InValue, InTextComparisonMode);
-		return false;
+		ReferencedDynamicCollections.Add(FoundDynamicCollection.Collection);
 	}
+}
 
-	virtual bool TestComplexExpression(const FName& InKey, const FTextFilterString& InValue, const ETextFilterComparisonOperation InComparisonOperation, const ETextFilterTextComparisonMode InTextComparisonMode) const override
+bool FFrontendFilter_GatherDynamicCollectionsExpressionContext::TestBasicStringExpression(const FTextFilterString& InValue, const ETextFilterTextComparisonMode InTextComparisonMode) const 
+{
+	TestAgainstAvailableCollections(InValue, InTextComparisonMode);
+	return false;
+}
+
+bool FFrontendFilter_GatherDynamicCollectionsExpressionContext::TestComplexExpression(const FName& InKey, const FTextFilterString& InValue, const ETextFilterComparisonOperation InComparisonOperation, const ETextFilterTextComparisonMode InTextComparisonMode) const 
+{
+	using namespace UE::ContentBrowser;
+
+	// Special case for collections, as these aren't contained within the asset registry meta-data
+	if (InKey == CollectionKeyName || InKey == TagKeyName)
 	{
-		// Special case for collections, as these aren't contained within the asset registry meta-data
-		if (InKey == CollectionKeyName || InKey == TagKeyName)
+		// Collections can only work with Equal or NotEqual type tests
+		if (InComparisonOperation != ETextFilterComparisonOperation::Equal && InComparisonOperation != ETextFilterComparisonOperation::NotEqual)
 		{
-			// Collections can only work with Equal or NotEqual type tests
-			if (InComparisonOperation != ETextFilterComparisonOperation::Equal && InComparisonOperation != ETextFilterComparisonOperation::NotEqual)
-			{
-				return false;
-			}
-
-			TestAgainstAvailableCollections(InValue, InTextComparisonMode);
+			return false;
 		}
 
-		return false;
+		TestAgainstAvailableCollections(InValue, InTextComparisonMode);
 	}
 
-private:
-	bool TestAgainstAvailableCollections(const FTextFilterString& InValue, const ETextFilterTextComparisonMode InTextComparisonMode) const
+	return false;
+}
+
+bool FFrontendFilter_GatherDynamicCollectionsExpressionContext::TestAgainstAvailableCollections(const FTextFilterString& InValue, const ETextFilterTextComparisonMode InTextComparisonMode) const
+{
+	for (const FCollectionNameType& DynamicCollection : AvailableDynamicCollections)
 	{
-		for (const FCollectionNameType& DynamicCollection : AvailableDynamicCollections)
+		const FString DynamicCollectionNameStr = DynamicCollection.Name.ToString();
+		if (TextFilterUtils::TestBasicStringExpression(DynamicCollectionNameStr, InValue, InTextComparisonMode))
 		{
-			const FString DynamicCollectionNameStr = DynamicCollection.Name.ToString();
-			if (TextFilterUtils::TestBasicStringExpression(DynamicCollectionNameStr, InValue, InTextComparisonMode))
+			const bool bCollectionAlreadyProcessed = FoundDynamicCollections.ContainsByPredicate([&DynamicCollection](const FDynamicCollectionNameAndDepth& Other)
 			{
-				const bool bCollectionAlreadyProcessed = FoundDynamicCollections.ContainsByPredicate([&](const FDynamicCollectionNameAndDepth& Other)
-				{
-					return DynamicCollection == Other.Collection;
-				});
+				return DynamicCollection == Other.Collection;
+			});
 
-				if (!bCollectionAlreadyProcessed)
-				{
-					FoundDynamicCollections.Add(FDynamicCollectionNameAndDepth(DynamicCollection, CurrentRecursionDepth));
+			if (!bCollectionAlreadyProcessed)
+			{
+				FoundDynamicCollections.Add(FDynamicCollectionNameAndDepth(DynamicCollection, CurrentRecursionDepth));
 
-					if (FCollectionManagerModule::IsModuleAvailable())
-					{
-						FCollectionManagerModule& CollectionManagerModule = FCollectionManagerModule::GetModule();
-						
-						// Also need to gather any collections referenced by this dynamic collection
-						++CurrentRecursionDepth;
-						bool bUnused = false;
-						CollectionManagerModule.Get().TestDynamicQuery(DynamicCollection.Name, DynamicCollection.Type, *this, bUnused);
-						--CurrentRecursionDepth;
-					}
+				if (FCollectionManagerModule::IsModuleAvailable())
+				{
+					FCollectionManagerModule& CollectionManagerModule = FCollectionManagerModule::GetModule();
+					
+					// Also need to gather any collections referenced by this dynamic collection
+					++CurrentRecursionDepth;
+					bool bUnused = false;
+					CollectionManagerModule.Get().TestDynamicQuery(DynamicCollection.Name, DynamicCollection.Type, *this, bUnused);
+					--CurrentRecursionDepth;
 				}
-
-				return true;
 			}
-		}
 
-		return false;
-	}
-
-	/** Contains a collection name along with its recursion depth in the dynamic query - used so we can test them depth first */
-	struct FDynamicCollectionNameAndDepth
-	{
-		FDynamicCollectionNameAndDepth(FCollectionNameType InCollection, const int32 InRecursionDepth)
-			: Collection(InCollection)
-			, RecursionDepth(InRecursionDepth)
-		{
+			return true;
 		}
+	}
 
-		FCollectionNameType Collection;
-		int32 RecursionDepth;
-	};
-
-	/** The currently available dynamic collections */
-	TArray<FCollectionNameType> AvailableDynamicCollections;
-
-	/** This will be populated with any dynamic collections that are being referenced by the current query - these collections may not all match when tested against the actual asset data */
-	TArray<FCollectionNameType>& ReferencedDynamicCollections;
-
-	/** Dynamic collections that have currently be found as part of the query (or recursive sub-query) */
-	mutable TArray<FDynamicCollectionNameAndDepth> FoundDynamicCollections;
-
-	/** Incremented when we test a sub-query, decremented once we're done */
-	mutable int32 CurrentRecursionDepth;
-
-	/** Keys used by TestComplexExpression */
-	const FName CollectionKeyName;
-	const FName TagKeyName;
-};
+	return false;
+}
 
 /** Expression context to test the given asset data against the current text filter */
 class FFrontendFilter_TextFilterExpressionContext : public ITextFilterExpressionContext
@@ -369,10 +341,12 @@ public:
 	virtual bool TestBasicStringExpression(const FTextFilterString& InValue, const ETextFilterTextComparisonMode InTextComparisonMode) const override
 	{
 		bool bIsHandlerMatch = false;
+PRAGMA_DISABLE_DEPRECATION_WARNINGS
 		if (UTextFilterValueHandlers::HandleTextFilterValue(*AssetPtr, InValue, InTextComparisonMode, bIsHandlerMatch))
 		{
 			return bIsHandlerMatch;
 		}
+PRAGMA_ENABLE_DEPRECATION_WARNINGS
 
 		if (InValue.CompareName(AssetPtr->GetItemName(), InTextComparisonMode))
 		{
@@ -435,10 +409,12 @@ public:
 	virtual bool TestComplexExpression(const FName& InKey, const FTextFilterString& InValue, const ETextFilterComparisonOperation InComparisonOperation, const ETextFilterTextComparisonMode InTextComparisonMode) const override
 	{
 		bool bIsHandlerMatch = false;
+PRAGMA_DISABLE_DEPRECATION_WARNINGS
 		if (UTextFilterKeyValueHandlers::HandleTextFilterKeyValue(*AssetPtr, InKey, InValue, InComparisonOperation, InTextComparisonMode, bIsHandlerMatch))
 		{
 			return bIsHandlerMatch;
 		}
+PRAGMA_ENABLE_DEPRECATION_WARNINGS
 
 		// Special case for the asset name, as this isn't contained within the asset registry meta-data
 		if (InKey == NameKeyName)
@@ -709,7 +685,6 @@ void FFrontendFilter_Text::RebuildReferencedDynamicCollections()
 	TextFilterExpressionEvaluator.TestTextFilter(FFrontendFilter_GatherDynamicCollectionsExpressionContext(ReferencedDynamicCollections));
 }
 
-
 /////////////////////////////////////////
 // FFrontendFilter_CheckedOut
 /////////////////////////////////////////
diff --git a/Engine/Source/Editor/ContentBrowser/Private/SAssetPicker.cpp b/Engine/Source/Editor/ContentBrowser/Private/SAssetPicker.cpp
--- a/Engine/Source/Editor/ContentBrowser/Private/SAssetPicker.cpp
+++ b/Engine/Source/Editor/ContentBrowser/Private/SAssetPicker.cpp
@@ -3,6 +3,7 @@
 #include "SAssetPicker.h"
 
 #include "AssetRegistry/AssetData.h"
+#include "AssetTextFilter.h"
 #include "AssetThumbnail.h"
 #include "CollectionManagerTypes.h"
 #include "ContentBrowserDataFilter.h"
@@ -316,6 +317,11 @@ void SAssetPicker::Construct( const FArguments& InArgs )
 		}
 	}
 
+	if (!InArgs._AssetPickerConfig.bAutohideSearchBar)
+	{
+		TextFilter = MakeShared<FAssetTextFilter>();
+	}
+	
 	// clang-format off
 	VerticalBox->AddSlot()
 	.FillHeight(1.f)
@@ -333,7 +339,8 @@ void SAssetPicker::Construct( const FArguments& InArgs )
 		.OnVisualizeAssetToolTip(InArgs._AssetPickerConfig.OnVisualizeAssetToolTip)
 		.OnAssetToolTipClosing(InArgs._AssetPickerConfig.OnAssetToolTipClosing)
 		.FrontendFilters(FrontendFilters)
-	   	.ShowRedirectors_Lambda([this]() { return ContentBrowserUtils::ShouldShowRedirectors(FilterListPtr); })
+		.TextFilter(TextFilter)
+		.ShowRedirectors_Lambda([this]() { return ContentBrowserUtils::ShouldShowRedirectors(FilterListPtr); })
 		.InitialSourcesData(CurrentSourcesData)
 		.InitialBackendFilter(CurrentBackendFilter)
 		.InitialViewType(InArgs._AssetPickerConfig.InitialAssetViewType)
@@ -385,9 +392,8 @@ void SAssetPicker::Construct( const FArguments& InArgs )
 
 	LoadSettings();
 
-	if (AssetViewPtr.IsValid() && !InArgs._AssetPickerConfig.bAutohideSearchBar)
+	if (TextFilter.IsValid())
 	{
-		TextFilter = MakeShareable(new FFrontendFilter_Text());
 		bool bClassNamesProvided = (InArgs._AssetPickerConfig.Filter.ClassPaths.Num() != 1);
 		TextFilter->SetIncludeClassName(bClassNamesProvided || AssetViewPtr->IsIncludingClassNames());
 		TextFilter->SetIncludeAssetPath(AssetViewPtr->IsIncludingAssetPaths());
@@ -487,12 +493,10 @@ void SAssetPicker::SetSearchBoxText(const FText& InSearchText)
 		TextFilter->SetRawFilterText(InSearchText);
 		if (InSearchText.IsEmpty())
 		{
-			FrontendFilters->Remove(TextFilter);
 			AssetViewPtr->SetUserSearching(false);
 		}
 		else
 		{
-			FrontendFilters->Add(TextFilter);
 			AssetViewPtr->SetUserSearching(true);
 		}
 	}
diff --git a/Engine/Source/Editor/ContentBrowser/Private/SAssetPicker.h b/Engine/Source/Editor/ContentBrowser/Private/SAssetPicker.h
--- a/Engine/Source/Editor/ContentBrowser/Private/SAssetPicker.h
+++ b/Engine/Source/Editor/ContentBrowser/Private/SAssetPicker.h
@@ -20,7 +20,6 @@
 #include "Widgets/SCompoundWidget.h"
 
 class FFilter_ShowOtherDevelopers;
-class FFrontendFilter_Text;
 class FUICommandList;
 class SAssetSearchBox;
 class SAssetView;
@@ -190,7 +189,7 @@ private:
 
 	/** Filters needed for filtering the assets */
 	TSharedPtr< FAssetFilterCollectionType > FilterCollection;
-	TSharedPtr< FFrontendFilter_Text > TextFilter;
+	TSharedPtr<class FAssetTextFilter> TextFilter;
 	TSharedPtr<FFilter_ShowOtherDevelopers> OtherDevelopersFilter;
 
 	EAssetTypeCategories::Type DefaultFilterMenuExpansion;
diff --git a/Engine/Source/Editor/ContentBrowser/Private/SAssetView.cpp b/Engine/Source/Editor/ContentBrowser/Private/SAssetView.cpp
--- a/Engine/Source/Editor/ContentBrowser/Private/SAssetView.cpp
+++ b/Engine/Source/Editor/ContentBrowser/Private/SAssetView.cpp
@@ -3,13 +3,19 @@
 #include "SAssetView.h"
 
 #include "Algo/AnyOf.h"
+#include "Algo/Compare.h"
+#include "Algo/RemoveIf.h"
 #include "Algo/Transform.h"
 #include "AssetRegistry/AssetRegistryState.h"
 #include "AssetSelection.h"
+#include "AssetTextFilter.h"
 #include "AssetToolsModule.h"
 #include "AssetView/AssetViewConfig.h"
 #include "AssetViewTypes.h"
 #include "AssetViewWidgets.h"
+#include "Async/ParallelFor.h"
+#include "Async/WordMutex.h"
+#include "Async/UniqueLock.h"
 #include "ContentBrowserConfig.h"
 #include "ContentBrowserDataDragDropOp.h"
 #include "ContentBrowserDataLegacyBridge.h"
@@ -77,13 +83,115 @@
 	#define checkAssetList(cond)
 #endif
 
-namespace
+namespace UE::AssetView
 {
 	/** Time delay between recently added items being added to the filtered asset items list */
 	const double TimeBetweenAddingNewAssets = 4.0;
 
 	/** Time delay between performing the last jump, and the jump term being reset */
 	const double JumpDelaySeconds = 2.0;
+
+	bool AllowAsync = true;
+	FAutoConsoleVariableRef CVarAllowAsync(
+		TEXT("AssetView.AllowAsync"),
+		AllowAsync,
+		TEXT("Whether to allow the asset view to perform work with async tasks (rather than time-sliced)"),
+		ECVF_Default
+	);
+	 
+	bool AllowParallelism = true;
+	FAutoConsoleVariableRef CVarAllowParallelism(
+		TEXT("AssetView.AllowParallelism"),
+		AllowParallelism,
+		TEXT("Whether to allow the asset view to perform work in parallel (e.g. ParallelFor)"),
+		ECVF_Default
+	);
+
+	// Return the max size of the batch of items to text filter per task - do fewer if parallelism is disabled
+	int32 GetMaxTextFilterItemBatch()
+	{
+		static const int32 NumWorkers = LowLevelTasks::FScheduler::Get().GetNumWorkers();
+		return AllowParallelism ? NumWorkers * 1024 : 1024;
+	}
+
+	bool AreBackendFiltersDifferent(const FARFilter& A, const FARFilter& B)
+	{
+		if (A.PackageNames.Num() != B.PackageNames.Num()
+		|| A.PackagePaths.Num() != B.PackageNames.Num() 
+		|| A.SoftObjectPaths.Num() != B.SoftObjectPaths.Num()
+		|| A.ClassPaths.Num() != B.ClassPaths.Num()
+		|| A.TagsAndValues.Num() != B.TagsAndValues.Num()
+		|| A.RecursiveClassPathsExclusionSet.Num() != B.RecursiveClassPathsExclusionSet.Num()
+		|| A.bRecursivePaths != B.bRecursivePaths
+		|| A.bRecursiveClasses != B.bRecursiveClasses
+		|| A.bIncludeOnlyOnDiskAssets != B.bIncludeOnlyOnDiskAssets
+		|| A.WithoutPackageFlags != B.WithoutPackageFlags 
+		|| A.WithPackageFlags != B.WithPackageFlags)
+		{
+			return true;
+		}
+
+		// Expect things to be generated in the same order by the filter bar, so just check linear matching
+		if (!Algo::Compare(A.PackageNames, B.PackageNames)
+		|| !Algo::Compare(A.PackagePaths, B.PackagePaths)
+		|| !Algo::Compare(A.SoftObjectPaths, B.SoftObjectPaths)
+		|| !Algo::Compare(A.ClassPaths, B.ClassPaths))
+		{
+			return true;
+		}
+
+		for (const FTopLevelAssetPath& Path : A.RecursiveClassPathsExclusionSet)
+		{
+			if (!B.RecursiveClassPathsExclusionSet.Contains(Path))
+			{
+				return true;
+			}
+		}
+
+		for (const FTopLevelAssetPath& Path : B.RecursiveClassPathsExclusionSet)
+		{
+			if (!A.RecursiveClassPathsExclusionSet.Contains(Path))
+			{
+				return true;
+			}
+		}
+
+		TArray<FName> AKeys;
+		A.TagsAndValues.GetKeys(AKeys);
+		for (FName Key : AKeys)
+		{
+			if (!B.TagsAndValues.Contains(Key))
+			{
+				return true;
+			}
+			TArray<TOptional<FString>> AValues;
+			A.TagsAndValues.MultiFind(Key, AValues);
+			Algo::SortBy(AValues, [](const TOptional<FString>& S) { return S.Get(FString()); });
+			TArray<TOptional<FString>> BValues;
+			B.TagsAndValues.MultiFind(Key, BValues);
+			Algo::SortBy(BValues, [](const TOptional<FString>& S) { return S.Get(FString()); });
+
+			if (!Algo::Compare(AValues, BValues))
+			{
+				return true;
+			}
+		}
+
+		return false;
+	}
+
+	bool AreCustomPermissionListsDifferent(TArray<TSharedRef<const FPathPermissionList>>* InCustomPermissionLists,
+		TArray<TSharedRef<const FPathPermissionList>>& ExistingPermissionLists)
+	{
+		if (InCustomPermissionLists == nullptr)
+		{
+			return ExistingPermissionLists.IsEmpty();
+		}
+
+		// Expect order to be built in the same way so if order is different, trigger a rebuild
+		// Also expect that if filters change their permission lists, they create a new object. 
+		return !Algo::Compare(*InCustomPermissionLists, ExistingPermissionLists);
+	}
 }
 
 
@@ -122,6 +230,11 @@ public:
 		}
 	}
 
+	bool NeedsQueryFilter()
+	{
+		return AssetView->OnShouldFilterItem.IsBound() || AssetView->OnShouldFilterAsset.IsBound();
+	}
+
 	bool DoesItemPassQueryFilter(const TSharedPtr<FAssetViewItem>& InItemToFilter)
 	{
 		// Folders aren't subject to additional filtering
@@ -142,40 +255,771 @@ public:
 			FAssetData ItemAssetData;
 			if (!InItemToFilter->GetItem().Legacy_TryGetAssetData(ItemAssetData) || AssetView->OnShouldFilterAsset.Execute(ItemAssetData))
 			{
-				return false;
+				return false;
+			}
+		}
+
+		return true;
+	}
+
+	bool DoesItemPassFrontendFilter(const TSharedPtr<FAssetViewItem>& InItemToFilter)
+	{
+		// Folders are only subject to "empty" filtering
+		if (InItemToFilter->IsFolder())
+		{
+			if (!ContentBrowserData->IsFolderVisible(InItemToFilter->GetItem().GetVirtualPath(), ContentBrowserUtils::GetIsFolderVisibleFlags(bDisplayEmptyFolders), FolderFilter))
+			{
+				return false;
+			}
+			return true;
+		}
+
+		// Run the item through the filters
+		if (AssetView->IsFrontendFilterActive() && !AssetView->PassesCurrentFrontendFilter(InItemToFilter->GetItem()))
+		{
+			return false;
+		}
+
+		return true;
+	}
+
+private:
+	SAssetView* AssetView = nullptr;
+	UContentBrowserDataSubsystem* ContentBrowserData = nullptr;
+	TOptional<FContentBrowserFolderContentsFilter> FolderFilter;
+	const bool bDisplayEmptyFolders = true;
+};
+
+
+struct FAssetViewItemFilterState
+{
+	uint8 Removed : 1;
+	uint8 PassedFrontendFilter : 1;
+	uint8 PassedTextFilter : 1;	
+
+	// This item passed filtering and was published to the view
+	uint8 Published : 1;
+	// Priority filtering was performed because of data updates, do not overwrite results with async filtering results
+	uint8 PriorityFiltered : 1;
+};
+
+/**
+ * Manages items returned from backend query and incrementally/asynchronously filtering them.
+ * - Recycling of old objects on new query
+ */
+class FAssetViewItemCollection
+{
+public: 
+	FAssetViewItemCollection()
+	{
+	}
+
+	/** Returns the number of items which were fetched and have not been removed. */
+	int32 Num() const { return NumValidItems; }
+
+	/** Returns true if there is any incomplete filtering work. */
+	bool HasItemsPendingFilter()
+	{
+		// No need to check text filtering progress/task here as PublishProgress cannot surpass text filtering
+		return ItemsPendingPriorityFilter.Num() != 0 || PublishProgress < Items.Num() || ItemsPendingPriorityPublish.Num();
+	}
+
+	/** Return the amount of progress made in filtering for presentation to the user; a number between 0 and Num() */
+	int32 GetFilterProgress() const
+	{
+		return PublishProgress;
+	}
+
+	/** 
+	 * Fetch all items from the given paths (sources) matching the given filter, recycling old FAssetViewItem objects.
+	 * 
+	 * @param bAllowItemRecycling Whether or not to allow reuse of items and therefore widgets. 
+	 * 	Setting to false can avoid lots of time firing modification delegates for recursive searches.
+	 */
+	void RefreshItemsFromBackend(const FSourcesData& SourcesData, const FContentBrowserDataFilter& DataFilter, bool bAllowItemRecycling);
+	
+	/** 
+	 * Find an FAssetViewItem containing the given content browser data if one exists.
+	 * Returned item should not be modified as background text processing may be operating on it.
+	 */
+	TSharedPtr<FAssetViewItem> FindItemForRename(const FContentBrowserItem& InItem);
+
+	/** 
+	 * Create the given item from user interaction (e.g. create asset, rename asset).
+	 * Makes the item visible immediately. 
+	 */
+	TSharedPtr<FAssetViewItem> CreateItemFromUser(FContentBrowserItem&& InItem, TArray<TSharedPtr<FAssetViewItem>>& FilteredAssetItems);
+	
+	/** 
+	 * Find an existing item or create one from an incremental data update.
+	 * If an item exists, the data in it is replaced and a callback fired to be handled by widgets bound to it.
+	 * Safe to call during threaded text filtering.
+	 */
+	TSharedPtr<FAssetViewItem> UpdateData(FContentBrowserItemData&& InData);
+
+	/**
+	 * Remove the given item data from the FAssetViewItem that contains it and if that item no longer contains any data,
+	 * remove the FAssetViewItem and return it.
+	 * Safe to call during threaded text filtering.
+	 */
+	TSharedPtr<FAssetViewItem> RemoveItemData(const FContentBrowserItemData& InItemData);
+	TSharedPtr<FAssetViewItem> RemoveItemData(const FContentBrowserMinimalItemData& InItemData);
+
+	/**
+	 * Remove the given item that was being created/renamed.
+	 * Safe to call during threaded text filtering.
+	 */
+	void RemoveItem(const TSharedPtr<FAssetViewItem>& ToRemove);
+
+	/** 
+	 * Cancel any in progress async text filtering operation and wait for tasks to shut down.
+	 */
+	void AbortTextFiltering();	
+	
+	/** 
+	 * Clear the filtering results of all known non-removed items to be run again with new filters.
+	 * Cancels async text filtering and waits for cancellation to complete safely.
+	 */
+	void ResetFilterState();
+	
+	/** 
+	 * Start filtering all items against the given text filter in the background.
+	 * Results will be fetched and merged during UpdateItemFiltering.
+	 */
+	void StartTextFiltering(TSharedPtr<FAssetTextFilter> TextFilter);
+
+	/** 
+	 * Run main-thread filtering on items until the specified end time.
+	 * 
+	 * @param InTextFilter Text filter if any for launching new async text filtering tasks if necessary
+	 * @param OutItems Array to populate with new items that passed the filter if any.
+	 */
+	void UpdateItemFiltering(FAssetViewFrontendFilterHelper& InHelper, double InEndTime, TArray<TSharedPtr<FAssetViewItem>>& OutItems);
+
+	/** 
+	 * Perform filtering on any items which already existed and had been filtered when a data update was received.
+	 * Returns true if any items changed filter result.
+	 * 
+	 * @param FilteredAssetItems Items which have already been determined to be visible, to be updated. 
+	 */
+	bool PerformPriorityFiltering(FAssetViewFrontendFilterHelper& Helper,
+		TArray<TSharedPtr<FAssetViewItem>>& FilteredAssetItems);
+
+private:
+	int32 CreateItem_Locked(FContentBrowserItemData&& InItem)
+	{
+		return CreateItem_Locked(FContentBrowserItem(MoveTemp(InItem)));
+	}
+
+	int32 CreateItem_Locked(FContentBrowserItem&& InItem)
+	{
+		uint32 Hash = HashItem(InItem);
+		TSharedPtr<FAssetViewItem> Item = Items.Add_GetRef(MakeShared<FAssetViewItem>(MoveTemp(InItem)));
+		FilterState.AddZeroed(1);
+		++NumValidItems;
+		if (!RefreshLookup())
+		{
+			// Resize lookup's index list to match capacity of Items rather than its own growth strategy 
+			if (Lookup.GetIndexSize() < (uint32)Items.Num())
+			{
+				Lookup.Resize(Items.Max());	
+			}
+			Lookup.Add(Hash, Items.Num() - 1);
+		}
+		return Items.Num() - 1;
+	}
+
+	uint32 HashItem(const FContentBrowserItem& Item)
+	{
+		check(Item.IsValid());
+		return GetTypeHash(Item.GetVirtualPath());
+	}
+	uint32 HashItem(const FContentBrowserItemData& Item)
+	{
+		check(Item.IsValid());
+		return GetTypeHash(Item.GetVirtualPath());
+	}
+	uint32 HashItem(const FContentBrowserMinimalItemData& Item)
+	{
+		check(!Item.GetVirtualPath().IsNone());
+		return GetTypeHash(Item.GetVirtualPath());
+	}
+	 
+	struct FTextFilterResult
+	{
+		int32 StartIndex;
+		TBitArray<> Results;
+		UE::Tasks::TTask<FTextFilterResult> Next;
+	};
+	FTextFilterResult AsyncFilterText(int32 StartIndex, int32 MaxItems) const;
+
+	inline bool ItemPassedAllFilters(int32 Index) const
+	{
+		return !FilterState[Index].Removed && FilterState[Index].PassedFrontendFilter && (bAllItemsPassedTextFilter || FilterState[Index].PassedTextFilter);
+	}
+	
+	inline TSharedPtr<FAssetViewItem> MarkItemRemoved(int32 Index)
+	{
+		check(Items[Index].IsValid() && !FilterState[Index].Removed);
+		--NumValidItems;
+		bItemsPendingRemove = true;
+		FilterState[Index].Removed = true;
+		// Do not null out the item because we want to be able to remove it from published items in PerformPriorityFiltering
+		return Items[Index];
+	}
+
+	/** 
+	 * If the number of stored items has grown beyond the bounds of Lookup, rebuild it with larger hash 
+	 * @return true if the lookup was rebuilt 
+	 */
+	bool RefreshLookup();
+
+private:
+	// Lock for access to the size and contents of Items - e.g. when text filtering is operating on a batch of items
+	// The size of Items may need to change or the ItemData object within items may need modification as data scanning progresses
+	mutable FRWLock Lock;
+
+	// Hash of items by virtual path. There may be multiple items with the same virtual path, deduplication is done manually during 
+	// population. Objects with the same path may exist unless they are folders, in which case they are merged.
+	// After population, objects are looked up by path & source for updates.
+	// This lookup is only used on the main thread so it can be safely rebuilt during async text filtering.
+	FHashTable Lookup;
+	// Linear list of items indexed by Lookup. May contain null entries. 
+	TArray<TSharedPtr<FAssetViewItem>> Items;
+	// State of non-text filtering matching items in Items. Only modified on main thread. 
+	TArray<FAssetViewItemFilterState> FilterState;
+
+	// How many items in Items are not null. Atomically decreased during batch merge of folder items.
+	// Also decreased when items are removed by data update notifications.
+	std::atomic<int32> NumValidItems{0};
+	// How many elements of Items have been tested against frontend filtering if required.
+	int32 FrontendFilterProgress = 0;
+	// How many elements of Items have gone through text filtering and had their results merged on the main thread.
+	int32 TextFilterProgress = 0;
+	// How many elements of Items have been published to the view - smaller of FilterProgress and TextFilterProgress on last update 
+	// Items with indices below this may have been added to the list/tile/column view so updates to those items require re-filtering
+	int32 PublishProgress = 0;
+
+	// Cached compiled text filter for the current filtering pass
+	TSharedPtr<FCompiledAssetTextFilter> CompiledTextFilter;
+
+	// Handle to ongoing task filtering Items by a text query. 
+	// Modifications during filtering are protected by the Lock in this class and otherwise yet-to-be-filtered items should not be
+	// provided to external code or modified.
+	UE::Tasks::TTask<FTextFilterResult> TextFilterTask;
+	// Flag to signal cancellation of text filtering task 
+	std::atomic<bool> ShouldCancelTextFiltering{false}; 
+
+	// Items which have been updated while visible, so should be re-filtered immediately
+	TSet<int32> ItemsPendingPriorityFilter;
+	// Items which were updated and passed filtering when they previously failed, so need to be published again
+	TSet<int32> ItemsPendingPriorityPublish;
+
+	// If true all items passed text filtering and TextFilterState may be emptied - e.g. if no text filter was applied at all
+	bool bAllItemsPassedTextFilter = false;
+
+	// Some items have been marked for removal but their pointers have not been cleared yet because we need to compare against them
+	std::atomic<bool> bItemsPendingRemove = false;
+};
+
+bool FAssetViewItemCollection::RefreshLookup()
+{
+	uint32 HashSize = TSetAllocator<>::GetNumberOfHashBuckets(Items.Num());
+	if (Lookup.GetHashSize() < HashSize)
+	{
+		Lookup.Clear(HashSize, Items.Max());
+		ParallelFor(TEXT("FAssetViewItemCollection::RefreshLookup"), Items.Num(), 16 * 1024, [this](int32 ItemIndex){
+			const TSharedPtr<FAssetViewItem> Item = Items[ItemIndex];
+			if (!Item.IsValid())
+			{
+				return;
+			}
+
+			uint32 Hash = HashItem(Item->GetItem());
+			Lookup.Add_Concurrent(Hash, ItemIndex);
+		}, UE::AssetView::AllowParallelism ? EParallelForFlags::None : EParallelForFlags::ForceSingleThread);
+		return true;
+	}
+	return false;
+}
+
+TSharedPtr<FAssetViewItem> FAssetViewItemCollection::FindItemForRename(const FContentBrowserItem& InItem)
+{
+	const uint32 Hash = HashItem(InItem);
+	FContentBrowserItemKey ItemKey(InItem);
+	for (uint32 It = Lookup.First(Hash); Lookup.IsValid(It); It = Lookup.Next(It))
+	{
+		if (Items[It].IsValid() && ItemKey == FContentBrowserItemKey(Items[It]->GetItem()))
+		{
+			checkf(FilterState[It].Published, 
+				TEXT("Only items which have been made visible in the UI should be available for renaming to maintain thread safety with async text filtering."));
+			return Items[It];
+		}
+	}
+	return {};
+}
+
+TSharedPtr<FAssetViewItem> FAssetViewItemCollection::CreateItemFromUser(FContentBrowserItem&& InItem, TArray<TSharedPtr<FAssetViewItem>>& FilteredAssetItems)
+{
+	int32 Index;
+	{
+		FWriteScopeLock Guard(Lock);
+		Index = CreateItem_Locked(MoveTemp(InItem));
+	}
+	// Make this item visible immediately, forcing it to be so regardless of current filter set until filtering is refreshed.
+	FilterState[Index].PassedFrontendFilter = true;
+	FilterState[Index].PassedTextFilter = true;
+	FilterState[Index].Published = true;
+	FilterState[Index].PriorityFiltered = true;
+	FilteredAssetItems.Add(Items[Index]);
+	return Items[Index];
+}
+
+TSharedPtr<FAssetViewItem> FAssetViewItemCollection::UpdateData(FContentBrowserItemData&& InData)
+{
+	const uint32 Hash = HashItem(InData);
+	FContentBrowserItemKey ItemKey(InData);
+	int32 ExistingItemIndex = INDEX_NONE;
+	for (uint32 It = Lookup.First(Hash); Lookup.IsValid(It); It = Lookup.Next(It))
+	{
+		if (Items[It].IsValid() && ItemKey == FContentBrowserItemKey(Items[It]->GetItem()))
+		{
+			ExistingItemIndex = It;
+			break;
+		}
+	}
+
+	if (ExistingItemIndex != INDEX_NONE)
+	{
+		FWriteScopeLock Guard(Lock);
+		// Update the item and mark it for re-filtering if it has already been filtered
+		Items[ExistingItemIndex]->AppendItemData(MoveTemp(InData));
+		Items[ExistingItemIndex]->BroadcastItemDataChanged();
+		check(!FilterState[ExistingItemIndex].Removed);
+	}
+	else
+	{
+		FWriteScopeLock Guard(Lock);
+		ExistingItemIndex = CreateItem_Locked(MoveTemp(InData));
+	}
+
+	if (ExistingItemIndex < FrontendFilterProgress) 
+	{
+		// This item was already filtered so we may want to remove it from the view or add it 
+		// We only check FrontendFilterProgress here as we only publish items that have passed both filters 
+		ItemsPendingPriorityFilter.Add(ExistingItemIndex);
+	}
+	return Items[ExistingItemIndex];
+}
+
+TSharedPtr<FAssetViewItem> FAssetViewItemCollection::RemoveItemData(const FContentBrowserItemData& InItemData)
+{
+	return RemoveItemData(FContentBrowserMinimalItemData(InItemData));
+}
+
+TSharedPtr<FAssetViewItem> FAssetViewItemCollection::RemoveItemData(const FContentBrowserMinimalItemData& InItemData)
+{
+	const uint32 Hash = HashItem(InItemData);
+	FContentBrowserItemKey ItemKey(InItemData.GetItemType(), InItemData.GetVirtualPath(), InItemData.GetDataSource());
+	for (uint32 It = Lookup.First(Hash); Lookup.IsValid(It); It = Lookup.Next(It))
+	{
+		if (Items[It].IsValid() && ItemKey == FContentBrowserItemKey(Items[It]->GetItem()))
+		{
+			TSharedRef<FAssetViewItem> ItemToRemove = Items[It].ToSharedRef();
+
+			{
+				// We only need to lock around the modification of the data stored in ItemToRemove because the background text search may be reading it.
+				FWriteScopeLock Guard(Lock);
+				ItemToRemove->RemoveItemData(InItemData);
+			}
+
+			// Only fully remove this item if every sub-item is removed (items become invalid when empty)
+			if (ItemToRemove->GetItem().IsValid())
+			{
+				return {};
+			}
+
+			// This item was already filtered so we may want to remove it from the view.
+			if (It < (uint32)FrontendFilterProgress) 
+			{
+				ItemsPendingPriorityFilter.Add(It);
+			}
+			Lookup.Remove(Hash, It);
+			return MarkItemRemoved(It);
+		}
+	}
+	return {};
+}
+ 
+
+void FAssetViewItemCollection::RemoveItem(const TSharedPtr<FAssetViewItem>& ToRemove)
+{
+	// There is no need to lock here because we don't modify the item which the background text search may be reading.
+	const uint32 Hash = HashItem(ToRemove->GetItem());
+	for (uint32 It = Lookup.First(Hash); Lookup.IsValid(It); It = Lookup.Next(It))
+	{
+		if (Items[It] == ToRemove)
+		{
+		 	Lookup.Remove(Hash, It);
+			// This item was already filtered so we may want to remove it from the view.
+			if (It < (uint32)FrontendFilterProgress)
+			{
+				ItemsPendingPriorityFilter.Add(It);
+			}
+			MarkItemRemoved(It);	
+			return;
+		}
+	}
+}
+
+void FAssetViewItemCollection::ResetFilterState()
+{
+	check(!TextFilterTask.IsValid());
+
+	if (bItemsPendingRemove)
+	{
+		for (int32 i=0; i < Items.Num(); ++i)
+		{
+			if (FilterState[i].Removed)
+			{
+				Items[i].Reset();
+			}
+		}
+		bItemsPendingRemove = false;
+	}
+
+	ShouldCancelTextFiltering = true;
+
+	ItemsPendingPriorityPublish.Reset();
+	FilterState.Reset();
+	FilterState.AddZeroed(Items.Num());
+	FrontendFilterProgress = 0;
+	PublishProgress = 0;
+	TextFilterProgress = 0;
+
+	// Recreate the Removed flag if necessary after FilterState was wiped so that we know which items are expected to be null
+	if (Items.Num() != NumValidItems)
+	{
+		for (int32 i=0; i < Items.Num(); ++i)
+		{
+			if (!Items[i].IsValid())
+			{
+				FilterState[i].Removed = true;
+			}
+		}
+	}		
+}
+
+void FAssetViewItemCollection::AbortTextFiltering()
+{
+	if (TextFilterTask.IsValid())
+	{
+		ShouldCancelTextFiltering = true;
+		TextFilterTask.BusyWait();
+		TextFilterTask = {};
+	}
+}
+
+void FAssetViewItemCollection::StartTextFiltering(TSharedPtr<FAssetTextFilter> TextFilter)
+{
+	// Text filter task reads CompiledTextFilter so must not be running here
+	check(!TextFilterTask.IsValid());
+
+	ShouldCancelTextFiltering = false;
+	if (!TextFilter.IsValid() || TextFilter->IsEmpty())
+	{
+		CompiledTextFilter.Reset();
+		bAllItemsPassedTextFilter = true;
+		return;
+	}
+
+	CompiledTextFilter = TextFilter->Compile();
+	bAllItemsPassedTextFilter = false;
+
+	if (UE::AssetView::AllowAsync)
+	{
+		const int32 MaxItemsPerTask = UE::AssetView::GetMaxTextFilterItemBatch();
+		TextFilterTask = UE::Tasks::Launch(UE_SOURCE_LOCATION, [this, MaxItemsPerTask]() {
+			return AsyncFilterText(0, MaxItemsPerTask);
+		});
+	}
+}
+
+FAssetViewItemCollection::FTextFilterResult FAssetViewItemCollection::AsyncFilterText(int32 StartIndex, int32 InMaxItems) const
+{
+	if (ShouldCancelTextFiltering)
+	{
+		return FTextFilterResult{ StartIndex, {}, {} };
+	}
+
+	FReadScopeLock Guard(Lock);
+
+	// How many items to filter in between checking for interruption and allowing other threads to acquire the lock
+	int32 NumItemsToFilter = FMath::Min(InMaxItems, Items.Num() - StartIndex);
+
+	TBitArray<> MergedResult;
+	MergedResult.Add(false, NumItemsToFilter);
+
+	constexpr int32 MinThreadWorkSize = 1024;
+	TArray<FCompiledAssetTextFilter> Contexts;
+	auto CreateContext = [this](int32 ContextIndex, int32 NumContexts) {
+		return CompiledTextFilter->CloneForThreading();
+	};
+	auto DoWork = [&MergedResult, StartIndex, this](FCompiledAssetTextFilter& Filter, int32 TaskIndex) {
+		const TSharedPtr<FAssetViewItem>& Item = Items[StartIndex + TaskIndex];
+		bool bPasses = Item.IsValid() && Filter.PassesFilter(Item->GetItem());
+		if (bPasses)
+		{
+			MergedResult[TaskIndex].AtomicSet(true);
+		}
+	};
+	ParallelForWithTaskContext(TEXT("AssetViewTextFiltering"), Contexts, NumItemsToFilter, MinThreadWorkSize, CreateContext,
+		DoWork, UE::AssetView::AllowParallelism ? EParallelForFlags::None : EParallelForFlags::ForceSingleThread);
+
+	int32 NextStartIndex = StartIndex + NumItemsToFilter;
+	// Need to check termination condition while holding the lock
+	bool bContinue = NextStartIndex < Items.Num() && UE::AssetView::AllowAsync; 
+	UE::Tasks::TTask<FTextFilterResult> NextTask;
+	if (bContinue)
+	{
+		return FTextFilterResult{ StartIndex, MoveTemp(MergedResult),
+			UE::Tasks::Launch(
+				UE_SOURCE_LOCATION, [this, NextStartIndex, InMaxItems]() mutable { return AsyncFilterText(NextStartIndex, InMaxItems); }) };
+	}
+	else
+	{
+		return FTextFilterResult{ StartIndex, MoveTemp(MergedResult), {} };
+	}
+}
+
+void FAssetViewItemCollection::UpdateItemFiltering(
+	FAssetViewFrontendFilterHelper& InHelper, double InEndTime, TArray<TSharedPtr<FAssetViewItem>>& OutItems)
+{
+	if (ItemsPendingPriorityFilter.Num())
+	{
+		PerformPriorityFiltering(InHelper, OutItems);
+	}
+
+	const bool bNeedsQueryFilter = InHelper.NeedsQueryFilter();
+	do 
+	{
+		constexpr int FilterBatchSize = 128;
+		const int32 End = FMath::Min(FrontendFilterProgress + FilterBatchSize, Items.Num());
+
+		// Query filter
+		if (bNeedsQueryFilter)
+		{
+			for (int32 Index = FrontendFilterProgress; Index < End; ++Index)
+			{
+				if (FilterState[Index].Removed || FilterState[Index].PriorityFiltered)
+				{
+					continue;
+				}
+				if (!InHelper.DoesItemPassQueryFilter(Items[Index]))
+				{
+					// Failing this filter is equivalent to not being returned from the backend
+					MarkItemRemoved(Index);
+				}
+			}
+		}
+		for (int32 Index = FrontendFilterProgress; Index < End; ++Index)
+		{
+			if (FilterState[Index].Removed || FilterState[Index].PriorityFiltered)
+			{
+				continue;
+			}
+
+			if (InHelper.DoesItemPassFrontendFilter(Items[Index]))
+			{
+				FilterState[Index].PassedFrontendFilter = true;
+			}
+		}
+
+		if (!bAllItemsPassedTextFilter)
+		{
+			check(CompiledTextFilter.IsValid());
+			// Result must be moved in so TextFilterTask can be overwritten 
+			auto MergeTextFilterResult = [this](FTextFilterResult Result) {
+				TextFilterTask = MoveTemp(Result.Next);
+				TextFilterProgress = Result.StartIndex + Result.Results.Num();
+			
+				for (TConstSetBitIterator<> It(Result.Results); It; ++It)
+				{
+					int32 ItemIndex = Result.StartIndex + It.GetIndex();
+					if (!FilterState[ItemIndex].PriorityFiltered)
+					{
+						FilterState[ItemIndex].PassedTextFilter = true;
+					}
+				}
+			};
+
+			if (TextFilterTask.IsValid() && TextFilterTask.IsCompleted())
+			{
+				MergeTextFilterResult(MoveTemp(TextFilterTask.GetResult()));
+			}
+
+			if (!TextFilterTask.IsValid() && TextFilterProgress < Items.Num() && UE::AssetView::AllowAsync)
+			{
+				// New elements were added after the text filter attempted to launch a continuation task
+				TextFilterTask = UE::Tasks::Launch(UE_SOURCE_LOCATION,
+					[this, StartIndex = TextFilterProgress]() {
+						return AsyncFilterText(StartIndex, UE::AssetView::GetMaxTextFilterItemBatch());
+					});
+			}
+			
+			// In case flag was flipped while filtering was running, wait til task ends before performing text filtering on the game thread.
+			if (!UE::AssetView::AllowAsync && !TextFilterTask.IsValid() && TextFilterProgress < Items.Num())
+			{
+				FTextFilterResult Result = AsyncFilterText(TextFilterProgress, UE::AssetView::GetMaxTextFilterItemBatch());
+				MergeTextFilterResult(Result);
+			}
+		}
+
+		FrontendFilterProgress = End;
+	} while(FPlatformTime::Seconds() < InEndTime);
+
+	// Append items which have passed both text and frontend filtering to OutItems
+	int32 CanPublish = FMath::Min(FrontendFilterProgress, bAllItemsPassedTextFilter ? Items.Num() : TextFilterProgress);
+	for (int32 i = PublishProgress; i < CanPublish; ++i)
+	{
+		if (!FilterState[i].PriorityFiltered)
+		{
+			const bool bPublish = ItemPassedAllFilters(i);
+			FilterState[i].Published = bPublish;
+			if (bPublish)
+			{
+				OutItems.Add(Items[i]);
+			}
+		}
+	}
+
+	// Publish items which originally failed filtering and then were updated to pass it 
+	for (int32 Index : ItemsPendingPriorityPublish)
+	{
+		// Check we didn't update the state again to failure or removal
+		const bool bPublish = ItemPassedAllFilters(Index);
+		check(!FilterState[Index].Published);
+		FilterState[Index].Published = bPublish;
+		if (bPublish)
+		{
+			OutItems.Add(Items[Index]);
+		}
+	}
+	ItemsPendingPriorityPublish.Reset();
+
+	PublishProgress = CanPublish;
+}
+
+bool FAssetViewItemCollection::PerformPriorityFiltering(FAssetViewFrontendFilterHelper& Helper,
+	TArray<TSharedPtr<FAssetViewItem>>& FilteredAssetItems)
+{
+	int32 PrevNum = FilteredAssetItems.Num();
+	if (ItemsPendingPriorityFilter.Num())
+	{
+		const bool bRunQueryFilter = Helper.NeedsQueryFilter();
+		if (bRunQueryFilter)
+		{
+			for (int32 Index : ItemsPendingPriorityFilter)
+			{
+				if (FilterState[Index].Removed)
+				{
+					continue;
+				}
+				if (!Helper.DoesItemPassQueryFilter(Items[Index]))
+				{
+					// Failing this filter is equivalent to not being returned from the backend
+					MarkItemRemoved(Index);
+				}
+			}
+		}
+		for (int32 Index : ItemsPendingPriorityFilter)
+		{
+			if (FilterState[Index].Removed)
+			{
+				continue;
+			}
+
+			const bool bPassedFrontendFilter = Helper.DoesItemPassFrontendFilter(Items[Index]);
+			if (bPassedFrontendFilter != FilterState[Index].PassedFrontendFilter)
+			{
+				// Only items which had previously passed filtering should be given for priority filtering, other items should
+				// be added at the end of the collection and processed as normal
+				check(FilterState[Index].PassedFrontendFilter);
+				FilterState[Index].PassedFrontendFilter = bPassedFrontendFilter;
+			}
+		}
+
+		if (!bAllItemsPassedTextFilter)
+		{
+			check(CompiledTextFilter.IsValid());
+			// TODO: Is it possible we get a very large update from the backend for items which have already been filtered? So should we launch tasks to do this? 
+			for (int32 Index : ItemsPendingPriorityFilter)
+			{
+				if (FilterState[Index].Removed)
+				{
+					continue;
+				}
+
+				bool bPassed = CompiledTextFilter->PassesFilter(Items[Index]->GetItem());
+				FilterState[Index].PassedTextFilter = bPassed;
+			}
+		}
+		
+		TSet<TSharedPtr<FAssetViewItem>> ToRemove;
+		for (int32 Index : ItemsPendingPriorityFilter)
+		{
+			FilterState[Index].PriorityFiltered = true;
+			if (Index >= PublishProgress)
+			{
+				// If item has yet to be published in the normal order, just leave the filter results for UpdateItemFiltering
+				continue;
 			}
-		}
 
-		return true;
-	}
+			const bool bPublish = !FilterState[Index].Removed && FilterState[Index].PassedFrontendFilter
+				&& (bAllItemsPassedTextFilter || FilterState[Index].PassedTextFilter);
 
-	bool DoesItemPassFrontendFilter(const TSharedPtr<FAssetViewItem>& InItemToFilter)
-	{
-		// Folders are only subject to "empty" filtering
-		if (InItemToFilter->IsFolder())
-		{
-			if (!ContentBrowserData->IsFolderVisible(InItemToFilter->GetItem().GetVirtualPath(), ContentBrowserUtils::GetIsFolderVisibleFlags(bDisplayEmptyFolders), FolderFilter))
+			if (FilterState[Index].Published && !bPublish)
 			{
-				return false;
+				// Remove item while maintaining sorting of remaining items
+				ToRemove.Add(Items[Index]);
+			}
+			else if (!FilterState[Index].Published && bPublish)
+			{
+				// Newly passing item - this means the view will have to be re-sorted 
+				// Defer addition until UpdateItemFiltering
+				ItemsPendingPriorityPublish.Add(Index);
 			}
-			return true;
 		}
 
-		// Run the item through the filters
-		if (AssetView->IsFrontendFilterActive() && !AssetView->PassesCurrentFrontendFilter(InItemToFilter->GetItem()))
+		FilteredAssetItems.SetNum(Algo::StableRemoveIf(
+			FilteredAssetItems, [&ToRemove](const TSharedPtr<FAssetViewItem> Item) { return ToRemove.Contains(Item); }));
+
+		ItemsPendingPriorityFilter.Reset();
+	}
+	
+	if (bItemsPendingRemove)
+	{
+		FWriteScopeLock Guard(Lock);
+		for (int32 i=0; i < Items.Num(); ++i)
 		{
-			return false;
+			if (FilterState[i].Removed && Items[i].IsValid())
+			{
+				Items[i].Reset();
+			}
 		}
-
-		return true;
+		bItemsPendingRemove = false;
 	}
+	return PrevNum != FilteredAssetItems.Num();
+}
 
-private:
-	SAssetView* AssetView = nullptr;
-	UContentBrowserDataSubsystem* ContentBrowserData = nullptr;
-	TOptional<FContentBrowserFolderContentsFilter> FolderFilter;
-	const bool bDisplayEmptyFolders = true;
-};
+/** 
+ * SAssetView
+ */
+SAssetView::SAssetView()
+: Items(MakePimpl<FAssetViewItemCollection>())
+{
+
+}
 
 SAssetView::~SAssetView()
 {
@@ -294,9 +1138,14 @@ void SAssetView::Construct( const FArguments& InArgs )
 	BackendFilter = InArgs._InitialBackendFilter;
 
 	FrontendFilters = InArgs._FrontendFilters;
-	if ( FrontendFilters.IsValid() )
+	if (FrontendFilters.IsValid())
+	{
+		FrontendFilters->OnChanged().AddSP(this, &SAssetView::OnFrontendFiltersChanged);
+	}
+	TextFilter = InArgs._TextFilter;
+	if (TextFilter.IsValid())
 	{
-		FrontendFilters->OnChanged().AddSP( this, &SAssetView::OnFrontendFiltersChanged );
+		TextFilter->OnChanged().AddSP(this, &SAssetView::OnFrontendFiltersChanged);
 	}
 
 	OnShouldFilterAsset = InArgs._OnShouldFilterAsset;
@@ -553,10 +1402,9 @@ void SAssetView::Construct( const FArguments& InArgs )
 
 TOptional< float > SAssetView::GetIsWorkingProgressBarState() const
 {
-	if (InitialNumAmortizedTasks > 0)
+	if (Items->HasItemsPendingFilter())
 	{
-		const int32 CompletedTasks = FMath::Max(0, InitialNumAmortizedTasks - ItemsPendingFrontendFilter.Num());
-		return static_cast<float>(CompletedTasks) / static_cast<float>(InitialNumAmortizedTasks);
+		return static_cast<float>(Items->GetFilterProgress()) / static_cast<float>(Items->Num());
 	}
 	return 0.0f;
 }
@@ -598,17 +1446,24 @@ bool SAssetView::IsAssetPathSelected() const
 
 void SAssetView::SetBackendFilter(const FARFilter& InBackendFilter, TArray<TSharedRef<const FPathPermissionList>>* InCustomPermissionLists)
 {
-	// Update the path and collection lists
-	BackendFilter = InBackendFilter;
-	if (InCustomPermissionLists)
-	{
-		BackendCustomPathFilters = *InCustomPermissionLists;
-	}
-	else
+	using namespace UE::AssetView;
+	// Sometimes "filter changed" notifications are broadcast for the content browser to rebuild its filtering when nothing actually changed
+	// Notably custom text filters will do this. 
+	// If we don't need to do a full refresh, don't bother. 
+	if (AreBackendFiltersDifferent(BackendFilter, InBackendFilter) 
+	|| AreCustomPermissionListsDifferent(InCustomPermissionLists, BackendCustomPathFilters))
 	{
-		BackendCustomPathFilters.Reset();
+		BackendFilter = InBackendFilter;
+		if (InCustomPermissionLists)
+		{
+			BackendCustomPathFilters = *InCustomPermissionLists;
+		}
+		else
+		{
+			BackendCustomPathFilters.Reset();
+		}
+		RequestSlowFullListRefresh();
 	}
-	RequestSlowFullListRefresh();
 }
 
 void SAssetView::AppendBackendFilter(FARFilter& FilterToAppendTo) const
@@ -701,7 +1556,7 @@ void SAssetView::BeginCreateDeferredItem()
 	if (DeferredItemToCreate.IsValid() && !DeferredItemToCreate->bWasAddedToView)
 	{
 		TSharedPtr<FAssetViewItem> NewItem = MakeShared<FAssetViewItem>(DeferredItemToCreate->ItemContext.GetItem());
-		NewItem->RenameWhenScrolledIntoView();
+		AwaitingScrollIntoViewForRename = NewItem;
 		DeferredItemToCreate->bWasAddedToView = true;
 
 		FilteredAssetItems.Insert(NewItem, 0);
@@ -723,6 +1578,7 @@ FContentBrowserItem SAssetView::EndCreateDeferredItem(const TSharedPtr<FAssetVie
 		checkf(FContentBrowserItemKey(InItem->GetItem()) == FContentBrowserItemKey(DeferredItemToCreate->ItemContext.GetItem()), TEXT("DeferredItemToCreate was still set when attempting to rename a different item!"));
 
 		// Remove the temporary item before we do any work to ensure the new item creation is not prevented
+		Items->RemoveItem(InItem);
 		FilteredAssetItems.Remove(InItem);
 		RefreshList();
 
@@ -754,9 +1610,9 @@ void SAssetView::CreateNewAsset(const FString& DefaultAssetName, const FString&
 
 void SAssetView::RenameItem(const FContentBrowserItem& ItemToRename)
 {
-	if (const TSharedPtr<FAssetViewItem> Item = AvailableBackendItems.FindRef(FContentBrowserItemKey(ItemToRename)))
+	if (const TSharedPtr<FAssetViewItem> Item = Items->FindItemForRename(ItemToRename))
 	{
-		Item->RenameWhenScrolledIntoView();
+		AwaitingScrollIntoViewForRename = Item;
 		
 		SetSelection(Item);
 		RequestScrollIntoView(Item);
@@ -1018,6 +1874,8 @@ void SAssetView::AdjustActiveSelection(int32 SelectionDelta)
 
 void SAssetView::Tick( const FGeometry& AllottedGeometry, const double InCurrentTime, const float InDeltaTime )
 {
+	using namespace UE::AssetView;
+
 	// Adjust min and max thumbnail scale based on dpi
 	MinThumbnailScale = (0.2f * ThumbnailScaleRangeScalar)/AllottedGeometry.Scale;
 	MaxThumbnailScale = (1.9f * ThumbnailScaleRangeScalar)/AllottedGeometry.Scale;
@@ -1076,12 +1934,12 @@ void SAssetView::Tick( const FGeometry& AllottedGeometry, const double InCurrent
 		if (AmortizeStartTime == 0)
 		{
 			AmortizeStartTime = FPlatformTime::Seconds();
-			InitialNumAmortizedTasks = ItemsPendingFrontendFilter.Num();
+			InitialNumAmortizedTasks = Items->Num();
 			
 			CurrentFrontendFilterTelemetry = { ViewCorrelationGuid, FilterSessionCorrelationGuid };
 			CurrentFrontendFilterTelemetry.FrontendFilters = FrontendFilters;
-			CurrentFrontendFilterTelemetry.TotalItemsToFilter = ItemsPendingFrontendFilter.Num() + ItemsPendingPriorityFilter.Num();
-			CurrentFrontendFilterTelemetry.PriorityItemsToFilter = ItemsPendingPriorityFilter.Num();
+			CurrentFrontendFilterTelemetry.TotalItemsToFilter = Items->Num();
+			CurrentFrontendFilterTelemetry.PriorityItemsToFilter = 0;
 		}
 
 		int32 PreviousFilteredAssetItems = FilteredAssetItems.Num();
@@ -1225,13 +2083,11 @@ void SAssetView::Tick( const FGeometry& AllottedGeometry, const double InCurrent
 		TSharedPtr<SWindow> OwnerWindow = FSlateApplication::Get().FindWidgetWindow(AsShared());
 		if (!OwnerWindow.IsValid())
 		{
-			AssetAwaitingRename->ClearRenameWhenScrolledIntoView();
 			AwaitingRename = nullptr;
 		}
 		else if (OwnerWindow->HasAnyUserFocusOrFocusedDescendants())
 		{
 			AssetAwaitingRename->OnRenameRequested().ExecuteIfBound();
-			AssetAwaitingRename->ClearRenameWhenScrolledIntoView();
 			AwaitingRename = nullptr;
 		}
 	}
@@ -1247,13 +2103,13 @@ void SAssetView::CalculateFillScale( const FGeometry& AllottedGeometry )
 		const float ScrollbarWidth = 16 + 1;
 		float TotalWidth = AllottedGeometry.GetLocalSize().X -(ScrollbarWidth);
 		float Coverage = TotalWidth / ItemWidth;
-		int32 Items = (int)( TotalWidth / ItemWidth );
+		int32 NumItems = (int)( TotalWidth / ItemWidth );
 
 		// If there isn't enough room to support even a single item, don't apply a fill scale.
-		if ( Items > 0 )
+		if ( NumItems > 0 )
 		{
-			float GapSpace = ItemWidth * ( Coverage - (float)Items );
-			float ExpandAmount = GapSpace / (float)Items;
+			float GapSpace = ItemWidth * ( Coverage - (float)NumItems );
+			float ExpandAmount = GapSpace / (float)NumItems;
 			FillScale = ( ItemWidth + ExpandAmount ) / ItemWidth;
 			FillScale = FMath::Max( 1.0f, FillScale );
 		}
@@ -1305,7 +2161,7 @@ void SAssetView::CalculateThumbnailHintColorAndOpacity()
 
 bool SAssetView::HasItemsPendingFilter() const
 {
-	return (ItemsPendingPriorityFilter.Num() + ItemsPendingFrontendFilter.Num()) > 0;
+	return Items->HasItemsPendingFilter();
 }
 
 void SAssetView::ProcessItemsPendingFilter(const double TickStartTime)
@@ -1313,97 +2169,24 @@ void SAssetView::ProcessItemsPendingFilter(const double TickStartTime)
 	const double ProcessItemsPendingFilterStartTime = FPlatformTime::Seconds();
 
 	FAssetViewFrontendFilterHelper FrontendFilterHelper(this);
+	const bool bFlushAllPendingItems = TickStartTime < 0;
+	int32 OldCount = FilteredAssetItems.Num();
+	Items->UpdateItemFiltering(FrontendFilterHelper, bFlushAllPendingItems ? MAX_dbl : TickStartTime + MaxSecondsPerFrame, FilteredAssetItems);
 
-	auto UpdateFilteredAssetItemTypeCounts = [this](const TSharedPtr<FAssetViewItem>& InItem)
+	if (CurrentViewType == EAssetViewType::Column)
 	{
-		if (CurrentViewType == EAssetViewType::Column)
+		for (int32 i = OldCount; i < FilteredAssetItems.Num(); ++i)
 		{
-			const FContentBrowserItemDataAttributeValue TypeNameValue = InItem->GetItem().GetItemAttribute(ContentBrowserItemAttributes::ItemTypeName);
+			const TSharedPtr<FAssetViewItem>& Item = FilteredAssetItems[i];
+			const FContentBrowserItemDataAttributeValue TypeNameValue = Item->GetItem().GetItemAttribute(ContentBrowserItemAttributes::ItemTypeName);
 			if (TypeNameValue.IsValid())
 			{
 				FilteredAssetItemTypeCounts.FindOrAdd(TypeNameValue.GetValue<FName>())++;
 			}
 		}
-	};
-
-	const bool bRunQueryFilter = OnShouldFilterAsset.IsBound() || OnShouldFilterItem.IsBound();
-	const bool bFlushAllPendingItems = TickStartTime < 0;
-
-	bool bRefreshList = false;
-	bool bHasTimeRemaining = true;
-
-	auto FilterItem = [this, bRunQueryFilter, &bRefreshList, &FrontendFilterHelper, &UpdateFilteredAssetItemTypeCounts](const TSharedPtr<FAssetViewItem>& ItemToFilter)
-	{
-		// Run the query filter if required
-		if (bRunQueryFilter)
-		{
-			const bool bPassedBackendFilter = FrontendFilterHelper.DoesItemPassQueryFilter(ItemToFilter);
-			if (!bPassedBackendFilter)
-			{
-				AvailableBackendItems.Remove(FContentBrowserItemKey(ItemToFilter->GetItem()));
-				return;
-			}
-		}
-
-		// Run the frontend filter
-		{
-			const bool bPassedFrontendFilter = FrontendFilterHelper.DoesItemPassFrontendFilter(ItemToFilter);
-			if (bPassedFrontendFilter)
-			{
-				checkAssetList(!FilteredAssetItems.Contains(ItemToFilter));
-
-				bRefreshList = true;
-				FilteredAssetItems.Add(ItemToFilter);
-				UpdateFilteredAssetItemTypeCounts(ItemToFilter);
-			}
-		}
-	};
-
-	// Run the prioritized set first
-	// This data must be processed this frame, so skip the amortization time checks within the loop itself
-	if (ItemsPendingPriorityFilter.Num() > 0)
-	{
-		for (const TSharedPtr<FAssetViewItem>& ItemToFilter : ItemsPendingPriorityFilter)
-		{
-			// Make sure this item isn't pending in another list
-			{
-				const uint32 ItemToFilterHash = GetTypeHash(ItemToFilter);
-				ItemsPendingFrontendFilter.RemoveByHash(ItemToFilterHash, ItemToFilter);
-			}
-
-			// Apply any filters and update the view
-			FilterItem(ItemToFilter);
-		}
-		ItemsPendingPriorityFilter.Reset();
-
-		// Check to see if we have run out of time in this tick
-		if (!bFlushAllPendingItems && (FPlatformTime::Seconds() - TickStartTime) > MaxSecondsPerFrame)
-		{
-			bHasTimeRemaining = false;
-		}
-	}
-
-	// Filter as many items as possible until we run out of time
-	if (bHasTimeRemaining && ItemsPendingFrontendFilter.Num() > 0)
-	{
-		for (auto ItemIter = ItemsPendingFrontendFilter.CreateIterator(); ItemIter; ++ItemIter)
-		{
-			const TSharedPtr<FAssetViewItem> ItemToFilter = *ItemIter;
-			ItemIter.RemoveCurrent();
-
-			// Apply any filters and update the view
-			FilterItem(ItemToFilter);
-
-			// Check to see if we have run out of time in this tick
-			if (!bFlushAllPendingItems && (FPlatformTime::Seconds() - TickStartTime) > MaxSecondsPerFrame)
-			{
-				bHasTimeRemaining = false;
-				break;
-			}
-		}
 	}
 
-	if (bRefreshList)
+	if (FilteredAssetItems.Num() > OldCount)
 	{
 		bPendingSortFilteredItems = true;
 		RefreshList();
@@ -1721,6 +2504,7 @@ TSharedRef<SAssetTileView> SAssetView::CreateTileView()
 		.SelectionMode( SelectionMode )
 		.ListItemsSource(&FilteredAssetItems)
 		.OnGenerateTile(this, &SAssetView::MakeTileViewWidget)
+		.OnItemToString_Debug_Static(&FAssetViewItem::ItemToString_Debug)
 		.OnItemScrolledIntoView(this, &SAssetView::ItemScrolledIntoView)
 		.OnContextMenuOpening(this, &SAssetView::OnGetContextMenuContent)
 		.OnMouseButtonDoubleClick(this, &SAssetView::OnListMouseButtonDoubleClick)
@@ -1742,6 +2526,7 @@ TSharedRef<SAssetListView> SAssetView::CreateListView()
 		.SelectionMode( SelectionMode )
 		.ListItemsSource(&FilteredAssetItems)
 		.OnGenerateRow(this, &SAssetView::MakeListViewWidget)
+		.OnItemToString_Debug_Static(&FAssetViewItem::ItemToString_Debug)
 		.OnItemScrolledIntoView(this, &SAssetView::ItemScrolledIntoView)
 		.OnContextMenuOpening(this, &SAssetView::OnGetContextMenuContent)
 		.OnMouseButtonDoubleClick(this, &SAssetView::OnListMouseButtonDoubleClick)
@@ -1787,6 +2572,7 @@ TSharedRef<SAssetColumnView> SAssetView::CreateColumnView()
 		.SelectionMode( SelectionMode )
 		.ListItemsSource(&FilteredAssetItems)
 		.OnGenerateRow(this, &SAssetView::MakeColumnViewWidget)
+		.OnItemToString_Debug_Static(&FAssetViewItem::ItemToString_Debug)
 		.OnItemScrolledIntoView(this, &SAssetView::ItemScrolledIntoView)
 		.OnContextMenuOpening(this, &SAssetView::OnGetContextMenuContent)
 		.OnMouseButtonDoubleClick(this, &SAssetView::OnListMouseButtonDoubleClick)
@@ -1983,7 +2769,7 @@ FContentBrowserDataFilter SAssetView::CreateBackendDataFilter(bool bInvalidateCa
 			}
 		}
 	}
-
+	
 	if (bShowDisallowedAssetClassAsUnsupportedItems && AssetClassPermissionList && AssetClassPermissionList->HasFiltering())
 	{
 		// The unsupported item will created as an unsupported asset item instead of normal asset item for the writable folders
@@ -2036,116 +2822,216 @@ void SAssetView::RefreshSourceItems()
 
 	FilterSessionCorrelationGuid = FGuid::NewGuid();
 	UE::Telemetry::ContentBrowser::FBackendFilterTelemetry Telemetry(ViewCorrelationGuid, FilterSessionCorrelationGuid);
-	FilteredAssetItems.Reset();
-	FilteredAssetItemTypeCounts.Reset();
 	VisibleItems.Reset();
 	RelevantThumbnails.Reset();
 
-	TMap<FContentBrowserItemKey, TSharedPtr<FAssetViewItem>> PreviousAvailableBackendItems = MoveTemp(AvailableBackendItems);
-	AvailableBackendItems.Reset();
-	ItemsPendingPriorityFilter.Reset();
-	ItemsPendingFrontendFilter.Reset();
+	if (SourcesData.OnEnumerateCustomSourceItemDatas.IsBound())
 	{
-		UContentBrowserDataSubsystem* ContentBrowserData = IContentBrowserDataModule::Get().GetSubsystem();
+		Telemetry.bHasCustomItemSources = true;
+	}
 
-		auto AddNewItem = [this, &PreviousAvailableBackendItems](FContentBrowserItemData&& InItemData)
-		{
-			const FContentBrowserItemKey ItemDataKey(InItemData);
-			const uint32 ItemDataKeyHash = GetTypeHash(ItemDataKey);
+	const bool bInvalidateFilterCache = true;
+	FContentBrowserDataFilter DataFilter = CreateBackendDataFilter(bInvalidateFilterCache);
+	Telemetry.DataFilter = &DataFilter;
+	bool bChangedRecursiveness = bWereItemsRecursivelyFiltered != DataFilter.bRecursivePaths;
+	bWereItemsRecursivelyFiltered = DataFilter.bRecursivePaths;
+
+	Items->RefreshItemsFromBackend(SourcesData, DataFilter, !bChangedRecursiveness);
+
+	Telemetry.NumBackendItems = Items->Num();
+	Telemetry.RefreshSourceItemsDurationSeconds = FPlatformTime::Seconds() - RefreshSourceItemsStartTime;
+	FTelemetryRouter::Get().ProvideTelemetry(Telemetry);
+	UE_LOG(LogContentBrowser, VeryVerbose, TEXT("AssetView - RefreshSourceItems completed in %0.4f seconds"), 
+		FPlatformTime::Seconds() - RefreshSourceItemsStartTime);
+}
+
+void FAssetViewItemCollection::RefreshItemsFromBackend(const FSourcesData& SourcesData, const FContentBrowserDataFilter& DataFilter, bool bAllowItemRecycling)
+{
+	if(TextFilterTask.IsValid())
+	{
+		ShouldCancelTextFiltering = true;
+		TextFilterTask.BusyWait();
+		TextFilterTask = {};
+	}
+
+	UContentBrowserDataSubsystem* ContentBrowserData = IContentBrowserDataModule::Get().GetSubsystem();
+	TArray<FContentBrowserItemData> NewItemDatas;
+	if (DataFilter.bRecursivePaths)
+	{
+		NewItemDatas.Reserve(1024 * 1024); // Assume many recursive searches will return a lot of items and start with a lot of space 
+	}
 
-			TSharedPtr<FAssetViewItem>& NewItem = AvailableBackendItems.FindOrAddByHash(ItemDataKeyHash, ItemDataKey);
-			if (!NewItem && InItemData.IsFile())
+	if (SourcesData.OnEnumerateCustomSourceItemDatas.IsBound())
+	{
+		SourcesData.OnEnumerateCustomSourceItemDatas.Execute([&NewItemDatas](FContentBrowserItemData&& InItemData) { NewItemDatas.Add(MoveTemp(InItemData)); return true; });
+	}
+
+	if (SourcesData.IsIncludingVirtualPaths() || SourcesData.HasCollections()) 
+	{
+		if (SourcesData.HasCollections() && EnumHasAnyFlags(DataFilter.ItemCategoryFilter, EContentBrowserItemCategoryFilter::IncludeCollections))
+		{
+			// If we are showing collections then we may need to add dummy folder items for the child collections
+			// Note: We don't check the IncludeFolders flag here, as that is forced to false when collections are selected,
+			// instead we check the state of bIncludeChildCollections which will be false when we want to show collection folders
+			const FContentBrowserDataCollectionFilter* CollectionFilter = DataFilter.ExtraFilters.FindFilter<FContentBrowserDataCollectionFilter>();
+			if (CollectionFilter && !CollectionFilter->bIncludeChildCollections)
 			{
-				// Re-use the old view item where possible to avoid list churn when our backend view already included the item
-				if (TSharedPtr<FAssetViewItem>* PreviousItem = PreviousAvailableBackendItems.FindByHash(ItemDataKeyHash, ItemDataKey))
+				FCollectionManagerModule& CollectionManagerModule = FCollectionManagerModule::GetModule();
+			
+				TArray<FCollectionNameType> ChildCollections;
+				for(const FCollectionNameType& Collection : SourcesData.Collections)
 				{
-					NewItem = *PreviousItem;
-					NewItem->ClearCachedCustomColumns();
+					ChildCollections.Reset();
+					CollectionManagerModule.Get().GetChildCollections(Collection.Name, Collection.Type, ChildCollections);
+
+					for (const FCollectionNameType& ChildCollection : ChildCollections)
+					{
+						// Use "Collections" as the root of the path to avoid this being confused with other view folders - see ContentBrowserUtils::IsCollectionPath
+						FContentBrowserItemData FolderItemData(
+							nullptr, 
+							EContentBrowserItemFlags::Type_Folder | EContentBrowserItemFlags::Category_Collection, 
+							*FString::Printf(TEXT("/Collections/%s/%s"), ECollectionShareType::ToString(ChildCollection.Type), *ChildCollection.Name.ToString()), 
+							ChildCollection.Name, 
+							FText::FromName(ChildCollection.Name), 
+							nullptr,
+							FName()
+							);
+
+						NewItemDatas.Add(MoveTemp(FolderItemData));
+					}
 				}
 			}
-			if (NewItem)
-			{
-				NewItem->AppendItemData(InItemData);
-				NewItem->CacheCustomColumns(CustomColumns, true, true, false /*bUpdateExisting*/);
-			}
-			else
-			{
-				NewItem = MakeShared<FAssetViewItem>(MoveTemp(InItemData));
-			}
-
-			return true;
-		};
+		}
 
-		if (SourcesData.OnEnumerateCustomSourceItemDatas.IsBound())
+		if (SourcesData.IsIncludingVirtualPaths())
 		{
-			Telemetry.bHasCustomItemSources = true;
-			SourcesData.OnEnumerateCustomSourceItemDatas.Execute(AddNewItem);
+			SCOPED_NAMED_EVENT(FetchCBItems, FColor::White);
+			static const FName RootPath = "/";
+			const TArrayView<const FName> DataSourcePaths = SourcesData.HasVirtualPaths() ? MakeArrayView(SourcesData.VirtualPaths) : MakeArrayView(&RootPath, 1);
+			for (const FName& DataSourcePath : DataSourcePaths)
+			{
+				// Ensure paths do not contain trailing slash
+				ensure(DataSourcePath == RootPath || !FStringView(FNameBuilder(DataSourcePath)).EndsWith(TEXT('/')));
+				ContentBrowserData->EnumerateItemsUnderPath(DataSourcePath, DataFilter, [&NewItemDatas](FContentBrowserItemData&& Item) { NewItemDatas.Add(MoveTemp(Item)); return true; });
+			}
 		}
+	}
 
-		FContentBrowserDataFilter DataFilter; // Must live long enough to provide telemetry
-		if (SourcesData.IsIncludingVirtualPaths() || SourcesData.HasCollections()) 
-		{
-			const bool bInvalidateFilterCache = true;
-			DataFilter = CreateBackendDataFilter(bInvalidateFilterCache);
-			Telemetry.DataFilter = &DataFilter;
+	TArray<TSharedPtr<FAssetViewItem>> OldItems = MoveTemp(Items);
+	FHashTable OldLookup = MoveTemp(Lookup); 
 
-			bWereItemsRecursivelyFiltered = DataFilter.bRecursivePaths;
+	Items.Reset(NewItemDatas.Num());
+	Items.AddZeroed(NewItemDatas.Num());
+	bItemsPendingRemove = false;
 
-			if (SourcesData.HasCollections() && EnumHasAnyFlags(DataFilter.ItemCategoryFilter, EContentBrowserItemCategoryFilter::IncludeCollections))
-			{
-				// If we are showing collections then we may need to add dummy folder items for the child collections
-				// Note: We don't check the IncludeFolders flag here, as that is forced to false when collections are selected,
-				// instead we check the state of bIncludeChildCollections which will be false when we want to show collection folders
-				const FContentBrowserDataCollectionFilter* CollectionFilter = DataFilter.ExtraFilters.FindFilter<FContentBrowserDataCollectionFilter>();
-				if (CollectionFilter && !CollectionFilter->bIncludeChildCollections)
-				{
-					FCollectionManagerModule& CollectionManagerModule = FCollectionManagerModule::GetModule();
+	// Create or recyle FAssetViewItem for each FContentBrowserItemData. Build new hashtable concurrently at the same time
+	uint32 HashSize = TSetAllocator<>::GetNumberOfHashBuckets(Items.Num());
+	std::atomic<bool> bAnyFolders(false);
+	std::atomic<bool> bAnyRecycled(false);
+	Lookup.Clear(HashSize, Items.Num());
+	{
+		SCOPED_NAMED_EVENT(CreateItems, FColor::White);
+		ParallelFor(TEXT("CreateFAssetViewItem"), NewItemDatas.Num(), 16 * 1024,
+			[&NewItemDatas, &OldItems, &OldLookup, bAllowItemRecycling, &bAnyRecycled, &bAnyFolders, this](int32 Index) {
+				FContentBrowserItemData& ItemData = NewItemDatas[Index];
+				FName VirtualPath = ItemData.GetVirtualPath();
+				TSharedPtr<FAssetViewItem> OldItem;
 				
-					TArray<FCollectionNameType> ChildCollections;
-					for(const FCollectionNameType& Collection : SourcesData.Collections)
+				FContentBrowserItemKey ItemKey(ItemData);
+				uint32 Hash = HashItem(ItemData);
+				if (bAllowItemRecycling)
+				{
+					for (uint32 SearchIndex = OldLookup.First(Hash); OldLookup.IsValid(SearchIndex); SearchIndex = OldLookup.Next(SearchIndex))
 					{
-						ChildCollections.Reset();
-						CollectionManagerModule.Get().GetChildCollections(Collection.Name, Collection.Type, ChildCollections);
-
-						for (const FCollectionNameType& ChildCollection : ChildCollections)
+						if (OldItems[SearchIndex].IsValid() && ItemKey == FContentBrowserItemKey(OldItems[SearchIndex]->GetItem()))
 						{
-							// Use "Collections" as the root of the path to avoid this being confused with other view folders - see ContentBrowserUtils::IsCollectionPath
-							FContentBrowserItemData FolderItemData(
-								nullptr, 
-								EContentBrowserItemFlags::Type_Folder | EContentBrowserItemFlags::Category_Collection, 
-								*FString::Printf(TEXT("/Collections/%s/%s"), ECollectionShareType::ToString(ChildCollection.Type), *ChildCollection.Name.ToString()), 
-								ChildCollection.Name, 
-								FText::FromName(ChildCollection.Name), 
-								nullptr,
-								FName()
-								);
-
-							const FContentBrowserItemKey FolderItemDataKey(FolderItemData);
-							AvailableBackendItems.Add(FolderItemDataKey, MakeShared<FAssetViewItem>(MoveTemp(FolderItemData)));
+							bAnyRecycled.store(true, std::memory_order_relaxed);
+							OldItem = OldItems[SearchIndex];
+							break;
 						}
 					}
 				}
-			}
 
-			if (SourcesData.IsIncludingVirtualPaths())
+				if (ItemData.IsFolder())
+				{
+					bAnyFolders.store(true, std::memory_order_relaxed);
+				}
+
+				if (OldItem.IsValid())
+				{
+					OldItem->ResetItemData(MoveTemp(ItemData));
+					Items[Index] = MoveTemp(OldItem);
+				}
+				else
+				{
+					Items[Index] = MakeShared<FAssetViewItem>(MoveTemp(ItemData));
+				}
+				Lookup.Add_Concurrent(Hash, Index);
+			}, UE::AssetView::AllowParallelism ? EParallelForFlags::None : EParallelForFlags::ForceSingleThread);
+	}
+
+	// Reset this before merging to avoid duplicate work around nulled entries
+	ResetFilterState();	
+
+	NumValidItems = Items.Num();
+	if (bAnyFolders.load())
+	{
+		SCOPED_NAMED_EVENT(MergeDuplicates, FColor::White);
+		// Merge items with the same path
+		// Loop over each bucket lookup for duplicate names in that bucket and merging the items
+		// This is done in parallel because each worker will only touch items in its bucket
+		ParallelFor(TEXT("MergeDuplicates"), HashSize, 8, [this](int32 JobIndex) {
+			uint32 Bucket = *reinterpret_cast<uint32*>(&JobIndex);
+			for (uint32 StartIndex = Lookup.First(Bucket); Lookup.IsValid(StartIndex); StartIndex = Lookup.Next(StartIndex))
 			{
-				static const FName RootPath = "/";
-				const TArrayView<const FName> DataSourcePaths = SourcesData.HasVirtualPaths() ? MakeArrayView(SourcesData.VirtualPaths) : MakeArrayView(&RootPath, 1);
-				for (const FName& DataSourcePath : DataSourcePaths)
+				if (!Items[StartIndex]->IsFolder())
 				{
-					// Ensure paths do not contain trailing slash
-					ensure(DataSourcePath == RootPath || !FStringView(FNameBuilder(DataSourcePath)).EndsWith(TEXT('/')));
-					ContentBrowserData->EnumerateItemsUnderPath(DataSourcePath, DataFilter, AddNewItem);
+					continue;
+				}
+
+				FContentBrowserItemKey MergeWithKey(Items[StartIndex]->GetItem());
+				uint32 Other = Lookup.Next(StartIndex);
+				while (Lookup.IsValid(Other))
+				{
+					check(Items[Other].IsValid());
+					FContentBrowserItemKey OtherKey(Items[Other]->GetItem());
+					if (MergeWithKey == OtherKey)
+					{
+						uint32 ToRemove = Other;
+						Other = Lookup.Next(Other);
+						Lookup.Remove(Bucket, ToRemove);
+						TSharedPtr<FAssetViewItem> RemovedItem = MarkItemRemoved(ToRemove);
+						Items[StartIndex]->AppendItemData(RemovedItem->GetItem()); 
+						Items[ToRemove].Reset();
+					}
+					else
+					{
+						Other = Lookup.Next(Other);
+					}
 				}
 			}
-		}
 
-		Telemetry.NumBackendItems = AvailableBackendItems.Num();
-		Telemetry.RefreshSourceItemsDurationSeconds = FPlatformTime::Seconds() - RefreshSourceItemsStartTime;
-		FTelemetryRouter::Get().ProvideTelemetry(Telemetry);
+		}, EParallelForFlags::Unbalanced | (UE::AssetView::AllowParallelism ? EParallelForFlags::None : EParallelForFlags::ForceSingleThread));
+	}
+
+	// We already nulled out the items we removed above 
+	bItemsPendingRemove = false;
+
+	// If we recycled any items (e.g. we changed item visibility settings but not path) notify their widgets that we changed the item data
+	if (bAnyRecycled.load())
+	{
+		for (const TSharedPtr<FAssetViewItem>& Item : Items)
+		{
+			if (Item.IsValid())
+			{
+				Item->BroadcastItemDataChanged();
+			}
+		}
 	}
 
-	UE_LOG(LogContentBrowser, VeryVerbose, TEXT("AssetView - RefreshSourceItems completed in %0.4f seconds"), FPlatformTime::Seconds() - RefreshSourceItemsStartTime);
+	// Until we start filtering and get a compiled filter, initialize to no filtering
+	bAllItemsPassedTextFilter = true;
+	CompiledTextFilter.Reset();
 }
 
 bool SAssetView::IsFilteringRecursively() const
@@ -2233,7 +3119,6 @@ void SAssetView::RefreshFilteredItems()
 
 	OnInterruptFiltering();
 	
-	ItemsPendingFrontendFilter.Reset();
 	FilteredAssetItems.Reset();
 	FilteredAssetItemTypeCounts.Reset();
 	RelevantThumbnails.Reset();
@@ -2244,10 +3129,11 @@ void SAssetView::RefreshFilteredItems()
 	LastSortTime = 0;
 	bPendingSortFilteredItems = true;
 
-	ItemsPendingFrontendFilter.Reserve(AvailableBackendItems.Num());
-	for (const auto& AvailableBackendItemPair : AvailableBackendItems)
+	Items->AbortTextFiltering();
+	Items->ResetFilterState();
+	if (TextFilter.IsValid())
 	{
-		ItemsPendingFrontendFilter.Add(AvailableBackendItemPair.Value);
+		Items->StartTextFiltering(TextFilter);
 	}
 
 	// Let the frontend filters know the currently used asset filter in case it is necessary to conditionally filter based on path or class filters
@@ -2581,9 +3467,28 @@ void SAssetView::OnFrontendFiltersChanged()
 
 	RequestQuickFrontendListRefresh();
 
-	// If we're not operating on recursively filtered data, we need to ensure a full slow
-	// refresh is performed.
-	if ( ShouldFilterRecursively() && !bWereItemsRecursivelyFiltered )
+	// Combine any currently active custom text filters with the asset text filtering task
+	if (TextFilter.IsValid() && FrontendFilters.IsValid())
+	{ 
+		TArray<FText> CustomTextFilters;
+		for (int32 i = 0; i < FrontendFilters->Num(); ++i)
+		{
+			TSharedPtr<FFrontendFilter> Filter = StaticCastSharedPtr<FFrontendFilter>(FrontendFilters->GetFilterAtIndex(i));
+			if (Filter.IsValid())
+			{
+				TOptional<FText> Text = Filter->GetAsCustomTextFilter();
+				if (Text.IsSet())
+				{
+					CustomTextFilters.Add(Text.GetValue());
+				}
+			}
+		}
+		TextFilter->SetCustomTextFilters(MoveTemp(CustomTextFilters));
+	}
+	
+
+	// If we're changing between recursive and non-recursive data, we need to fully refresh the source items
+	if (ShouldFilterRecursively() != bWereItemsRecursivelyFiltered)
 	{
 		RequestSlowFullListRefresh();
 	}
@@ -3826,7 +4731,7 @@ TSharedRef<ITableRow> SAssetView::MakeColumnViewWidget(TSharedPtr<FAssetViewItem
 	}
 
 	// Update the cached custom data
-	AssetItem->CacheCustomColumns(CustomColumns, false, true, false);
+	AssetItem->CacheCustomColumns(CustomColumns, /* bUpdateSortData */ false, /* bUpdateDisplayText */ true, /* bUpdateExisting */ false);
 	
 	return
 		SNew( SAssetColumnViewRow, OwnerTable )
@@ -3847,6 +4752,7 @@ TSharedRef<ITableRow> SAssetView::MakeColumnViewWidget(TSharedPtr<FAssetViewItem
 		);
 }
 
+
 void SAssetView::AssetItemWidgetDestroyed(const TSharedPtr<FAssetViewItem>& Item)
 {
 	if(RenamingAsset.Pin().Get() == Item.Get())
@@ -4010,8 +4916,10 @@ void SAssetView::AssetSelectionChanged( TSharedPtr< FAssetViewItem > AssetItem,
 
 void SAssetView::ItemScrolledIntoView(TSharedPtr<FAssetViewItem> AssetItem, const TSharedPtr<ITableRow>& Widget )
 {
-	if (AssetItem->ShouldRenameWhenScrolledIntoView())
+	if (AssetItem == AwaitingScrollIntoViewForRename)
 	{
+		AwaitingScrollIntoViewForRename.Reset();
+
 		// Make sure we have window focus to avoid the inline text editor from canceling itself if we try to click on it
 		// This can happen if creating an asset opens an intermediary window which steals our focus, 
 		// eg, the blueprint and slate widget style class windows (TTP# 314240)
@@ -4189,7 +5097,6 @@ void SAssetView::AssetRenameBegin(const TSharedPtr<FAssetViewItem>& Item, const
 
 void SAssetView::AssetRenameCommit(const TSharedPtr<FAssetViewItem>& Item, const FString& NewName, const FSlateRect& MessageAnchor, const ETextCommit::Type CommitType)
 {
-	bool bSuccess = false;
 	FText ErrorMessage;
 	TSharedPtr<FAssetViewItem> UpdatedItem;
 
@@ -4202,11 +5109,8 @@ void SAssetView::AssetRenameCommit(const TSharedPtr<FAssetViewItem>& Item, const
 		FContentBrowserItem NewItem = EndCreateDeferredItem(Item, NewName, bFinalize, ErrorMessage);
 		if (NewItem.IsValid())
 		{
-			bSuccess = true;
-
 			// Add result to view
-			UpdatedItem = AvailableBackendItems.Add(FContentBrowserItemKey(NewItem), MakeShared<FAssetViewItem>(NewItem));
-			FilteredAssetItems.Add(UpdatedItem);
+			UpdatedItem = Items->CreateItemFromUser(MoveTemp(NewItem), FilteredAssetItems);
 		}
 	}
 	else if (CommitType != ETextCommit::OnCleared && !Item->GetItem().GetItemName().ToString().Equals(NewName))
@@ -4217,57 +5121,51 @@ void SAssetView::AssetRenameCommit(const TSharedPtr<FAssetViewItem>& Item, const
 		FContentBrowserItem NewItem;
 		if (Item->GetItem().CanRename(&NewName, &ErrorMessage) && Item->GetItem().Rename(NewName, &NewItem))
 		{
-			bSuccess = true;
-
 			// Add result to view (the old item will be removed via the notifications, as not all data sources may have been able to perform the rename)
-			UpdatedItem = AvailableBackendItems.Add(FContentBrowserItemKey(NewItem), MakeShared<FAssetViewItem>(NewItem));
-			FilteredAssetItems.Add(UpdatedItem);
+			UpdatedItem = Items->CreateItemFromUser(MoveTemp(NewItem), FilteredAssetItems);
 		}
 	}
 	
-	if (bSuccess)
+	if (UpdatedItem)
 	{
-		if (UpdatedItem)
-		{
-			// Sort in the new item
-			bPendingSortFilteredItems = true;
+		// Sort in the new item
+		bPendingSortFilteredItems = true;
 
-			if (UpdatedItem->IsFile())
-			{
-				// Refresh the thumbnail
-				if (TSharedPtr<FAssetThumbnail> AssetThumbnail = RelevantThumbnails.FindRef(Item))
-				{
-					if (UpdatedItem != Item)
-					{
-						// This item was newly created - move the thumbnail over from the temporary item
-						RelevantThumbnails.Remove(Item);
-						RelevantThumbnails.Add(UpdatedItem, AssetThumbnail);
-						UpdatedItem->GetItem().UpdateThumbnail(*AssetThumbnail);
-					}
-					if (AssetThumbnail->GetAssetData().IsValid())
-					{
-						AssetThumbnailPool->RefreshThumbnail(AssetThumbnail);
-					}
-				}
-			}
-			
-			// Sync the view
+		if (UpdatedItem->IsFile())
+		{
+			// Refresh the thumbnail
+			if (TSharedPtr<FAssetThumbnail> AssetThumbnail = RelevantThumbnails.FindRef(Item))
 			{
-				TArray<FContentBrowserItem> ItemsToSync;
-				ItemsToSync.Add(UpdatedItem->GetItem());
-
-				if (OnItemRenameCommitted.IsBound() && !bUserSearching)
+				if (UpdatedItem != Item)
 				{
-					// If our parent wants to potentially handle the sync, let it, but only if we're not currently searching (or it would cancel the search)
-					OnItemRenameCommitted.Execute(ItemsToSync);
+					// This item was newly created - move the thumbnail over from the temporary item
+					RelevantThumbnails.Remove(Item);
+					RelevantThumbnails.Add(UpdatedItem, AssetThumbnail);
+					UpdatedItem->GetItem().UpdateThumbnail(*AssetThumbnail);
 				}
-				else
+				if (AssetThumbnail->GetAssetData().IsValid())
 				{
-					// Otherwise, sync just the view
-					SyncToItems(ItemsToSync);
+					AssetThumbnailPool->RefreshThumbnail(AssetThumbnail);
 				}
 			}
 		}
+		
+		// Sync the view
+		{
+			TArray<FContentBrowserItem> ItemsToSync;
+			ItemsToSync.Add(UpdatedItem->GetItem());
+
+			if (OnItemRenameCommitted.IsBound() && !bUserSearching)
+			{
+				// If our parent wants to potentially handle the sync, let it, but only if we're not currently searching (or it would cancel the search)
+				OnItemRenameCommitted.Execute(ItemsToSync);
+			}
+			else
+			{
+				// Otherwise, sync just the view
+				SyncToItems(ItemsToSync);
+			}
+		}
 	}
 	else if (!ErrorMessage.IsEmpty())
 	{
@@ -4575,11 +5473,13 @@ bool SAssetView::HasSingleCollectionSource() const
 
 void SAssetView::SetUserSearching(bool bInSearching)
 {
-	if(bUserSearching != bInSearching)
+	bUserSearching = bInSearching;
+
+	// If we're changing between recursive and non-recursive data, we need to fully refresh the source items
+	if (ShouldFilterRecursively() != bWereItemsRecursivelyFiltered)
 	{
 		RequestSlowFullListRefresh();
 	}
-	bUserSearching = bInSearching;
 }
 
 void SAssetView::HandleSettingChanged(FName PropertyName)
@@ -4881,56 +5781,18 @@ void SAssetView::HandleItemDataUpdated(TArrayView<const FContentBrowserItemDataU
 	bool bRefreshView = false;
 	TSet<TSharedPtr<FAssetViewItem>> ItemsPendingInplaceFrontendFilter;
 
-	auto AddItem = [this, &ItemsPendingInplaceFrontendFilter](const FContentBrowserItemKey& InItemDataKey, FContentBrowserItemData&& InItemData)
+	auto AddItem = [this, &ItemsPendingInplaceFrontendFilter](FContentBrowserItemData&& InItemData)
 	{
-		TSharedPtr<FAssetViewItem>& ItemToUpdate = AvailableBackendItems.FindOrAdd(InItemDataKey);
-		if (ItemToUpdate)
-		{
-			// Update the item
-			ItemToUpdate->AppendItemData(MoveTemp(InItemData));
-
-			// Update the custom column data
-			ItemToUpdate->CacheCustomColumns(CustomColumns, true, true, true);
-
-			// This item was modified, so put it in the list of items to be in-place re-tested against the active frontend filter (this can avoid a costly re-sort of the view)
-			// If the item can't be queried in-place (because the item isn't in the view) then it will be added to ItemsPendingPriorityFilter instead
-			ItemsPendingInplaceFrontendFilter.Add(ItemToUpdate);
-		}
-		else
-		{
-			ItemToUpdate = MakeShared<FAssetViewItem>(MoveTemp(InItemData));
-
-			// This item is new so put it in the pending set to be processed over time
-			ItemsPendingFrontendFilter.Add(ItemToUpdate);
-		}
+		TSharedPtr<FAssetViewItem> ItemToUpdate = Items->UpdateData(MoveTemp(InItemData));
+		// Update the custom column data if it exists
+		ItemToUpdate->CacheCustomColumns(CustomColumns, /* bUpdateSortData */ true, /* bUpdateDisplayText */ true, /* bUpdateExisting */ true);
 	};
 
-	auto RemoveItem = [this, &bRefreshView, &ItemsPendingInplaceFrontendFilter](const FContentBrowserItemKey& ItemKey, const FContentBrowserMinimalItemData& ItemDataKey)
+	auto RemoveItem = [this, &bRefreshView, &ItemsPendingInplaceFrontendFilter](const FContentBrowserMinimalItemData& ItemDataKey)
 	{
-		const uint32 ItemDataKeyHash = GetTypeHash(ItemKey);
-
-		if (const TSharedPtr<FAssetViewItem>* ItemToRemovePtr = AvailableBackendItems.FindByHash(ItemDataKeyHash, ItemKey))
+		TSharedPtr<FAssetViewItem> RemovedItem = Items->RemoveItemData(ItemDataKey);
+		if (RemovedItem.IsValid())
 		{
-			TSharedPtr<FAssetViewItem> ItemToRemove = *ItemToRemovePtr;
-			check(ItemToRemove);
-
-			// Only fully remove this item if every sub-item is removed (items become invalid when empty)
-			ItemToRemove->RemoveItemData(ItemDataKey);
-			if (ItemToRemove->GetItem().IsValid())
-			{
-				return;
-			}
-
-			AvailableBackendItems.RemoveByHash(ItemDataKeyHash, ItemKey);
-
-			const uint32 ItemToRemoveHash = GetTypeHash(ItemToRemove);
-
-			// Also ensure this item has been removed from the pending filter lists and the current list view data
-			FilteredAssetItems.RemoveSingle(ItemToRemove);
-			ItemsPendingPriorityFilter.RemoveByHash(ItemToRemoveHash, ItemToRemove);
-			ItemsPendingFrontendFilter.RemoveByHash(ItemToRemoveHash, ItemToRemove);
-			ItemsPendingInplaceFrontendFilter.RemoveByHash(ItemToRemoveHash, ItemToRemove);
-
 			// Need to refresh manually after removing items, as adding relies on the pending filter lists to trigger this
 			bRefreshView = true;
 		}
@@ -4966,35 +5828,34 @@ void SAssetView::HandleItemDataUpdated(TArrayView<const FContentBrowserItemDataU
 	{
 		bool bItemPassFilter = false;
 		FContentBrowserItemData ItemData = GetBackendFilterCompliantItem(ItemDataUpdate.GetItemData(), bItemPassFilter);
-		const FContentBrowserItemKey ItemDataKey(ItemData);
 
 		switch (ItemDataUpdate.GetUpdateType())
 		{
 		case EContentBrowserItemUpdateType::Added:
 			if (bItemPassFilter)
 			{
-				AddItem(ItemDataKey, MoveTemp(ItemData));
+				AddItem(MoveTemp(ItemData));
 			}
 			break;
 
 		case EContentBrowserItemUpdateType::Modified:
 			if (bItemPassFilter)
 			{
-				AddItem(ItemDataKey, MoveTemp(ItemData));
+				AddItem(MoveTemp(ItemData));
 			}
 			else
 			{
-				RemoveItem(ItemDataKey, FContentBrowserMinimalItemData(ItemData));
+				RemoveItem(FContentBrowserMinimalItemData(ItemData));
 			}
 			break;
 
 		case EContentBrowserItemUpdateType::Moved:
 			{
 				const FContentBrowserMinimalItemData OldItemDataKey(ItemData.GetItemType(), ItemDataUpdate.GetPreviousVirtualPath(), ItemData.GetOwnerDataSource());
-				RemoveItem(ItemDataKey, OldItemDataKey);
+				RemoveItem(OldItemDataKey);
 				if (bItemPassFilter)
 				{
-					AddItem(FContentBrowserItemKey(ItemData), MoveTemp(ItemData));
+					AddItem(MoveTemp(ItemData));
 				}
 				else
 				{
@@ -5004,7 +5865,7 @@ void SAssetView::HandleItemDataUpdated(TArrayView<const FContentBrowserItemDataU
 			break;
 
 		case EContentBrowserItemUpdateType::Removed:
-			RemoveItem(ItemDataKey, FContentBrowserMinimalItemData(ItemData));
+			RemoveItem(FContentBrowserMinimalItemData(ItemData));
 			break;
 
 		default:
@@ -5013,57 +5874,10 @@ void SAssetView::HandleItemDataUpdated(TArrayView<const FContentBrowserItemDataU
 		}
 	}
 
-	// Now patch in the in-place frontend filter requests (if possible)
-	if (ItemsPendingInplaceFrontendFilter.Num() > 0)
+	FAssetViewFrontendFilterHelper FrontendFilterHelper(this);
+	if (Items->PerformPriorityFiltering(FrontendFilterHelper, FilteredAssetItems))
 	{
-		FAssetViewFrontendFilterHelper FrontendFilterHelper(this);
-		const bool bRunQueryFilter = OnShouldFilterAsset.IsBound() || OnShouldFilterItem.IsBound();
-
-		for (auto It = FilteredAssetItems.CreateIterator(); It && ItemsPendingInplaceFrontendFilter.Num() > 0; ++It)
-		{
-			const TSharedPtr<FAssetViewItem> ItemToFilter = *It;
-
-			if (ItemsPendingInplaceFrontendFilter.Remove(ItemToFilter) > 0)
-			{
-				bool bRemoveItem = false;
-
-				// Run the query filter if required
-				if (bRunQueryFilter)
-				{
-					const bool bPassedBackendFilter = FrontendFilterHelper.DoesItemPassQueryFilter(ItemToFilter);
-					if (!bPassedBackendFilter)
-					{
-						bRemoveItem = true;
-						AvailableBackendItems.Remove(FContentBrowserItemKey(ItemToFilter->GetItem()));
-					}
-				}
-
-				// Run the frontend filter
-				if (!bRemoveItem)
-				{
-					const bool bPassedFrontendFilter = FrontendFilterHelper.DoesItemPassFrontendFilter(ItemToFilter);
-					if (!bPassedFrontendFilter)
-					{
-						bRemoveItem = true;
-					}
-				}
-
-				// Remove this item?
-				if (bRemoveItem)
-				{
-					bRefreshView = true;
-					It.RemoveCurrent();
-				}
-			}
-		}
-
-		// Do we still have items that could not be in-place filtered?
-		// If so, add them to ItemsPendingPriorityFilter so they are processed into the view ASAP
-		if (ItemsPendingInplaceFrontendFilter.Num() > 0)
-		{
-			ItemsPendingPriorityFilter.Append(MoveTemp(ItemsPendingInplaceFrontendFilter));
-			ItemsPendingInplaceFrontendFilter.Reset();
-		}
+		bRefreshView = true;
 	}
 
 	if (bRefreshView)
@@ -5071,7 +5885,8 @@ void SAssetView::HandleItemDataUpdated(TArrayView<const FContentBrowserItemDataU
 		RefreshList();
 	}
 
-	UE_LOG(LogContentBrowser, VeryVerbose, TEXT("AssetView - HandleItemDataUpdated completed in %0.4f seconds for %d items (%d available items)"), FPlatformTime::Seconds() - HandleItemDataUpdatedStartTime, InUpdatedItems.Num(), AvailableBackendItems.Num());
+	UE_LOG(LogContentBrowser, VeryVerbose, TEXT("AssetView - HandleItemDataUpdated completed in %0.4f seconds for %d items (%d available items)"), 
+		FPlatformTime::Seconds() - HandleItemDataUpdatedStartTime, InUpdatedItems.Num(), Items->Num());
 }
 
 void SAssetView::HandleItemDataDiscoveryComplete()
diff --git a/Engine/Source/Editor/ContentBrowser/Private/SContentBrowser.cpp b/Engine/Source/Editor/ContentBrowser/Private/SContentBrowser.cpp
--- a/Engine/Source/Editor/ContentBrowser/Private/SContentBrowser.cpp
+++ b/Engine/Source/Editor/ContentBrowser/Private/SContentBrowser.cpp
@@ -10,6 +10,7 @@
 #include "AssetRegistry/AssetDataTagMap.h"
 #include "AssetRegistry/AssetRegistryModule.h"
 #include "AssetRegistry/IAssetRegistry.h"
+#include "AssetTextFilter.h"
 #include "AssetThumbnail.h"
 #include "AssetToolsModule.h"
 #include "AssetViewUtils.h"
@@ -345,7 +346,7 @@ void SContentBrowser::Construct( const FArguments& InArgs, const FName& InInstan
 	PathContextMenu->SetOnFolderFavoriteToggled(FPathContextMenu::FOnFolderFavoriteToggled::CreateSP(this, &SContentBrowser::ToggleFolderFavorite));
 	PathContextMenu->SetOnPrivateContentEditToggled(FPathContextMenu::FOnPrivateContentEditToggled::CreateSP(this, &SContentBrowser::TogglePrivateContentEdit));
 	FrontendFilters = MakeShareable(new FAssetFilterCollectionType());
-	TextFilter = MakeShareable( new FFrontendFilter_Text() );
+	TextFilter = MakeShared<FAssetTextFilter>();
 
 	PluginPathFilters = MakeShareable(new FPluginFilterCollectionType());
 
@@ -382,34 +383,35 @@ void SContentBrowser::Construct( const FArguments& InArgs, const FName& InInstan
 	static const IConsoleVariable* EnablePublicAssetFeatureCVar = IConsoleManager::Get().FindConsoleVariable(TEXT("AssetTools.EnablePublicAssetFeature"));
 	const bool bEnablePrivateContentFeature = EnablePublicAssetFeatureCVar && EnablePublicAssetFeatureCVar->GetBool();
 	AssetViewPtr = SNew(SAssetView)
-					   .ThumbnailLabel(Config != nullptr ? Config->ThumbnailLabel : EThumbnailLabel::ClassName)
-					   //.ThumbnailScale(Config != nullptr ? Config->ThumbnailScale : 0.18f)
-					   .InitialViewType(Config != nullptr ? Config->InitialAssetViewType : EAssetViewType::Tile)
-					   .OnNewItemRequested(this, &SContentBrowser::OnNewItemRequested)
+			.ThumbnailLabel(Config != nullptr ? Config->ThumbnailLabel : EThumbnailLabel::ClassName)
+			//.ThumbnailScale(Config != nullptr ? Config->ThumbnailScale : 0.18f)
+			.InitialViewType(Config != nullptr ? Config->InitialAssetViewType : EAssetViewType::Tile)
+			.OnNewItemRequested(this, &SContentBrowser::OnNewItemRequested)
 					   .OnItemSelectionChanged(this, &SContentBrowser::OnItemSelectionChanged, EContentBrowserViewContext::AssetView)
-					   .OnItemsActivated(this, &SContentBrowser::OnItemsActivated)
-					   .OnGetItemContextMenu(this, &SContentBrowser::GetItemContextMenu, EContentBrowserViewContext::AssetView)
-					   .OnItemRenameCommitted(this, &SContentBrowser::OnItemRenameCommitted)
-					   .OnShouldFilterItem(bEnablePrivateContentFeature ? FOnShouldFilterItem::CreateSP(this, &SContentBrowser::HandlePrivateContentFilter) : FOnShouldFilterItem())
-					   .FrontendFilters(FrontendFilters)
-					   .ShowRedirectors_Lambda([this]() { return ContentBrowserUtils::ShouldShowRedirectors(FilterListPtr); })
-					   .HighlightedText(this, &SContentBrowser::GetHighlightedText)
-					   .ShowBottomToolbar(bShowBottomToolbar)
-					   .ShowViewOptions(false) // We control this for the main content browser
-					   .AllowThumbnailEditMode(true)
-					   .AllowThumbnailHintLabel(false)
-					   .CanShowFolders(Config != nullptr ? Config->bCanShowFolders : true)
-					   .CanShowClasses(Config != nullptr ? Config->bCanShowClasses : true)
-					   .CanShowRealTimeThumbnails(Config != nullptr ? Config->bCanShowRealTimeThumbnails : true)
-					   .CanShowDevelopersFolder(Config != nullptr ? Config->bCanShowDevelopersFolder : true)
-					   .CanShowFavorites(true)
-					   .CanDockCollections(true)
-					   .AddMetaData<FTagMetaData>(FTagMetaData(TEXT("ContentBrowserAssets")))
-					   .OwningContentBrowser(SharedThis(this))
-					   .OnSearchOptionsChanged(this, &SContentBrowser::HandleAssetViewSearchOptionsChanged)
-					   .bShowPathViewFilters(true)
-					   .FillEmptySpaceInTileView(true)
-					   .ShowDisallowedAssetClassAsUnsupportedItems(true);
+			.OnItemsActivated(this, &SContentBrowser::OnItemsActivated)
+			.OnGetItemContextMenu(this, &SContentBrowser::GetItemContextMenu, EContentBrowserViewContext::AssetView)
+			.OnItemRenameCommitted(this, &SContentBrowser::OnItemRenameCommitted)
+			.OnShouldFilterItem(bEnablePrivateContentFeature ? FOnShouldFilterItem::CreateSP(this, &SContentBrowser::HandlePrivateContentFilter) : FOnShouldFilterItem())
+			.FrontendFilters(FrontendFilters)
+			.TextFilter(TextFilter)
+			.ShowRedirectors_Lambda([this]() { return ContentBrowserUtils::ShouldShowRedirectors(FilterListPtr); })
+			.HighlightedText(this, &SContentBrowser::GetHighlightedText)
+			.ShowBottomToolbar(bShowBottomToolbar)
+			.ShowViewOptions(false) // We control this for the main content browser
+			.AllowThumbnailEditMode(true)
+			.AllowThumbnailHintLabel(false)
+			.CanShowFolders(Config != nullptr ? Config->bCanShowFolders : true)
+			.CanShowClasses(Config != nullptr ? Config->bCanShowClasses : true)
+			.CanShowRealTimeThumbnails(Config != nullptr ? Config->bCanShowRealTimeThumbnails : true)
+			.CanShowDevelopersFolder(Config != nullptr ? Config->bCanShowDevelopersFolder : true)
+			.CanShowFavorites(true)
+			.CanDockCollections(true)
+			.AddMetaData<FTagMetaData>(FTagMetaData(TEXT("ContentBrowserAssets")))
+			.OwningContentBrowser(SharedThis(this))
+			.OnSearchOptionsChanged(this, &SContentBrowser::HandleAssetViewSearchOptionsChanged)
+			.bShowPathViewFilters(true)
+			.FillEmptySpaceInTileView(true)
+			.ShowDisallowedAssetClassAsUnsupportedItems(true);
 
 	TSharedRef<SWidget> ViewOptions = SNullWidget::NullWidget;
 
@@ -1070,7 +1072,7 @@ TSharedRef<SWidget> SContentBrowser::CreateAssetView(const FContentBrowserConfig
 					.UseSharedSettings(true)
 					.CreateTextFilter(SFilterList::FCreateTextFilter::CreateLambda([this]
 					{
-						TSharedPtr<FFrontendFilter_CustomText> NewFilter = MakeShareable(new FFrontendFilter_CustomText);
+						TSharedPtr<FFrontendFilter_CustomText> NewFilter = MakeShared<FFrontendFilter_CustomText>();
 
 						// Make sure the new filter has the right search options from the AssetView. We only have to set it once, SFilterList handles syncing it on change
 						NewFilter->UpdateCustomTextFilterIncludes(AssetViewPtr->IsIncludingClassNames(), AssetViewPtr->IsIncludingAssetPaths(), AssetViewPtr->IsIncludingCollectionNames());
@@ -2414,16 +2416,14 @@ void SContentBrowser::SetSearchBoxText(const FText& InSearchText)
 	// Has anything changed? (need to test case as the operators are case-sensitive)
 	if (!InSearchText.ToString().Equals(TextFilter->GetRawFilterText().ToString(), ESearchCase::CaseSensitive))
 	{
-		TextFilter->SetRawFilterText( InSearchText );
-		SearchBoxPtr->SetError( TextFilter->GetFilterErrorText() );
-		if(InSearchText.IsEmpty())
+		TextFilter->SetRawFilterText(InSearchText);
+		SearchBoxPtr->SetError(TextFilter->GetFilterErrorText());
+		if (InSearchText.IsEmpty())
 		{
-			FrontendFilters->Remove(TextFilter);
 			AssetViewPtr->SetUserSearching(false);
 		}
 		else
 		{
-			FrontendFilters->Add(TextFilter);
 			AssetViewPtr->SetUserSearching(true);
 		}
 	}
diff --git a/Engine/Source/Editor/ContentBrowser/Private/SContentBrowser.h b/Engine/Source/Editor/ContentBrowser/Private/SContentBrowser.h
--- a/Engine/Source/Editor/ContentBrowser/Private/SContentBrowser.h
+++ b/Engine/Source/Editor/ContentBrowser/Private/SContentBrowser.h
@@ -30,7 +30,7 @@
 class FContentBrowserItemDataTemporaryContext;
 class FContentBrowserItemDataUpdate;
 class FExtender;
-class FFrontendFilter_Text;
+class FAssetTextFilter;
 class FSourcesSearch;
 class FTabManager;
 class FUICommandList;
@@ -417,7 +417,7 @@ private:
 	 * e.g. private content feature
 	 */
 	void OnConsoleVariableChanged();
-
+	
 	/** Handler for clicking the history back button */
 	FReply BackClicked();
 
@@ -735,8 +735,8 @@ private:
 	/** The list of FrontendFilters currently applied to the asset view */
 	TSharedPtr<FAssetFilterCollectionType> FrontendFilters;
 
-	/** The text filter to use on the assets */
-	TSharedPtr< FFrontendFilter_Text > TextFilter;
+	/** The text filter to use on the assets, separate from FrontendFilters so it can be specifically optimized */
+	TSharedPtr<FAssetTextFilter> TextFilter;
 
 	/** Commands handled by this widget */
 	TSharedPtr< FUICommandList > Commands;
diff --git a/Engine/Source/Editor/ContentBrowser/Private/SFilterList.cpp b/Engine/Source/Editor/ContentBrowser/Private/SFilterList.cpp
--- a/Engine/Source/Editor/ContentBrowser/Private/SFilterList.cpp
+++ b/Engine/Source/Editor/ContentBrowser/Private/SFilterList.cpp
@@ -756,6 +756,12 @@ void SFilterList::SetFilterLayout(EFilterBarLayout InFilterBarLayout)
 // FFilterListCustomTextFilter
 /////////////////////////////////////////
 
+FFrontendFilter_CustomText::FFrontendFilter_CustomText()
+	: FFrontendFilter(nullptr)
+{
+
+}
+
 /** Returns the system name for this filter */
 FString FFrontendFilter_CustomText::GetName() const
 {
@@ -769,7 +775,7 @@ FText FFrontendFilter_CustomText::GetDisplayName() const
 }
 FText FFrontendFilter_CustomText::GetToolTipText() const
 {
-	return GetRawFilterText();
+	return RawFilterText;
 }
 
 FLinearColor FFrontendFilter_CustomText::GetColor() const
@@ -779,16 +785,16 @@ FLinearColor FFrontendFilter_CustomText::GetColor() const
 
 void FFrontendFilter_CustomText::UpdateCustomTextFilterIncludes(const bool InIncludeClassName, const bool InIncludeAssetPath, const bool InIncludeCollectionNames)
 {
-	SetIncludeClassName(InIncludeClassName);
-	SetIncludeAssetPath(InIncludeAssetPath);
-	SetIncludeCollectionNames(InIncludeCollectionNames);
+	bIncludeClassName = InIncludeClassName;
+	bIncludeAssetPath = InIncludeAssetPath;
+	bIncludeCollectionNames = InIncludeCollectionNames;
 }
 
 void FFrontendFilter_CustomText::SetFromCustomTextFilterData(const FCustomTextFilterData& InFilterData)
 {
 	Color = InFilterData.FilterColor;
 	DisplayName = InFilterData.FilterLabel;
-	SetRawFilterText(InFilterData.FilterString);
+	RawFilterText = InFilterData.FilterString;
 }
 
 FCustomTextFilterData FFrontendFilter_CustomText::CreateCustomTextFilterData() const
@@ -797,7 +803,7 @@ FCustomTextFilterData FFrontendFilter_CustomText::CreateCustomTextFilterData() c
 
 	CustomTextFilterData.FilterColor = Color;
 	CustomTextFilterData.FilterLabel = DisplayName;
-	CustomTextFilterData.FilterString = GetRawFilterText();
+	CustomTextFilterData.FilterString = RawFilterText;
 
 	return CustomTextFilterData;
 }
@@ -807,4 +813,9 @@ TSharedPtr<FFilterBase<FAssetFilterType>> FFrontendFilter_CustomText::GetFilter(
 	return AsShared();
 }
 
+TOptional<FText> FFrontendFilter_CustomText::GetAsCustomTextFilter()
+{
+	return RawFilterText;
+}
+
 #undef LOCTEXT_NAMESPACE
diff --git a/Engine/Source/Editor/ContentBrowser/Private/SFilterList.h b/Engine/Source/Editor/ContentBrowser/Private/SFilterList.h
--- a/Engine/Source/Editor/ContentBrowser/Private/SFilterList.h
+++ b/Engine/Source/Editor/ContentBrowser/Private/SFilterList.h
@@ -225,20 +225,24 @@ private:
 
 /* A custom implementation of ICustomTextFilter that uses FFrontendFilter_Text to handle comparing items in the
  * Asset View to Custom Text Filters. This ensures that the advanced search syntax etc behaves properly when used
- * by a custom text filter created by saving a search
+ * by a custom text filter created by saving a search.
+ * This is implemented as a fake frontend filter which modifies background asynchronous text filtering behavior in the asset view.
  */
 class FFrontendFilter_CustomText :
 	public ICustomTextFilter<FAssetFilterType>,
-	public FFrontendFilter_Text,
+	public FFrontendFilter,
 	public TSharedFromThis<FFrontendFilter_CustomText>
 {
 public:
+	FFrontendFilter_CustomText();
 
 	// FFrontendFilter implementation
 	virtual FString GetName() const override;
 	virtual FText GetDisplayName() const override;
 	virtual FText GetToolTipText() const override;
 	virtual FLinearColor GetColor() const override;
+	virtual TOptional<FText> GetAsCustomTextFilter();
+	virtual bool PassesFilter(FAssetFilterType InItem) const override { return true; }
 	
 	/** Updates bIncludeClassName, bIncludeAssetPath and bIncludeCollectionNames for this filter */
 	void UpdateCustomTextFilterIncludes(const bool InIncludeClassName, const bool InIncludeAssetPath, const bool InIncludeCollectionNames);
@@ -259,6 +263,13 @@ protected:
 	/* The Display Name of this custom filter that the user sees */
 	FText DisplayName;
 
+	/** The raw text of the query, either simple string or advanced search syntax. */
+	FText RawFilterText;
+
 	/* The Color of this filter pill */
 	FLinearColor Color;
+
+	bool bIncludeClassName = false;
+	bool bIncludeAssetPath = false;
+	bool bIncludeCollectionNames = false;
 };
\ No newline at end of file
diff --git a/Engine/Source/Editor/ContentBrowser/Private/TextFilterKeyValueHandlers.cpp b/Engine/Source/Editor/ContentBrowser/Private/TextFilterKeyValueHandlers.cpp
--- a/Engine/Source/Editor/ContentBrowser/Private/TextFilterKeyValueHandlers.cpp
+++ b/Engine/Source/Editor/ContentBrowser/Private/TextFilterKeyValueHandlers.cpp
@@ -10,16 +10,5 @@ struct FContentBrowserItem;
 
 bool UTextFilterKeyValueHandlers::HandleTextFilterKeyValue(const FContentBrowserItem& InContentBrowserItem, const FName& InKey, const FTextFilterString& InValue, const ETextFilterComparisonOperation InComparisonOperation, const ETextFilterTextComparisonMode InTextComparisonMode, bool& bOutIsMatch)
 {
-	for (const FTextFilterKeyValueHandlerEntry& Handler : GetDefault<UTextFilterKeyValueHandlers>()->TextFilterKeyValueHandlers)
-	{
-		if (InKey.IsEqual(Handler.Key) || Handler.Key.IsNone())
-		{
-			if (UClass* TextFilterKeyValueHandlerClass = Handler.HandlerClass.LoadSynchronous())
-			{
-				return GetDefault<UTextFilterKeyValueHandler>(TextFilterKeyValueHandlerClass)->HandleTextFilterKeyValue(InContentBrowserItem, InKey, InValue, InComparisonOperation, InTextComparisonMode, bOutIsMatch);
-			}
-		}
-	}
-
 	return false;
 }
\ No newline at end of file
diff --git a/Engine/Source/Editor/ContentBrowser/Private/TextFilterKeyValueHandlers.h b/Engine/Source/Editor/ContentBrowser/Private/TextFilterKeyValueHandlers.h
--- a/Engine/Source/Editor/ContentBrowser/Private/TextFilterKeyValueHandlers.h
+++ b/Engine/Source/Editor/ContentBrowser/Private/TextFilterKeyValueHandlers.h
@@ -24,17 +24,22 @@ struct FTextFilterKeyValueHandlerEntry
 	UPROPERTY(config)
 	FName Key;
 
+PRAGMA_DISABLE_DEPRECATION_WARNINGS
 	UPROPERTY(config)
 	TSoftClassPtr<UTextFilterKeyValueHandler> HandlerClass;
+PRAGMA_ENABLE_DEPRECATION_WARNINGS
 };
 
 UCLASS(transient, config = Editor)
-class UTextFilterKeyValueHandlers : public UObject
+class UE_DEPRECATED(5.5, "This type has been replaced by IAssetTextFilterHandler which must be manually instantiated and registered and must be implemented in a threadsafe way.")
+UTextFilterKeyValueHandlers : public UObject
 {
 	GENERATED_BODY()
 public:
+PRAGMA_DISABLE_DEPRECATION_WARNINGS
 	UPROPERTY(config)
 	TArray<FTextFilterKeyValueHandlerEntry> TextFilterKeyValueHandlers;
+PRAGMA_ENABLE_DEPRECATION_WARNINGS
 
 	static bool HandleTextFilterKeyValue(const FContentBrowserItem& InContentBrowserItem, const FName& InKey, const FTextFilterString& InValue, const ETextFilterComparisonOperation InComparisonOperation, const ETextFilterTextComparisonMode InTextComparisonMode, bool& bOutIsMatch);
 };
\ No newline at end of file
diff --git a/Engine/Source/Editor/ContentBrowser/Private/TextFilterValueHandlers.cpp b/Engine/Source/Editor/ContentBrowser/Private/TextFilterValueHandlers.cpp
--- a/Engine/Source/Editor/ContentBrowser/Private/TextFilterValueHandlers.cpp
+++ b/Engine/Source/Editor/ContentBrowser/Private/TextFilterValueHandlers.cpp
@@ -10,16 +10,6 @@ struct FContentBrowserItem;
 
 bool UTextFilterValueHandlers::HandleTextFilterValue(const FContentBrowserItem& InContentBrowserItem, const FTextFilterString& InValue, const ETextFilterTextComparisonMode InTextComparisonMode, bool& bOutIsMatch)
 {
-	for (const TSoftClassPtr<UTextFilterValueHandler>& HandlerClass : GetDefault<UTextFilterValueHandlers>()->TextFilterValueHandlers)
-	{
-		if (UClass* TextFilterValueHandlerClass = HandlerClass.LoadSynchronous())
-		{
-			if (GetDefault<UTextFilterValueHandler>(TextFilterValueHandlerClass)->HandleTextFilterValue(InContentBrowserItem, InValue, InTextComparisonMode, bOutIsMatch))
-			{
-				return true;
-			}
-		}
-	}
 
 	return false;
 }
\ No newline at end of file
diff --git a/Engine/Source/Editor/ContentBrowser/Private/TextFilterValueHandlers.h b/Engine/Source/Editor/ContentBrowser/Private/TextFilterValueHandlers.h
--- a/Engine/Source/Editor/ContentBrowser/Private/TextFilterValueHandlers.h
+++ b/Engine/Source/Editor/ContentBrowser/Private/TextFilterValueHandlers.h
@@ -16,12 +16,15 @@ class UTextFilterValueHandler;
 struct FContentBrowserItem;
 
 UCLASS(transient, config = Editor)
-class UTextFilterValueHandlers : public UObject
+class UE_DEPRECATED(5.5, "This type has been replaced by IAssetTextFilterHandler which must be manually instantiated and registered and must be implemented in a threadsafe way.")
+UTextFilterValueHandlers : public UObject
 {
 	GENERATED_BODY()
 public:
+PRAGMA_DISABLE_DEPRECATION_WARNINGS
 	UPROPERTY(config)
 	TArray<TSoftClassPtr<UTextFilterValueHandler>> TextFilterValueHandlers;
+PRAGMA_ENABLE_DEPRECATION_WARNINGS
 
 	static bool HandleTextFilterValue(const FContentBrowserItem& InContentBrowserItem, const FTextFilterString& InValue, const ETextFilterTextComparisonMode InTextComparisonMode, bool& bOutIsMatch);
 };
\ No newline at end of file
diff --git a/Engine/Source/Editor/ContentBrowser/Public/AssetTextFilter.h b/Engine/Source/Editor/ContentBrowser/Public/AssetTextFilter.h
new file mode 100644
--- /dev/null
+++ b/Engine/Source/Editor/ContentBrowser/Public/AssetTextFilter.h
@@ -0,0 +1,253 @@
+// Copyright Epic Games, Inc. All Rights Reserved.
+
+#pragma once
+
+#include "IContentBrowserSingleton.h"
+
+#include "Misc/StringBuilder.h"
+#include "Misc/TextFilterExpressionEvaluator.h"
+#include "Misc/TextFilterUtils.h"
+
+class FAssetTextFilterContext;
+class ICollectionManager;
+class FCompiledAssetTextFilter;
+
+/** 
+ * Class handling text filtering for content browser views.
+ * Unlike FFrontendFilter, guarantees that filtering can be performed concurrently.
+ * @see FAssetTextFilterContext used to share the settings of this filter across multiple parallel workers.
+ */
+class FAssetTextFilter : public TSharedFromThis<FAssetTextFilter>
+{
+public:
+    CONTENTBROWSER_API FAssetTextFilter();
+ 	CONTENTBROWSER_API ~FAssetTextFilter();
+
+	// Cannot be moved or copied because of bound delegates.
+	FAssetTextFilter(FAssetTextFilter&& Other) = delete;
+	FAssetTextFilter& operator=(FAssetTextFilter&& Other) = delete;
+
+	FAssetTextFilter(const FAssetTextFilter& Other) = delete;
+	FAssetTextFilter& operator=(const FAssetTextFilter& Other) = delete;
+
+public:
+	/**
+	 * Create a compiled filter which can safely be used on other threads.
+	 * Combines the primary text filter with custom saved filters. 
+	 */
+	TSharedPtr<FCompiledAssetTextFilter> Compile();
+
+	/** Returns true if the filter contains no primary text and no custom saved queries */
+	bool IsEmpty() const;
+
+	/** Provides a set of saved filters/queries to be performed asynchronously alongside the main text filtering. */
+	void SetCustomTextFilters(TArray<FText> InQueries);
+
+	/** Returns the unsanitized and unsplit filter terms */
+	CONTENTBROWSER_API FText GetRawFilterText() const;
+
+	/** Set the Text to be used as the Filter's restrictions */
+	CONTENTBROWSER_API void SetRawFilterText(const FText& InFilterText);
+
+	/** Get the last error returned from lexing or compiling the current filter text */
+	CONTENTBROWSER_API FText GetFilterErrorText() const;
+
+	/** If bIncludeClassName is true, the text filter will include an asset's class name in the search */
+	CONTENTBROWSER_API void SetIncludeClassName(const bool InIncludeClassName);
+
+	/** If bIncludeAssetPath is true, the text filter will match against full Asset path */
+	CONTENTBROWSER_API void SetIncludeAssetPath(const bool InIncludeAssetPath);
+
+    /** Returns the last value set with SetIncludeAssetPath */
+	CONTENTBROWSER_API bool GetIncludeAssetPath() const;
+
+	/** If bIncludeCollectionNames is true, the text filter will match against collection names as well */
+	CONTENTBROWSER_API void SetIncludeCollectionNames(const bool InIncludeCollectionNames);
+
+    /** Returns the last value set with SetIncludeCollectionNames */
+	CONTENTBROWSER_API bool GetIncludeCollectionNames() const;
+
+	/** Delegate to bind for when the effective text filter changes so filtering can be re-run. */
+	CONTENTBROWSER_API FSimpleMulticastDelegate& OnChanged() { return ChangedEvent; }
+
+private:
+	/** Handles an on collection created event */
+	void HandleCollectionCreated(const FCollectionNameType& Collection);
+
+	/** Handles an on collection destroyed event */
+	void HandleCollectionDestroyed(const FCollectionNameType& Collection);
+
+	/** Handles an on collection renamed event */
+	void HandleCollectionRenamed(const FCollectionNameType& OriginalCollection, const FCollectionNameType& NewCollection);
+
+	/** Handles an on collection updated event */
+	void HandleCollectionUpdated(const FCollectionNameType& Collection);
+
+	void BroadcastChangedEvent() const { ChangedEvent.Broadcast(); }
+
+	/** An array of dynamic collections that are being referenced by the current query. These should be tested against each asset when it's looking for collections that contain it */
+	TArray<FCollectionNameType> ReferencedDynamicCollections;
+
+	/** 
+	 * Expression evaluator that can be used to perform complex text filter queries.
+	 * When CustomTextFilters is empty, this filter can be used as-is, otherwise it is necessary to compiled a combined evaluator.
+	 */
+	FTextFilterExpressionEvaluator TextFilterExpressionEvaluator;
+
+	/** Additional queries from saved user filters. These and any primary text filter are combined with AND semantics. */
+	TArray<FText> CustomTextFilters;
+
+	/** Delegate handles */
+	FDelegateHandle OnCollectionCreatedHandle;
+	FDelegateHandle OnCollectionDestroyedHandle;
+	FDelegateHandle OnCollectionRenamedHandle;
+	FDelegateHandle OnCollectionUpdatedHandle;
+    
+    // Filter options
+    bool bIncludeClassName = false;
+    bool bIncludeAssetPath = false;
+    bool bIncludeCollectionNames = false;
+
+	bool bReferencedDynamicCollectionsDirty = true;
+
+	FSimpleMulticastDelegate ChangedEvent;
+};
+
+
+/** 
+ * Context object for parallel filtering
+ * Allows reuse of memory when filtering in ParallelForWithTaskContext
+ */
+class FCompiledAssetTextFilter : private ITextFilterExpressionContext, public TSharedFromThis<FCompiledAssetTextFilter>
+{
+	typedef TRemoveReference<FAssetFilterType>::Type* FAssetFilterTypePtr;
+    friend class FAssetTextFilter;
+
+	struct FPrivateToken { explicit FPrivateToken() = default; };
+
+public:
+	FCompiledAssetTextFilter(
+		FPrivateToken,
+		TSharedRef<const FTextFilterExpressionEvaluator> InSharedEvaluator,
+		TSharedPtr<const TArray<FCollectionNameType>> InReferencedDynamicCollections,
+		bool InIncludeClassName,
+		bool InIncludeAssetPath,
+		bool InIncludeCollectionNames);
+    virtual ~FCompiledAssetTextFilter() {}
+
+    FCompiledAssetTextFilter(FCompiledAssetTextFilter&&) = default;
+    FCompiledAssetTextFilter& operator=( FCompiledAssetTextFilter&&) = default;
+
+	/** 
+	 * Clone a copy of this which shares filtering data but can be used on a different thread to the original. 
+	 * Each FCompiledAssetTextFilter should only be used on a single thread at once. 
+	 */
+	FCompiledAssetTextFilter CloneForThreading();
+
+	/**
+     * Non-const method to check an item against the filter, using member fields to amortize memory allocations for
+	 * many such checks 
+     */
+	bool PassesFilter(FAssetFilterType InItem);
+
+	/** Test the given value against the strings extracted from the current item */
+	virtual bool TestBasicStringExpression(
+		const FTextFilterString& InValue, const ETextFilterTextComparisonMode InTextComparisonMode) const override;
+
+	/** Perform a complex expression test for the current item */
+	virtual bool TestComplexExpression(
+		const FName& InKey,
+		const FTextFilterString& InValue,
+		const ETextFilterComparisonOperation InComparisonOperation,
+		const ETextFilterTextComparisonMode InTextComparisonMode) const override;
+private:
+    // Private copying for explicit clone operation
+    FCompiledAssetTextFilter(FCompiledAssetTextFilter&) = default;
+    FCompiledAssetTextFilter& operator=(const FCompiledAssetTextFilter&) = default;
+
+    /** Compiled text query */
+    TSharedRef<const FTextFilterExpressionEvaluator> Evaluator;
+
+	/** Dynamic collections referenced by the compiled query */
+	TSharedPtr<const TArray<FCollectionNameType>> ReferencedDynamicCollections;
+
+	/** Cached Collection manager */
+	ICollectionManager* CollectionManager;
+
+    // Filter options
+    bool bIncludeClassName = false;
+    bool bIncludeAssetPath = false;
+    bool bIncludeCollectionNames = false;
+    
+	// State for PassesFilter
+
+	/** Pointer to the asset we're currently filtering */
+	FAssetFilterTypePtr AssetPtr;
+    
+    /** Shared buffer for text keys to search */
+    FString TextBuffer;
+
+	/** Display name of the current asset */
+	FStringView AssetDisplayName;
+	/** Full path of the current asset */
+	FStringView AssetFullPath;
+	/** The export text name of the current asset */
+	FStringView AssetExportTextPath;
+
+	/** Names of the collections that the current asset is in */
+	TArray<FName> AssetCollectionNames;
+};
+
+/**
+ * Interface that can be implemented to extend text filtering in the content browser.
+ *
+ * Objects implementing this interface must be manually added and removed from participation in text filtering with
+ * RegisterHandler and UnregisterHandler to allow synchronization with async text filtering.
+ *
+ * Objects implementing this interface must be able to have HandleTextFilterValue and HandleTextFilterKeyValue called on
+ * any thread in between RegisterHandler and UnregisterHandler
+ */
+class IAssetTextFilterHandler
+{
+public:
+	CONTENTBROWSER_API virtual ~IAssetTextFilterHandler();
+
+	/**
+	 * Implement this function to handle basic text matching.
+	 *
+	 * @param InContentBrowserItem The item to be filtered.
+	 * @param InValue The text string the user wants to filter items against.
+	 * @param InTextComparisonMode How to compare text e.g. prefix/suffix/contains
+	 * @param bOutIsMatch Whether the item matches the user's query
+	 * @return true If this handler provided an answer for whether the item matches. bOutIsMatch may be true or false in this case.
+	 * @return false If this handler did not decide whether this item was filtered, and other handlers or the built-in implementation should continue.
+	 */
+	CONTENTBROWSER_API virtual bool HandleTextFilterValue(const FContentBrowserItem& InContentBrowserItem,
+		const FTextFilterString& InValue,
+		const ETextFilterTextComparisonMode InTextComparisonMode,
+		bool& bOutIsMatch) = 0;
+
+	/**
+	 * Implement this function to handle complex tests of keys against values (e.g. class=actor)
+	 *
+	 * @param InContentBrowserItem The item to be filtered.
+	 * @param InKey The key identifying what data from InContentBrowserItem to compare against.
+	 * @param InValue The user's provided text to compare the item against.
+	 * @param InComparisonOperation Comparison operator for numeric comparisons.
+	 * @param InTextComparisonMode How to compare text e.g. prefix/suffix/contains
+	 * @param bOutIsMatch Whether the item matches the user's query
+	 * @return true If this handler provided an answer for whether the item matches. bOutIsMatch may be true or false in this case.
+	 * @return false If this handler did not decide whether this item was filtered, and other handlers or the built-in implementation should continue.
+	 */
+	CONTENTBROWSER_API virtual bool HandleTextFilterKeyValue(const FContentBrowserItem& InContentBrowserItem,
+		const FName& InKey,
+		const FTextFilterString& InValue,
+		const ETextFilterComparisonOperation InComparisonOperation,
+		const ETextFilterTextComparisonMode InTextComparisonMode,
+		bool& bOutIsMatch) = 0;
+
+	/** Call on your derived instance of IAssetTextFilterHandler to have it participate in asset text filtering. */
+	CONTENTBROWSER_API void RegisterHandler();
+	/** Call on your derived instance of IAssetTextFilterHandler to remove it from asset text filtering. */
+	CONTENTBROWSER_API void UnregisterHandler();
+};
diff --git a/Engine/Source/Editor/ContentBrowser/Public/AssetViewTypes.h b/Engine/Source/Editor/ContentBrowser/Public/AssetViewTypes.h
deleted file mode 100644
diff --git a/Engine/Source/Editor/ContentBrowser/Public/FrontendFilterBase.h b/Engine/Source/Editor/ContentBrowser/Public/FrontendFilterBase.h
--- a/Engine/Source/Editor/ContentBrowser/Public/FrontendFilterBase.h
+++ b/Engine/Source/Editor/ContentBrowser/Public/FrontendFilterBase.h
@@ -24,6 +24,9 @@ public:
 	/** Invoke to set the source filter that is currently used to filter assets in the asset view */
 	virtual void SetCurrentFilter(TArrayView<const FName> InSourcePaths, const FContentBrowserDataFilter& InBaseFilter) { }
 	
+	/** If this filter is a saved custom text filter, return the text query it represents. */
+	virtual TOptional<FText> GetAsCustomTextFilter() { return TOptional<FText>(); }
+	
 	// FFilterBase Interface
 	
 	/** Returns the color this filter button will be when displayed as a button */
diff --git a/Engine/Source/Editor/ContentBrowser/Public/FrontendFilters.h b/Engine/Source/Editor/ContentBrowser/Public/FrontendFilters.h
--- a/Engine/Source/Editor/ContentBrowser/Public/FrontendFilters.h
+++ b/Engine/Source/Editor/ContentBrowser/Public/FrontendFilters.h
@@ -35,6 +35,7 @@ class FMenuBuilder;
 struct FAssetCompileData;
 
 /** A filter for text search */
+UE_DEPRECATED(5.5, "FFrontendFilter_Text has been deprecated in favor of FAssetTextFilter, used by SAssetView to perform text filtering in parallel.");
 class CONTENTBROWSER_API FFrontendFilter_Text : public FFrontendFilter
 {
 public:
diff --git a/Engine/Source/Editor/ContentBrowser/Public/SAssetView.h b/Engine/Source/Editor/ContentBrowser/Public/SAssetView.h
--- a/Engine/Source/Editor/ContentBrowser/Public/SAssetView.h
+++ b/Engine/Source/Editor/ContentBrowser/Public/SAssetView.h
@@ -33,6 +33,7 @@
 #include "Misc/Optional.h"
 #include "SourcesData.h"
 #include "Styling/SlateColor.h"
+#include "Templates/PimplPtr.h"
 #include "Templates/SharedPointer.h"
 #include "Templates/TypeHash.h"
 #include "Templates/UniquePtr.h"
@@ -45,8 +46,10 @@
 #include "Widgets/Views/STableRow.h"
 #include "Widgets/Views/STableViewBase.h"
 
+class FAssetTextFilter;
 class FAssetThumbnail;
 class FAssetViewItem;
+class FAssetViewItemCollection;
 class FContentBrowserItemDataTemporaryContext;
 class FContentBrowserItemDataUpdate;
 class FDragDropEvent;
@@ -182,6 +185,9 @@ public:
 		/** The filter collection used to further filter down assets returned from the backend */
 		SLATE_ARGUMENT( TSharedPtr<FAssetFilterCollectionType>, FrontendFilters )
 
+		/** Text filter object */
+		SLATE_ARGUMENT(TSharedPtr<FAssetTextFilter>, TextFilter)
+
 		/** If true, redirectors are visible even if not explicitly searching for them. */
 		SLATE_ATTRIBUTE(bool, ShowRedirectors);
 
@@ -294,6 +300,7 @@ public:
 		SLATE_ARGUMENT(TOptional<FName>, AssetViewOptionsProfile)
 	SLATE_END_ARGS()
 
+	SAssetView();
 	~SAssetView();
 
 	/** Constructs this widget with InArgs */
@@ -711,6 +718,7 @@ private:
 	/** Handler for column view widget creation */
 	TSharedRef<ITableRow> MakeColumnViewWidget(TSharedPtr<FAssetViewItem> AssetItem, const TSharedRef<STableViewBase>& OwnerTable);
 
+
 	/** Handler for when any asset item widget gets destroyed */
 	void AssetItemWidgetDestroyed(const TSharedPtr<FAssetViewItem>& Item);
 	
@@ -899,20 +907,8 @@ private:
 private:
 	friend class FAssetViewFrontendFilterHelper;
 
-	/** The available items from querying the backend data sources */
-	TMap<FContentBrowserItemKey, TSharedPtr<FAssetViewItem>> AvailableBackendItems;
-
-	/**
-	 * The items from AvailableBackendItems that are pending a run through any additional filtering before they can be shown in the filtered view list.
-	 * @note This filtering will run without amortization via ProcessItemsPendingFilter, so only use it for items that *must* be processed this frame.
-	 */
-	TSet<TSharedPtr<FAssetViewItem>> ItemsPendingPriorityFilter;
-
-	/**
-	 * The items from AvailableBackendItems that are pending a run through any additional frontend filters before they can be shown in the filtered view list.
-	 * @note This filtering will run amortized on the game thread via ProcessItemsPendingFilter.
-	 */
-	TSet<TSharedPtr<FAssetViewItem>> ItemsPendingFrontendFilter;
+	/** Private type managing data retrieved from the backend and async filtering thereof */
+	TPimplPtr<FAssetViewItemCollection> Items;
 
 	/** The items that are being shown in the filtered view list */
 	TArray<TSharedPtr<FAssetViewItem>> FilteredAssetItems;
@@ -943,6 +939,7 @@ private:
 	//  - not 'permissions' so may be ignore if e.g. the user explicitly selects a filtered folder
 	TArray<TSharedRef<const FPathPermissionList>> BackendCustomPathFilters;
 	TSharedPtr<FAssetFilterCollectionType> FrontendFilters;
+	TSharedPtr<FAssetTextFilter> TextFilter;
 
 	TAttribute<bool> bShowRedirectors;
 	bool bLastShowRedirectors;
@@ -1020,6 +1017,9 @@ private:
 	double SortDelaySeconds;
 
 	/** Weak ptr to the asset that is waiting to be renamed when scrolled into view, and the window is active */
+	TWeakPtr<FAssetViewItem> AwaitingScrollIntoViewForRename;
+
+	/** Weak ptr to the asset that is waiting to be renamed now that it has scrolled into view, if window is active */
 	TWeakPtr<FAssetViewItem> AwaitingRename;
 
 	/** Set when the user is in the process of naming an asset */
diff --git a/Engine/Source/Editor/ContentBrowser/Public/TextFilterKeyValueHandler.h b/Engine/Source/Editor/ContentBrowser/Public/TextFilterKeyValueHandler.h
--- a/Engine/Source/Editor/ContentBrowser/Public/TextFilterKeyValueHandler.h
+++ b/Engine/Source/Editor/ContentBrowser/Public/TextFilterKeyValueHandler.h
@@ -8,7 +8,8 @@
 #include "TextFilterKeyValueHandler.generated.h"
 
 UCLASS(abstract, transient, MinimalAPI)
-class UTextFilterKeyValueHandler : public UObject
+class UE_DEPRECATED(5.5, "This type has been replaced by IAssetTextFilterHandler which must be manually instantiated and registered and must be implemented in a threadsafe way.")
+UTextFilterKeyValueHandler : public UObject
 {
 	GENERATED_BODY()
 public:
diff --git a/Engine/Source/Editor/ContentBrowser/Public/TextFilterValueHandler.h b/Engine/Source/Editor/ContentBrowser/Public/TextFilterValueHandler.h
--- a/Engine/Source/Editor/ContentBrowser/Public/TextFilterValueHandler.h
+++ b/Engine/Source/Editor/ContentBrowser/Public/TextFilterValueHandler.h
@@ -8,7 +8,8 @@
 #include "TextFilterValueHandler.generated.h"
 
 UCLASS(abstract, transient, MinimalAPI)
-class UTextFilterValueHandler : public UObject
+class UE_DEPRECATED(5.5, "This type has been replaced by IAssetTextFilterHandler which must be manually instantiated and registered and must be implemented in a threadsafe way.")
+UTextFilterValueHandler : public UObject
 {
 	GENERATED_BODY()
 public:
diff --git a/Engine/Source/Editor/UMGEditor/Private/Library/SLibraryViewModel.cpp b/Engine/Source/Editor/UMGEditor/Private/Library/SLibraryViewModel.cpp
--- a/Engine/Source/Editor/UMGEditor/Private/Library/SLibraryViewModel.cpp
+++ b/Engine/Source/Editor/UMGEditor/Private/Library/SLibraryViewModel.cpp
@@ -8,6 +8,7 @@
 #include "Editor.h"
 #include "FrontendFilters.h"
 #include "SAssetView.h"
+#include "AssetTextFilter.h"
 
 #if WITH_EDITOR
 	#include "Styling/AppStyle.h"
@@ -118,11 +119,9 @@ void FWidgetTemplateListViewModel::ConstructListView(TArray<TSharedPtr<FWidgetTe
 
 		CachedLowercaseWidgetFilter = FString(FilterString.ToString()).ToLower();
 
-		SearchFilter = MakeShared<FFrontendFilter_Text>();
-		SearchFilter->SetActive(false);
+		SearchFilter = MakeShared<FAssetTextFilter>();
 
 		TemplatesFilter = MakeShared<FAssetFilterCollectionType>();
-		TemplatesFilter->Add(SearchFilter);
 	}
 
 	if (!AssetViewPtr)
@@ -132,6 +131,7 @@ void FWidgetTemplateListViewModel::ConstructListView(TArray<TSharedPtr<FWidgetTe
 			.InitialSourcesData(*LibrarySourceData)
 			.InitialThumbnailSize(EThumbnailSize::Small)
 			.FrontendFilters(TemplatesFilter)
+			.TextFilter(SearchFilter)
 			.ForceShowEngineContent(true)
 			.ForceShowPluginContent(true)
 			.ForceHideScrollbar(true)
@@ -240,7 +240,6 @@ void FWidgetTemplateListViewModel::SetSearchText(const FText& InSearchText)
 {
 	if (SearchFilter && AssetViewPtr)
 	{
-		SearchFilter->SetActive(!InSearchText.IsEmpty());
 		SearchFilter->SetRawFilterText(InSearchText);
 
 		AssetViewPtr->SetUserSearching(!InSearchText.IsEmpty());
diff --git a/Engine/Source/Editor/UMGEditor/Private/Library/SLibraryViewModel.h b/Engine/Source/Editor/UMGEditor/Private/Library/SLibraryViewModel.h
--- a/Engine/Source/Editor/UMGEditor/Private/Library/SLibraryViewModel.h
+++ b/Engine/Source/Editor/UMGEditor/Private/Library/SLibraryViewModel.h
@@ -50,7 +50,7 @@ private:
 	TSharedPtr<SAssetView> AssetViewPtr;
 
 	/** Filter we forward our search text to for the asset view */
-	TSharedPtr<class FFrontendFilter_Text> SearchFilter;
+	TSharedPtr<class FAssetTextFilter> SearchFilter;
 
 	/** Source data for the asset view */
 	TSharedPtr<struct FSourcesData> LibrarySourceData;
