Fix deadlock caused by re-entrancy from collection manager -> collections view -> collection manager attempting to acquire non-recursive lock.

diff --git a/Engine/Source/Editor/ContentBrowser/Private/SCollectionView.cpp b/Engine/Source/Editor/ContentBrowser/Private/SCollectionView.cpp
--- a/Engine/Source/Editor/ContentBrowser/Private/SCollectionView.cpp
+++ b/Engine/Source/Editor/ContentBrowser/Private/SCollectionView.cpp
@@ -322,13 +322,7 @@ void SCollectionView::HandleSourceControlProviderChanged(ISourceControlProvider&
 
 void SCollectionView::HandleSourceControlStateChanged()
 {
-	TRACE_CPUPROFILER_EVENT_SCOPE(SCollectionView::HandleSourceControlStateChanged);
-
-	// Update the status of each collection
-	for (const auto& AvailableCollectionInfo : AvailableCollections)
-	{
-		UpdateCollectionItemStatus(AvailableCollectionInfo.Value.ToSharedRef());
-	}
+	bQueueItemStatusUpdate = true;
 }
 
 void SCollectionView::UpdateCollectionItemStatus( const TSharedRef<FCollectionItem>& CollectionItem )
@@ -778,6 +772,18 @@ void SCollectionView::Tick( const FGeometry& AllottedGeometry, const double InCu
 			ISourceControlModule::Get().QueueStatusUpdate(CollectionFilesToRefresh);
 		}
 	}
+
+	if (bQueueItemStatusUpdate)
+	{
+		TRACE_CPUPROFILER_EVENT_SCOPE_STR("CollectionView Update Collection Items Status");
+
+		bQueueItemStatusUpdate = false;
+		// Update the status of each collection
+		for (const TPair<FCollectionNameType, TSharedPtr<FCollectionItem>>& AvailableCollectionInfo : AvailableCollections)
+		{
+			UpdateCollectionItemStatus(AvailableCollectionInfo.Value.ToSharedRef());
+		}
+	}
 }
 
 FReply SCollectionView::OnKeyDown( const FGeometry& MyGeometry, const FKeyEvent& InKeyEvent )
diff --git a/Engine/Source/Editor/ContentBrowser/Private/SCollectionView.h b/Engine/Source/Editor/ContentBrowser/Private/SCollectionView.h
--- a/Engine/Source/Editor/ContentBrowser/Private/SCollectionView.h
+++ b/Engine/Source/Editor/ContentBrowser/Private/SCollectionView.h
@@ -377,4 +377,7 @@ private:
 
 	/** True if we should queue an SCC refresh for the collections on the next Tick */
 	bool bQueueSCCRefresh;
+
+	/** True if we need to update the source control state of items on next tick */
+	bool bQueueItemStatusUpdate = false;
 };
