Handle removed items properly in UpdateData

diff --git a/Engine/Source/Editor/ContentBrowser/Private/SAssetView.cpp b/Engine/Source/Editor/ContentBrowser/Private/SAssetView.cpp
--- a/Engine/Source/Editor/ContentBrowser/Private/SAssetView.cpp
+++ b/Engine/Source/Editor/ContentBrowser/Private/SAssetView.cpp
@@ -582,7 +582,7 @@ TSharedPtr<FAssetViewItem> FAssetViewItemCollection::UpdateData(FContentBrowserI
 	int32 ExistingItemIndex = INDEX_NONE;
 	for (uint32 It = Lookup.First(Hash); Lookup.IsValid(It); It = Lookup.Next(It))
 	{
-		if (Items[It].IsValid() && ItemKey == FContentBrowserItemKey(Items[It]->GetItem()))
+		if (Items[It].IsValid() && !FilterState[It].Removed && ItemKey == FContentBrowserItemKey(Items[It]->GetItem()))
 		{
 			ExistingItemIndex = It;
 			break;
