Check removed flag in more Lookup searches.

diff --git a/Engine/Source/Editor/ContentBrowser/Private/SAssetView.cpp b/Engine/Source/Editor/ContentBrowser/Private/SAssetView.cpp
--- a/Engine/Source/Editor/ContentBrowser/Private/SAssetView.cpp
+++ b/Engine/Source/Editor/ContentBrowser/Private/SAssetView.cpp
@@ -549,7 +549,7 @@ TSharedPtr<FAssetViewItem> FAssetViewItemCollection::FindItemForRename(const FCo
 	FContentBrowserItemKey ItemKey(InItem);
 	for (uint32 It = Lookup.First(Hash); Lookup.IsValid(It); It = Lookup.Next(It))
 	{
-		if (Items[It].IsValid() && ItemKey == FContentBrowserItemKey(Items[It]->GetItem()))
+		if (Items[It].IsValid() && !FilterState[It].Removed && ItemKey == FContentBrowserItemKey(Items[It]->GetItem()))
 		{
 			checkf(FilterState[It].Published, 
 				TEXT("Only items which have been made visible in the UI should be available for renaming to maintain thread safety with async text filtering."));
@@ -623,7 +623,7 @@ TSharedPtr<FAssetViewItem> FAssetViewItemCollection::RemoveItemData(const FConte
 	FContentBrowserItemKey ItemKey(InItemData.GetItemType(), InItemData.GetVirtualPath(), InItemData.GetDataSource());
 	for (uint32 It = Lookup.First(Hash); Lookup.IsValid(It); It = Lookup.Next(It))
 	{
-		if (Items[It].IsValid() && ItemKey == FContentBrowserItemKey(Items[It]->GetItem()))
+		if (Items[It].IsValid() && !FilterState[It].Removed && ItemKey == FContentBrowserItemKey(Items[It]->GetItem()))
 		{
 			TSharedRef<FAssetViewItem> ItemToRemove = Items[It].ToSharedRef();
 
@@ -660,6 +660,7 @@ void FAssetViewItemCollection::RemoveItem(const TSharedPtr<FAssetViewItem>& ToRe
 	{
 		if (Items[It] == ToRemove)
 		{
+			check(!FilterState[It].Removed);
 		 	Lookup.Remove(Hash, It);
 			// This item was already filtered so we may want to remove it from the view.
 			if (It < (uint32)FrontendFilterProgress)
