Handle possible race condition between cancelling text filtering and spawning continuation.

diff --git a/Engine/Source/Editor/ContentBrowser/Private/SAssetView.cpp b/Engine/Source/Editor/ContentBrowser/Private/SAssetView.cpp
--- a/Engine/Source/Editor/ContentBrowser/Private/SAssetView.cpp
+++ b/Engine/Source/Editor/ContentBrowser/Private/SAssetView.cpp
@@ -713,11 +713,12 @@ void FAssetViewItemCollection::ResetFilterState()
 
 void FAssetViewItemCollection::AbortTextFiltering()
 {
-	if (TextFilterTask.IsValid())
+	ShouldCancelTextFiltering = true;
+	// Wait until the task sees the flag and doesn't spawn a continuation
+	while (TextFilterTask.IsValid())
 	{
-		ShouldCancelTextFiltering = true;
-		TextFilterTask.BusyWait();
-		TextFilterTask = {};
+		TextFilterTask.Wait();
+		TextFilterTask = TextFilterTask.GetResult().Next;
 	}
 }
 
@@ -779,7 +780,7 @@ FAssetViewItemCollection::FTextFilterResult FAssetViewItemCollection::AsyncFilte
 
 	int32 NextStartIndex = StartIndex + NumItemsToFilter;
 	// Need to check termination condition while holding the lock
-	bool bContinue = NextStartIndex < Items.Num() && UE::AssetView::AllowAsync; 
+	bool bContinue = NextStartIndex < Items.Num() && UE::AssetView::AllowAsync && !ShouldCancelTextFiltering; 
 	UE::Tasks::TTask<FTextFilterResult> NextTask;
 	if (bContinue)
 	{
@@ -2848,12 +2849,7 @@ void SAssetView::RefreshSourceItems()
 
 void FAssetViewItemCollection::RefreshItemsFromBackend(const FSourcesData& SourcesData, const FContentBrowserDataFilter& DataFilter, bool bAllowItemRecycling)
 {
-	if(TextFilterTask.IsValid())
-	{
-		ShouldCancelTextFiltering = true;
-		TextFilterTask.BusyWait();
-		TextFilterTask = {};
-	}
+	AbortTextFiltering();
 
 	UContentBrowserDataSubsystem* ContentBrowserData = IContentBrowserDataModule::Get().GetSubsystem();
 	TArray<FContentBrowserItemData> NewItemDatas;
