Fix accessing removed items when rebuilding lookup.

diff --git a/Engine/Source/Editor/ContentBrowser/Private/SAssetView.cpp b/Engine/Source/Editor/ContentBrowser/Private/SAssetView.cpp
--- a/Engine/Source/Editor/ContentBrowser/Private/SAssetView.cpp
+++ b/Engine/Source/Editor/ContentBrowser/Private/SAssetView.cpp
@@ -445,6 +445,11 @@ private:
 		return GetTypeHash(Item.GetVirtualPath());
 	}
 	 
+	inline bool IsItemValid(int32 ItemIndex)
+	{
+		return Items[ItemIndex].IsValid() && FilterState[ItemIndex].Removed == false;
+	}
+
 	struct FTextFilterResult
 	{
 		int32 StartIndex;
@@ -529,12 +534,12 @@ bool FAssetViewItemCollection::RefreshLookup()
 	{
 		Lookup.Clear(HashSize, Items.Max());
 		ParallelFor(TEXT("FAssetViewItemCollection::RefreshLookup"), Items.Num(), 16 * 1024, [this](int32 ItemIndex){
-			const TSharedPtr<FAssetViewItem> Item = Items[ItemIndex];
-			if (!Item.IsValid())
+			if (!IsItemValid(ItemIndex))
 			{
 				return;
 			}
 
+			const TSharedPtr<FAssetViewItem> Item = Items[ItemIndex];
 			uint32 Hash = HashItem(Item->GetItem());
 			Lookup.Add_Concurrent(Hash, ItemIndex);
 		}, UE::AssetView::AllowParallelism ? EParallelForFlags::None : EParallelForFlags::ForceSingleThread);
@@ -549,7 +554,7 @@ TSharedPtr<FAssetViewItem> FAssetViewItemCollection::FindItemForRename(const FCo
 	FContentBrowserItemKey ItemKey(InItem);
 	for (uint32 It = Lookup.First(Hash); Lookup.IsValid(It); It = Lookup.Next(It))
 	{
-		if (Items[It].IsValid() && !FilterState[It].Removed && ItemKey == FContentBrowserItemKey(Items[It]->GetItem()))
+		if (IsItemValid(It) && ItemKey == FContentBrowserItemKey(Items[It]->GetItem()))
 		{
 			checkf(FilterState[It].Published, 
 				TEXT("Only items which have been made visible in the UI should be available for renaming to maintain thread safety with async text filtering."));
@@ -582,7 +587,7 @@ TSharedPtr<FAssetViewItem> FAssetViewItemCollection::UpdateData(FContentBrowserI
 	int32 ExistingItemIndex = INDEX_NONE;
 	for (uint32 It = Lookup.First(Hash); Lookup.IsValid(It); It = Lookup.Next(It))
 	{
-		if (Items[It].IsValid() && !FilterState[It].Removed && ItemKey == FContentBrowserItemKey(Items[It]->GetItem()))
+		if (IsItemValid(It) && ItemKey == FContentBrowserItemKey(Items[It]->GetItem()))
 		{
 			ExistingItemIndex = It;
 			break;
@@ -623,7 +628,7 @@ TSharedPtr<FAssetViewItem> FAssetViewItemCollection::RemoveItemData(const FConte
 	FContentBrowserItemKey ItemKey(InItemData.GetItemType(), InItemData.GetVirtualPath(), InItemData.GetDataSource());
 	for (uint32 It = Lookup.First(Hash); Lookup.IsValid(It); It = Lookup.Next(It))
 	{
-		if (Items[It].IsValid() && !FilterState[It].Removed && ItemKey == FContentBrowserItemKey(Items[It]->GetItem()))
+		if (IsItemValid(It) && ItemKey == FContentBrowserItemKey(Items[It]->GetItem()))
 		{
 			TSharedRef<FAssetViewItem> ItemToRemove = Items[It].ToSharedRef();
 
