Remove unnecessary assertion from priority filtering.

diff --git a/Engine/Source/Editor/ContentBrowser/Private/SAssetView.cpp b/Engine/Source/Editor/ContentBrowser/Private/SAssetView.cpp
--- a/Engine/Source/Editor/ContentBrowser/Private/SAssetView.cpp
+++ b/Engine/Source/Editor/ContentBrowser/Private/SAssetView.cpp
@@ -946,14 +946,8 @@ bool FAssetViewItemCollection::PerformPriorityFiltering(FAssetViewFrontendFilter
 				continue;
 			}
 
-			const bool bPassedFrontendFilter = Helper.DoesItemPassFrontendFilter(Items[Index]);
-			if (bPassedFrontendFilter != FilterState[Index].PassedFrontendFilter)
-			{
-				// Only items which had previously passed filtering should be given for priority filtering, other items should
-				// be added at the end of the collection and processed as normal
-				check(FilterState[Index].PassedFrontendFilter);
-				FilterState[Index].PassedFrontendFilter = bPassedFrontendFilter;
-			}
+			// This may hide an item which was shown or show an item which was hidden - later we will check if we can remove without a re-sort, or if we need to add and re-sort
+			FilterState[Index].PassedFrontendFilter = Helper.DoesItemPassFrontendFilter(Items[Index]);
 		}
 
 		if (!bAllItemsPassedTextFilter)
