Add some checks against reusing an old asset item data multiple times
Change ItemsPendingPriorityFilter checks to use PublishProgress instead of FrontendFilterProgress.

diff --git a/Engine/Source/Editor/ContentBrowser/Private/AssetViewTypes.cpp b/Engine/Source/Editor/ContentBrowser/Private/AssetViewTypes.cpp
--- a/Engine/Source/Editor/ContentBrowser/Private/AssetViewTypes.cpp
+++ b/Engine/Source/Editor/ContentBrowser/Private/AssetViewTypes.cpp
@@ -11,34 +11,50 @@
 #include "Misc/AssertionMacros.h"
 #include "Templates/UnrealTemplate.h"
 
-FAssetViewItem::FAssetViewItem(FContentBrowserItem&& InItem)
+FAssetViewItem::FAssetViewItem(int32 InIndex)
+	: Index(InIndex)
+{
+}
+
+FAssetViewItem::FAssetViewItem(int32 InIndex, FContentBrowserItem&& InItem)
 	: Item(MoveTemp(InItem))
+	, Index(InIndex)
 {
 	checkf(Item.IsValid(), TEXT("FAssetViewItem was constructed from an invalid item!"));
 }
 
-FAssetViewItem::FAssetViewItem(const FContentBrowserItem& InItem)
+FAssetViewItem::FAssetViewItem(int32 InIndex, const FContentBrowserItem& InItem)
 	: Item(InItem)
+	, Index(InIndex)
 {
 	checkf(Item.IsValid(), TEXT("FAssetViewItem was constructed from an invalid item!"));
 }
 
-FAssetViewItem::FAssetViewItem(FContentBrowserItemData&& InItemData)
+FAssetViewItem::FAssetViewItem(int32 InIndex, FContentBrowserItemData&& InItemData)
 	: Item(MoveTemp(InItemData))
+	, Index(InIndex)
 {
 	checkf(Item.IsValid(), TEXT("FAssetViewItem was constructed from an invalid item!"));
 }
 
-FAssetViewItem::FAssetViewItem(const FContentBrowserItemData& InItemData)
+FAssetViewItem::FAssetViewItem(int32 InIndex, const FContentBrowserItemData& InItemData)
 	: Item(InItemData)
+	, Index(InIndex)
 {
 	checkf(Item.IsValid(), TEXT("FAssetViewItem was constructed from an invalid item!"));
 }
 
-void FAssetViewItem::ResetItemData(FContentBrowserItemData InItemData)
+void FAssetViewItem::ResetItemData(int32 OldIndex, int32 NewIndex, FContentBrowserItemData InItemData)
 {
-	Item = FContentBrowserItem{MoveTemp(InItemData)};
-	// Do not broadcast event here, it will be broadcast on the main thread after bulk building/recycling of items
+	if (Index.compare_exchange_strong(OldIndex, NewIndex, std::memory_order_relaxed))
+	{
+		Item = FContentBrowserItem{MoveTemp(InItemData)};
+		// Do not broadcast event here, it will be broadcast on the main thread after bulk building/recycling of items
+	}
+	else
+	{
+		checkf(false, TEXT("Concurrency issue detected recycling FAssetViewItem from old index %d to new index %d - already reassigned to"), OldIndex, NewIndex, Index.load());
+	}
 }
 
 void FAssetViewItem::BroadcastItemDataChanged()
diff --git a/Engine/Source/Editor/ContentBrowser/Private/AssetViewTypes.h b/Engine/Source/Editor/ContentBrowser/Private/AssetViewTypes.h
--- a/Engine/Source/Editor/ContentBrowser/Private/AssetViewTypes.h
+++ b/Engine/Source/Editor/ContentBrowser/Private/AssetViewTypes.h
@@ -11,6 +11,8 @@
 #include "UObject/NameTypes.h"
 #include "UObject/Object.h"
 
+#include <atomic>
+
 class FContentBrowserItemData;
 class FString;
 class FText;
@@ -20,16 +22,16 @@ struct FAssetViewCustomColumn;
 class FAssetViewItem
 {
 public:
-	FAssetViewItem() = default;
+	FAssetViewItem(int32 Index);
 
-	explicit FAssetViewItem(FContentBrowserItem&& InItem);
-	explicit FAssetViewItem(const FContentBrowserItem& InItem);
+	explicit FAssetViewItem(int32 Index, FContentBrowserItem&& InItem);
+	explicit FAssetViewItem(int32 Index, const FContentBrowserItem& InItem);
 
-	explicit FAssetViewItem(FContentBrowserItemData&& InItemData);
-	explicit FAssetViewItem(const FContentBrowserItemData& InItemData);
+	explicit FAssetViewItem(int32 Index, FContentBrowserItemData&& InItemData);
+	explicit FAssetViewItem(int32 Index, const FContentBrowserItemData& InItemData);
 
 	// When recycling an object, clear the item data and replace it with the given data.
-	void ResetItemData(FContentBrowserItemData InItemData);
+	void ResetItemData(int32 OldIndex, int32 Index, FContentBrowserItemData InItemData);
 
 	void AppendItemData(const FContentBrowserItem& InItem);
 
@@ -88,6 +90,12 @@ private:
 	/** Underlying Content Browser item data */
 	FContentBrowserItem Item;
 
+	/**
+	 * Index at which this is stored in the asset view's item collection.
+	 * Can be used to detect an item being added to the collection twice by mistake.
+	 */
+	std::atomic<int32> Index;
+
 	/** An event to fire when the data for this item changes */
 	FSimpleMulticastDelegate ItemDataChangedEvent;
 
diff --git a/Engine/Source/Editor/ContentBrowser/Private/SAssetView.cpp b/Engine/Source/Editor/ContentBrowser/Private/SAssetView.cpp
--- a/Engine/Source/Editor/ContentBrowser/Private/SAssetView.cpp
+++ b/Engine/Source/Editor/ContentBrowser/Private/SAssetView.cpp
@@ -414,7 +414,8 @@ private:
 	int32 CreateItem_Locked(FContentBrowserItem&& InItem)
 	{
 		uint32 Hash = HashItem(InItem);
-		TSharedPtr<FAssetViewItem> Item = Items.Add_GetRef(MakeShared<FAssetViewItem>(MoveTemp(InItem)));
+		TSharedPtr<FAssetViewItem> NewItem = MakeShared<FAssetViewItem>(Items.Num(), MoveTemp(InItem));
+		Items.Add(MoveTemp(NewItem));
 		FilterState.AddZeroed(1);
 		++NumValidItems;
 		if (!RefreshLookup())
@@ -608,10 +609,9 @@ TSharedPtr<FAssetViewItem> FAssetViewItemCollection::UpdateData(FContentBrowserI
 		ExistingItemIndex = CreateItem_Locked(MoveTemp(InData));
 	}
 
-	if (ExistingItemIndex < FrontendFilterProgress) 
+	if (ExistingItemIndex < PublishProgress) 
 	{
 		// This item was already filtered so we may want to remove it from the view or add it 
-		// We only check FrontendFilterProgress here as we only publish items that have passed both filters 
 		ItemsPendingPriorityFilter.Add(ExistingItemIndex);
 	}
 	return Items[ExistingItemIndex];
@@ -645,7 +645,7 @@ TSharedPtr<FAssetViewItem> FAssetViewItemCollection::RemoveItemData(const FConte
 			}
 
 			// This item was already filtered so we may want to remove it from the view.
-			if (It < (uint32)FrontendFilterProgress) 
+			if (It < (uint32)PublishProgress) 
 			{
 				ItemsPendingPriorityFilter.Add(It);
 			}
@@ -668,7 +668,7 @@ void FAssetViewItemCollection::RemoveItem(const TSharedPtr<FAssetViewItem>& ToRe
 			check(!FilterState[It].Removed);
 		 	Lookup.Remove(Hash, It);
 			// This item was already filtered so we may want to remove it from the view.
-			if (It < (uint32)FrontendFilterProgress)
+			if (It < (uint32)PublishProgress)
 			{
 				ItemsPendingPriorityFilter.Add(It);
 			}
@@ -889,14 +889,12 @@ void FAssetViewItemCollection::UpdateItemFiltering(
 	int32 CanPublish = FMath::Min(FrontendFilterProgress, bAllItemsPassedTextFilter ? Items.Num() : TextFilterProgress);
 	for (int32 i = PublishProgress; i < CanPublish; ++i)
 	{
-		if (!FilterState[i].PriorityFiltered)
+		const bool bPublish = ItemPassedAllFilters(i);
+		checkf(!FilterState[i].Published, TEXT("Standard-publish item %d was already published. PublishProgress: %d CanPublish: %d"), i, PublishProgress, CanPublish);
+		FilterState[i].Published = bPublish;
+		if (bPublish)
 		{
-			const bool bPublish = ItemPassedAllFilters(i);
-			FilterState[i].Published = bPublish;
-			if (bPublish)
-			{
-				OutItems.Add(Items[i]);
-			}
+			OutItems.Add(Items[i]);
 		}
 	}
 
@@ -905,7 +903,7 @@ void FAssetViewItemCollection::UpdateItemFiltering(
 	{
 		// Check we didn't update the state again to failure or removal
 		const bool bPublish = ItemPassedAllFilters(Index);
-		check(!FilterState[Index].Published);
+		checkf(!FilterState[Index].Published, TEXT("Priority-publish item %d was already published. PublishProgress: %d CanPublish: %d"), Index, PublishProgress, CanPublish);
 		FilterState[Index].Published = bPublish;
 		if (bPublish)
 		{
@@ -1556,7 +1554,7 @@ void SAssetView::BeginCreateDeferredItem()
 {
 	if (DeferredItemToCreate.IsValid() && !DeferredItemToCreate->bWasAddedToView)
 	{
-		TSharedPtr<FAssetViewItem> NewItem = MakeShared<FAssetViewItem>(DeferredItemToCreate->ItemContext.GetItem());
+		TSharedPtr<FAssetViewItem> NewItem = MakeShared<FAssetViewItem>(INDEX_NONE, DeferredItemToCreate->ItemContext.GetItem());
 		AwaitingScrollIntoViewForRename = NewItem;
 		DeferredItemToCreate->bWasAddedToView = true;
 
@@ -2932,17 +2930,18 @@ void FAssetViewItemCollection::RefreshItemsFromBackend(const FSourcesData& Sourc
 				FContentBrowserItemData& ItemData = NewItemDatas[Index];
 				FName VirtualPath = ItemData.GetVirtualPath();
 				TSharedPtr<FAssetViewItem> OldItem;
+				int32 OldItemIndex = INDEX_NONE;
 				
 				FContentBrowserItemKey ItemKey(ItemData);
 				uint32 Hash = HashItem(ItemData);
 				if (bAllowItemRecycling)
 				{
-					for (uint32 SearchIndex = OldLookup.First(Hash); OldLookup.IsValid(SearchIndex); SearchIndex = OldLookup.Next(SearchIndex))
+					for (OldItemIndex = OldLookup.First(Hash); OldLookup.IsValid(OldItemIndex); OldItemIndex = OldLookup.Next(OldItemIndex))
 					{
-						if (OldItems[SearchIndex].IsValid() && ItemKey == FContentBrowserItemKey(OldItems[SearchIndex]->GetItem()))
+						if (OldItems[OldItemIndex].IsValid() && ItemKey == FContentBrowserItemKey(OldItems[OldItemIndex]->GetItem()))
 						{
 							bAnyRecycled.store(true, std::memory_order_relaxed);
-							OldItem = OldItems[SearchIndex];
+							OldItem = OldItems[OldItemIndex];
 							break;
 						}
 					}
@@ -2955,12 +2954,12 @@ void FAssetViewItemCollection::RefreshItemsFromBackend(const FSourcesData& Sourc
 
 				if (OldItem.IsValid())
 				{
-					OldItem->ResetItemData(MoveTemp(ItemData));
+					OldItem->ResetItemData(OldItemIndex, Index, MoveTemp(ItemData));
 					Items[Index] = MoveTemp(OldItem);
 				}
 				else
 				{
-					Items[Index] = MakeShared<FAssetViewItem>(MoveTemp(ItemData));
+					Items[Index] = MakeShared<FAssetViewItem>(Index, MoveTemp(ItemData));
 				}
 				Lookup.Add_Concurrent(Hash, Index);
 			}, UE::AssetView::AllowParallelism ? EParallelForFlags::None : EParallelForFlags::ForceSingleThread);
