Fix use of compare_exchange_strong

diff --git a/Engine/Source/Editor/ContentBrowser/Private/AssetViewTypes.cpp b/Engine/Source/Editor/ContentBrowser/Private/AssetViewTypes.cpp
--- a/Engine/Source/Editor/ContentBrowser/Private/AssetViewTypes.cpp
+++ b/Engine/Source/Editor/ContentBrowser/Private/AssetViewTypes.cpp
@@ -46,14 +46,17 @@ FAssetViewItem::FAssetViewItem(int32 InIndex, const FContentBrowserItemData& InI
 
 void FAssetViewItem::ResetItemData(int32 OldIndex, int32 NewIndex, FContentBrowserItemData InItemData)
 {
-	if (Index.compare_exchange_strong(OldIndex, NewIndex, std::memory_order_relaxed))
+	int32 Expected = OldIndex;
+	if (Index.compare_exchange_strong(Expected, NewIndex, std::memory_order_relaxed))
 	{
 		Item = FContentBrowserItem{MoveTemp(InItemData)};
 		// Do not broadcast event here, it will be broadcast on the main thread after bulk building/recycling of items
 	}
 	else
 	{
-		checkf(false, TEXT("Concurrency issue detected recycling FAssetViewItem from old index %d to new index %d - already reassigned to"), OldIndex, NewIndex, Index.load());
+		checkf(false, TEXT("Concurrency issue detected recycling FAssetViewItem (%s) from old index %d to new index %d - already reassigned to %d"),
+			*WriteToString<256>(InItemData.GetVirtualPath()),
+			OldIndex, NewIndex, Expected);
 	}
 }
 
