Fix reusing old node for multiple new content browser items, possible racing on modifying the object fields.

diff --git a/Engine/Source/Editor/ContentBrowser/Private/SAssetView.cpp b/Engine/Source/Editor/ContentBrowser/Private/SAssetView.cpp
--- a/Engine/Source/Editor/ContentBrowser/Private/SAssetView.cpp
+++ b/Engine/Source/Editor/ContentBrowser/Private/SAssetView.cpp
@@ -2913,6 +2913,15 @@ void FAssetViewItemCollection::RefreshItemsFromBackend(const FSourcesData& Sourc
 
 	TArray<TSharedPtr<FAssetViewItem>> OldItems = MoveTemp(Items);
 	FHashTable OldLookup = MoveTemp(Lookup); 
+	TArray<FContentBrowserItemKey> OldItemKeys;
+	OldItemKeys.AddZeroed(OldItems.Num());
+	ParallelFor(TEXT("ExtractOldItemKeys"), OldItems.Num(), 16 * 1024, 
+		[&OldItems, &OldItemKeys](int32 Index) {
+			if (OldItems[Index].IsValid())
+			{
+				OldItemKeys[Index] = FContentBrowserItemKey(OldItems[Index]->GetItem());
+			}
+		});
 
 	Items.Reset(NewItemDatas.Num());
 	Items.AddZeroed(NewItemDatas.Num());
@@ -2924,12 +2933,13 @@ void FAssetViewItemCollection::RefreshItemsFromBackend(const FSourcesData& Sourc
 	std::atomic<bool> bAnyRecycled(false);
 	Lookup.Clear(HashSize, Items.Num());
 	{
+		// Used to handle multiple item data (folder) mapping to the same old item.
+		UE::FMutex OldItemMutex;
 		SCOPED_NAMED_EVENT(CreateItems, FColor::White);
 		ParallelFor(TEXT("CreateFAssetViewItem"), NewItemDatas.Num(), 16 * 1024,
-			[&NewItemDatas, &OldItems, &OldLookup, bAllowItemRecycling, &bAnyRecycled, &bAnyFolders, this](int32 Index) {
+			[&NewItemDatas, &OldItems, &OldLookup, &OldItemKeys, bAllowItemRecycling, &bAnyRecycled, &bAnyFolders, &OldItemMutex, this](int32 Index) {
 				FContentBrowserItemData& ItemData = NewItemDatas[Index];
 				FName VirtualPath = ItemData.GetVirtualPath();
-				TSharedPtr<FAssetViewItem> OldItem;
 				int32 OldItemIndex = INDEX_NONE;
 				
 				FContentBrowserItemKey ItemKey(ItemData);
@@ -2938,10 +2948,9 @@ void FAssetViewItemCollection::RefreshItemsFromBackend(const FSourcesData& Sourc
 				{
 					for (OldItemIndex = OldLookup.First(Hash); OldLookup.IsValid(OldItemIndex); OldItemIndex = OldLookup.Next(OldItemIndex))
 					{
-						if (OldItems[OldItemIndex].IsValid() && ItemKey == FContentBrowserItemKey(OldItems[OldItemIndex]->GetItem()))
+						if (ItemKey == OldItemKeys[OldItemIndex])
 						{
 							bAnyRecycled.store(true, std::memory_order_relaxed);
-							OldItem = OldItems[OldItemIndex];
 							break;
 						}
 					}
@@ -2952,15 +2961,27 @@ void FAssetViewItemCollection::RefreshItemsFromBackend(const FSourcesData& Sourc
 					bAnyFolders.store(true, std::memory_order_relaxed);
 				}
 
-				if (OldItem.IsValid())
-				{
-					OldItem->ResetItemData(OldItemIndex, Index, MoveTemp(ItemData));
-					Items[Index] = MoveTemp(OldItem);
-				}
-				else
+				if (OldLookup.IsValid(OldItemIndex))
 				{
-					Items[Index] = MakeShared<FAssetViewItem>(Index, MoveTemp(ItemData));
+					TSharedPtr<FAssetViewItem> OldItem;
+					// Try and acquire old item if another thread doesn't get there first (folder items share keys)
+					{
+						UE::TUniqueLock Lock(OldItemMutex);
+						OldItem = MoveTemp(OldItems[OldItemIndex]);
+						OldItems[OldItemIndex].Reset();
+					}
+
+					if (OldItem.IsValid())
+					{
+						OldItem->ResetItemData(OldItemIndex, Index, MoveTemp(ItemData));
+						Items[Index] = MoveTemp(OldItem);
+						Lookup.Add_Concurrent(Hash, Index);
+						return;
+					}
 				}
+
+				// Was not able to recycle an old item
+				Items[Index] = MakeShared<FAssetViewItem>(Index, MoveTemp(ItemData));
 				Lookup.Add_Concurrent(Hash, Index);
 			}, UE::AssetView::AllowParallelism ? EParallelForFlags::None : EParallelForFlags::ForceSingleThread);
 	}
