Fix shadowing error
#rnx

diff --git a/Engine/Source/Editor/ContentBrowser/Private/SAssetView.cpp b/Engine/Source/Editor/ContentBrowser/Private/SAssetView.cpp
--- a/Engine/Source/Editor/ContentBrowser/Private/SAssetView.cpp
+++ b/Engine/Source/Editor/ContentBrowser/Private/SAssetView.cpp
@@ -2966,7 +2966,7 @@ void FAssetViewItemCollection::RefreshItemsFromBackend(const FSourcesData& Sourc
 					TSharedPtr<FAssetViewItem> OldItem;
 					// Try and acquire old item if another thread doesn't get there first (folder items share keys)
 					{
-						UE::TUniqueLock Lock(OldItemMutex);
+						UE::TUniqueLock OldItemLock(OldItemMutex);
 						OldItem = MoveTemp(OldItems[OldItemIndex]);
 						OldItems[OldItemIndex].Reset();
 					}
