SAssetView: Replace checks by ensure to fix asset creation
#rnx

diff --git a/Engine/Source/Editor/ContentBrowser/Private/SAssetView.cpp b/Engine/Source/Editor/ContentBrowser/Private/SAssetView.cpp
--- a/Engine/Source/Editor/ContentBrowser/Private/SAssetView.cpp
+++ b/Engine/Source/Editor/ContentBrowser/Private/SAssetView.cpp
@@ -86,10 +86,10 @@
 namespace UE::AssetView
 {
 	/** Time delay between recently added items being added to the filtered asset items list */
-	const double TimeBetweenAddingNewAssets = 4.0;
+	constexpr double TimeBetweenAddingNewAssets = 4.0;
 
 	/** Time delay between performing the last jump, and the jump term being reset */
-	const double JumpDelaySeconds = 2.0;
+	constexpr double JumpDelaySeconds = 2.0;
 
 	bool AllowAsync = true;
 	FAutoConsoleVariableRef CVarAllowAsync(
@@ -889,25 +889,29 @@ void FAssetViewItemCollection::UpdateItemFiltering(
 	int32 CanPublish = FMath::Min(FrontendFilterProgress, bAllItemsPassedTextFilter ? Items.Num() : TextFilterProgress);
 	for (int32 i = PublishProgress; i < CanPublish; ++i)
 	{
-		const bool bPublish = ItemPassedAllFilters(i);
-		checkf(!FilterState[i].Published, TEXT("Standard-publish item %d was already published. PublishProgress: %d CanPublish: %d"), i, PublishProgress, CanPublish);
-		FilterState[i].Published = bPublish;
-		if (bPublish)
+		if (ensureMsgf(!FilterState[i].Published, TEXT("Standard-publish item %d was already published. PublishProgress: %d CanPublish: %d"), i, PublishProgress, CanPublish))
 		{
-			OutItems.Add(Items[i]);
+			const bool bPublish = ItemPassedAllFilters(i);
+			FilterState[i].Published = bPublish;
+			if (bPublish)
+			{
+				OutItems.Add(Items[i]);
+			}
 		}
 	}
 
 	// Publish items which originally failed filtering and then were updated to pass it 
 	for (int32 Index : ItemsPendingPriorityPublish)
 	{
-		// Check we didn't update the state again to failure or removal
-		const bool bPublish = ItemPassedAllFilters(Index);
-		checkf(!FilterState[Index].Published, TEXT("Priority-publish item %d was already published. PublishProgress: %d CanPublish: %d"), Index, PublishProgress, CanPublish);
-		FilterState[Index].Published = bPublish;
-		if (bPublish)
+		if (ensureMsgf(!FilterState[Index].Published, TEXT("Priority-publish item %d was already published. PublishProgress: %d CanPublish: %d"), Index, PublishProgress, CanPublish))
 		{
-			OutItems.Add(Items[Index]);
+			// Check we didn't update the state again to failure or removal
+			const bool bPublish = ItemPassedAllFilters(Index);
+			FilterState[Index].Published = bPublish;
+			if (bPublish)
+			{
+				OutItems.Add(Items[Index]);
+			}
 		}
 	}
 	ItemsPendingPriorityPublish.Reset();
