Restore skipping of early published items (e.g. new items) in UpdateItemFiltering

diff --git a/Engine/Source/Editor/ContentBrowser/Private/SAssetView.cpp b/Engine/Source/Editor/ContentBrowser/Private/SAssetView.cpp
--- a/Engine/Source/Editor/ContentBrowser/Private/SAssetView.cpp
+++ b/Engine/Source/Editor/ContentBrowser/Private/SAssetView.cpp
@@ -889,7 +889,7 @@ void FAssetViewItemCollection::UpdateItemFiltering(
 	int32 CanPublish = FMath::Min(FrontendFilterProgress, bAllItemsPassedTextFilter ? Items.Num() : TextFilterProgress);
 	for (int32 i = PublishProgress; i < CanPublish; ++i)
 	{
-		if (ensureMsgf(!FilterState[i].Published, TEXT("Standard-publish item %d was already published. PublishProgress: %d CanPublish: %d"), i, PublishProgress, CanPublish))
+		if (!FilterState[i].PriorityFiltered && ensureMsgf(!FilterState[i].Published, TEXT("Standard-publish item %d was already published. PublishProgress: %d CanPublish: %d"), i, PublishProgress, CanPublish))
 		{
 			const bool bPublish = ItemPassedAllFilters(i);
 			FilterState[i].Published = bPublish;
@@ -971,13 +971,14 @@ bool FAssetViewItemCollection::PerformPriorityFiltering(FAssetViewFrontendFilter
 		TSet<TSharedPtr<FAssetViewItem>> ToRemove;
 		for (int32 Index : ItemsPendingPriorityFilter)
 		{
-			FilterState[Index].PriorityFiltered = true;
 			if (Index >= PublishProgress)
 			{
 				// If item has yet to be published in the normal order, just leave the filter results for UpdateItemFiltering
 				continue;
 			}
 
+			// Only set flag if we're taking control of this item's publish state
+			FilterState[Index].PriorityFiltered = true;
 			const bool bPublish = !FilterState[Index].Removed && FilterState[Index].PassedFrontendFilter
 				&& (bAllItemsPassedTextFilter || FilterState[Index].PassedTextFilter);
 
