[ContentBrowser] Sources Panel in content browser should not resize itself

diff --git a/Engine/Source/Editor/ContentBrowser/Private/SContentBrowser.cpp b/Engine/Source/Editor/ContentBrowser/Private/SContentBrowser.cpp
--- a/Engine/Source/Editor/ContentBrowser/Private/SContentBrowser.cpp
+++ b/Engine/Source/Editor/ContentBrowser/Private/SContentBrowser.cpp
@@ -564,11 +564,14 @@ void SContentBrowser::Construct( const FArguments& InArgs, const FName& InInstan
 
 			// Sources View
 			+ SSplitter::Slot()
-			.Value(0.15f)
+			.Resizable(true)
+			.SizeRule(SSplitter::SizeToContent)
+			.OnSlotResized(this, &SContentBrowser::OnPathViewBoxColumnResized)
 			[
 				SNew(SBox)
 				.Padding(FMargin(4.f))
 				.Visibility(this, &SContentBrowser::GetSourcesViewVisibility)
+				.WidthOverride(this, &SContentBrowser::GetPathViewBoxWidthOverride)
 				[
 					SNew(SBorder)
 					.Padding(FMargin(0))
@@ -1544,6 +1547,16 @@ SSplitter::ESizeRule SContentBrowser::GetCollectionsAreaSizeRule() const
 	return CollectionArea->IsExpanded() && GetDockedCollectionsVisibility() == EVisibility::Visible ? SSplitter::ESizeRule::FractionOfParent : SSplitter::ESizeRule::SizeToContent;
 }
 
+void SContentBrowser::OnPathViewBoxColumnResized(float InSize)
+{
+	PathViewBoxWidth = InSize;
+}
+
+FOptionalSize SContentBrowser::GetPathViewBoxWidthOverride() const
+{
+	return FOptionalSize(PathViewBoxWidth);
+}
+
 float SContentBrowser::GetFavoritesAreaMinSize() const
 {
 	return GetFavoriteFolderVisibility() == EVisibility::Visible ? 29.0f : 0.0f;
diff --git a/Engine/Source/Editor/ContentBrowser/Private/SContentBrowser.h b/Engine/Source/Editor/ContentBrowser/Private/SContentBrowser.h
--- a/Engine/Source/Editor/ContentBrowser/Private/SContentBrowser.h
+++ b/Engine/Source/Editor/ContentBrowser/Private/SContentBrowser.h
@@ -612,6 +612,12 @@ private:
 	SSplitter::ESizeRule GetPathAreaSizeRule() const;
 	SSplitter::ESizeRule GetCollectionsAreaSizeRule() const;
 
+	/** Called when the Splitter containing the PathVew Box gets resized */
+	void OnPathViewBoxColumnResized(float InSize);
+
+	/** Returns the PathView Box width */
+	FOptionalSize GetPathViewBoxWidthOverride() const;
+
 	/** Gets the min size for various areas. When areas are not visible the min size is 0, otherwise there is a minimum size to prevent overlap */
 	float GetFavoritesAreaMinSize() const;
 	float GetCollectionsAreaMinSize() const;
@@ -723,6 +729,9 @@ private:
 	/** True if source should not be changed from an outside source */
 	bool bIsLocked = false;
 
+	/** Starting width of the PathView Box */
+	float PathViewBoxWidth = 150.f;
+
 	/** Cached result of CanWriteToPath to avoid recalculating it every frame */
 	mutable bool bCachedCanWriteToCurrentPath = false;
 	
