Fixed a crash when closing the editor in FFilter_ShowOtherDevelopers::~FFilter_ShowOtherDevelopers.

diff --git a/Engine/Source/Editor/ContentBrowser/Private/Filters.cpp b/Engine/Source/Editor/ContentBrowser/Private/Filters.cpp
--- a/Engine/Source/Editor/ContentBrowser/Private/Filters.cpp
+++ b/Engine/Source/Editor/ContentBrowser/Private/Filters.cpp
@@ -74,9 +74,14 @@ FFilter_ShowOtherDevelopers::FFilter_ShowOtherDevelopers(TSharedPtr<FFrontendFil
 
 FFilter_ShowOtherDevelopers::~FFilter_ShowOtherDevelopers()
 {
-	UContentBrowserDataSubsystem* ContentBrowserData = IContentBrowserDataModule::Get().GetSubsystem();
-	ContentBrowserData->OnItemDataUpdated().Remove(ItemDataUpdatedHandle);
-	ContentBrowserData->OnItemDataRefreshed().Remove(ItemDataRefreshedHandle);
+	if (IContentBrowserDataModule* ContentBrowserModule = IContentBrowserDataModule::GetPtr())
+	{
+		if (UContentBrowserDataSubsystem* ContentBrowserData = ContentBrowserModule->GetSubsystem())
+		{
+			ContentBrowserData->OnItemDataUpdated().Remove(ItemDataUpdatedHandle);
+			ContentBrowserData->OnItemDataRefreshed().Remove(ItemDataRefreshedHandle);
+		}
+	}
 }
 
 TSharedRef<const FPathPermissionList> FFilter_ShowOtherDevelopers::GetPathPermissionList() 
