[ContentBrowser] Filter Splitter slot in AssetView should not resize itself while resizing the whole CB panel

diff --git a/Engine/Source/Editor/ContentBrowser/Private/SContentBrowser.cpp b/Engine/Source/Editor/ContentBrowser/Private/SContentBrowser.cpp
--- a/Engine/Source/Editor/ContentBrowser/Private/SContentBrowser.cpp
+++ b/Engine/Source/Editor/ContentBrowser/Private/SContentBrowser.cpp
@@ -982,42 +982,47 @@ void SContentBrowser::OnFilterBarLayoutChanging(EFilterBarLayout NewLayout)
 
 			/* Filters in an SScrollBox */
 			+ SSplitter::Slot()
-			.Value(0.12f)
 			.MinSize(95.f)
+			.Resizable(true)
+			.SizeRule(SSplitter::SizeToContent)
+			.OnSlotResized(this, &SContentBrowser::OnFilterBoxColumnResized)
 			[
-				SNew(SVerticalBox)
+				SNew(SBox)
+				.WidthOverride(this, &SContentBrowser::GetFilterViewBoxWidthOverride)
 				// Don't take up space when there are no filters
 				.Visibility_Lambda([this]
 				{
 					return FilterListPtr->HasAnyFilters() ? EVisibility::Visible : EVisibility::Collapsed;
 				})
-
-				// Header
-				+ SVerticalBox::Slot()
-				.Padding(0.0f, 2.0f)
-				.AutoHeight()
 				[
-					
-					SNew(SBorder)
-					.BorderImage(FAppStyle::Get().GetBrush("Brushes.Header"))
-					.Padding(FMargin(8.0f, 6.0f))
-					.Content()
+					SNew(SVerticalBox)
+
+					// Header
+					+ SVerticalBox::Slot()
+					.Padding(0.0f, 2.0f)
+					.AutoHeight()
 					[
-						SNew(STextBlock)
-						.Text(LOCTEXT("FilterListVerticalHeader", "Filters"))
-						.TextStyle(FAppStyle::Get(), "ButtonText")
-						.Font(FAppStyle::Get().GetFontStyle("NormalFontBold"))
+						
+						SNew(SBorder)
+						.BorderImage(FAppStyle::Get().GetBrush("Brushes.Header"))
+						.Padding(FMargin(8.0f, 6.0f))
+						.Content()
+						[
+							SNew(STextBlock)
+							.Text(LOCTEXT("FilterListVerticalHeader", "Filters"))
+							.TextStyle(FAppStyle::Get(), "ButtonText")
+							.Font(FAppStyle::Get().GetFontStyle("NormalFontBold"))
+						]
 					]
-				]
 
-				// Filter List
-				+ SVerticalBox::Slot()
-				.Padding(0.0f)
-				.FillHeight(1.0f)
-				[
-					FilterListPtr.ToSharedRef()
+					// Filter List
+					+ SVerticalBox::Slot()
+					.Padding(0.0f)
+					.FillHeight(1.0f)
+					[
+						FilterListPtr.ToSharedRef()
+					]
 				]
-				
 			]
 		
 			+ SSplitter::Slot()
@@ -1557,6 +1562,16 @@ FOptionalSize SContentBrowser::GetPathViewBoxWidthOverride() const
 	return FOptionalSize(PathViewBoxWidth);
 }
 
+void SContentBrowser::OnFilterBoxColumnResized(float InSize)
+{
+	FilterBoxWidth = InSize;
+}
+
+FOptionalSize SContentBrowser::GetFilterViewBoxWidthOverride() const
+{
+	return FOptionalSize(FilterBoxWidth);
+}
+
 float SContentBrowser::GetFavoritesAreaMinSize() const
 {
 	return GetFavoriteFolderVisibility() == EVisibility::Visible ? 29.0f : 0.0f;
diff --git a/Engine/Source/Editor/ContentBrowser/Private/SContentBrowser.h b/Engine/Source/Editor/ContentBrowser/Private/SContentBrowser.h
--- a/Engine/Source/Editor/ContentBrowser/Private/SContentBrowser.h
+++ b/Engine/Source/Editor/ContentBrowser/Private/SContentBrowser.h
@@ -612,12 +612,18 @@ private:
 	SSplitter::ESizeRule GetPathAreaSizeRule() const;
 	SSplitter::ESizeRule GetCollectionsAreaSizeRule() const;
 
-	/** Called when the Splitter containing the PathVew Box gets resized */
+	/** Called when the Splitter containing the PathView Box gets resized */
 	void OnPathViewBoxColumnResized(float InSize);
 
 	/** Returns the PathView Box width */
 	FOptionalSize GetPathViewBoxWidthOverride() const;
 
+	/** Called when the Splitter containing the filters Box gets resized */
+	void OnFilterBoxColumnResized(float InSize);
+
+	/** Returns the filters Box width */
+	FOptionalSize GetFilterViewBoxWidthOverride() const;
+
 	/** Gets the min size for various areas. When areas are not visible the min size is 0, otherwise there is a minimum size to prevent overlap */
 	float GetFavoritesAreaMinSize() const;
 	float GetCollectionsAreaMinSize() const;
@@ -732,6 +738,9 @@ private:
 	/** Starting width of the PathView Box */
 	float PathViewBoxWidth = 150.f;
 
+	/** Starting width of the Filter Box */
+	float FilterBoxWidth = 100.f;
+
 	/** Cached result of CanWriteToPath to avoid recalculating it every frame */
 	mutable bool bCachedCanWriteToCurrentPath = false;
 	
