[ContentBrowser] PathView SplitterSlot should save its size when closing/re-opening
PathView SplitterSlot is resetting its size value when closed, should keep it saved.

diff --git a/Engine/Source/Editor/ContentBrowser/Private/SContentBrowser.cpp b/Engine/Source/Editor/ContentBrowser/Private/SContentBrowser.cpp
--- a/Engine/Source/Editor/ContentBrowser/Private/SContentBrowser.cpp
+++ b/Engine/Source/Editor/ContentBrowser/Private/SContentBrowser.cpp
@@ -1996,8 +1996,19 @@ void SContentBrowser::SaveSettings() const
 
 	for(int32 SlotIndex = 0; SlotIndex < PathAssetSplitterPtr->GetChildren()->Num(); SlotIndex++)
 	{
-		float SplitterSize = PathAssetSplitterPtr->SlotAt(SlotIndex).GetSizeValue();
-		GConfig->SetFloat(*SettingsIniSection, *(SettingsString + FString::Printf(TEXT(".VerticalSplitter.SlotSize%d"), SlotIndex)), SplitterSize, GEditorPerProjectIni);
+		// First Slot containing the PathView is using SizeToContent so the SizeValue is not updated
+		// Adding another config line for older projects otherwise when they open for the first time after this update the SplitterSlot size will be based on the older one
+		// The older one is a normalized value, so it is too small and will make the SplitterSlot for the PathView seems like if it was collapsed the very first time you re-open it
+		const bool bIsFirstSlot = SlotIndex == 0;
+		float SplitterSize = bIsFirstSlot ? PathViewBoxWidth : PathAssetSplitterPtr->SlotAt(SlotIndex).GetSizeValue();
+		if (bIsFirstSlot)
+		{
+			GConfig->SetFloat(*SettingsIniSection, *(SettingsString + FString::Printf(TEXT(".VerticalSplitter.FixedSlotSize%d"), SlotIndex)), SplitterSize, GEditorPerProjectIni);
+		}
+		else
+		{
+			GConfig->SetFloat(*SettingsIniSection, *(SettingsString + FString::Printf(TEXT(".VerticalSplitter.SlotSize%d"), SlotIndex)), SplitterSize, GEditorPerProjectIni);
+		}
 	}
 
 	for (int32 SlotIndex = 0; SlotIndex < PathFavoriteSplitterPtr->GetChildren()->Num(); SlotIndex++)
@@ -2147,9 +2158,21 @@ void SContentBrowser::LoadSettings(const FName& InInstanceName)
 
 	for(int32 SlotIndex = 0; SlotIndex < PathAssetSplitterPtr->GetChildren()->Num(); SlotIndex++)
 	{
-		float SplitterSize = PathAssetSplitterPtr->SlotAt(SlotIndex).GetSizeValue();
-		GConfig->GetFloat(*SettingsIniSection, *(SettingsString + FString::Printf(TEXT(".VerticalSplitter.SlotSize%d"), SlotIndex)), SplitterSize, GEditorPerProjectIni);
-		PathAssetSplitterPtr->SlotAt(SlotIndex).SetSizeValue(SplitterSize);
+		// First Slot containing the PathView is using SizeToContent so the SizeValue is not updated
+		// Adding another config line for older projects otherwise when they open for the first time after this update the SplitterSlot size will be based on the older one
+		// The older one is a normalized value, so it is too small and make will make the SplitterSlot for the PathView seems like if it was collapsed the very first time you re-open it
+		const bool bIsFirstSlot = SlotIndex == 0;
+		float SplitterSize = bIsFirstSlot ? PathViewBoxWidth : PathAssetSplitterPtr->SlotAt(SlotIndex).GetSizeValue();
+		if (bIsFirstSlot)
+		{
+			GConfig->GetFloat(*SettingsIniSection, *(SettingsString + FString::Printf(TEXT(".VerticalSplitter.FixedSlotSize%d"), SlotIndex)), SplitterSize, GEditorPerProjectIni);
+			PathViewBoxWidth = SplitterSize;
+		}
+		else
+		{
+			GConfig->GetFloat(*SettingsIniSection, *(SettingsString + FString::Printf(TEXT(".VerticalSplitter.SlotSize%d"), SlotIndex)), SplitterSize, GEditorPerProjectIni);
+			PathAssetSplitterPtr->SlotAt(SlotIndex).SetSizeValue(SplitterSize);
+		}
 	}
 
 	for (int32 SlotIndex = 0; SlotIndex < PathFavoriteSplitterPtr->GetChildren()->Num(); SlotIndex++)
