[ContentBrowser] When CTRL+C on an asset can we copy the Package Path instead of the Class Path

Added 2 new commmands to copy the ObjectPath and the PackageName.
Leaved the CTRL+C keybind mapped to the same behavior as before.
Added new commands to the keybind system so that users can re-bind them.
Cleaned up some hardcoded checks for keys in the OnKeyDown.

diff --git a/Engine/Plugins/Editor/ContentBrowser/ContentBrowserAliasDataSoure/Source/ContentBrowserAliasDataSource/Private/ContentBrowserAliasDataSource.cpp b/Engine/Plugins/Editor/ContentBrowser/ContentBrowserAliasDataSoure/Source/ContentBrowserAliasDataSource/Private/ContentBrowserAliasDataSource.cpp
--- a/Engine/Plugins/Editor/ContentBrowser/ContentBrowserAliasDataSoure/Source/ContentBrowserAliasDataSource/Private/ContentBrowserAliasDataSource.cpp
+++ b/Engine/Plugins/Editor/ContentBrowser/ContentBrowserAliasDataSoure/Source/ContentBrowserAliasDataSource/Private/ContentBrowserAliasDataSource.cpp
@@ -933,6 +933,16 @@ bool UContentBrowserAliasDataSource::AppendItemReference(const FContentBrowserIt
 	return ContentBrowserAssetData::AppendItemReference(AssetRegistry, this, InItem, InOutStr);
 }
 
+bool UContentBrowserAliasDataSource::AppendItemObjectPath(const FContentBrowserItemData& InItem, FString& InOutStr)
+{
+	return ContentBrowserAssetData::AppendItemObjectPath(AssetRegistry, this, InItem, InOutStr);
+}
+
+bool UContentBrowserAliasDataSource::AppendItemPackageName(const FContentBrowserItemData& InItem, FString& InOutStr)
+{
+	return ContentBrowserAssetData::AppendItemPackageName(AssetRegistry, this, InItem, InOutStr);
+}
+
 bool UContentBrowserAliasDataSource::UpdateThumbnail(const FContentBrowserItemData& InItem, FAssetThumbnail& InThumbnail)
 {
 	return ContentBrowserAssetData::UpdateItemThumbnail(this, InItem, InThumbnail);
diff --git a/Engine/Plugins/Editor/ContentBrowser/ContentBrowserAliasDataSoure/Source/ContentBrowserAliasDataSource/Public/ContentBrowserAliasDataSource.h b/Engine/Plugins/Editor/ContentBrowser/ContentBrowserAliasDataSoure/Source/ContentBrowserAliasDataSource/Public/ContentBrowserAliasDataSource.h
--- a/Engine/Plugins/Editor/ContentBrowser/ContentBrowserAliasDataSoure/Source/ContentBrowserAliasDataSource/Public/ContentBrowserAliasDataSource.h
+++ b/Engine/Plugins/Editor/ContentBrowser/ContentBrowserAliasDataSoure/Source/ContentBrowserAliasDataSource/Public/ContentBrowserAliasDataSource.h
@@ -89,6 +89,8 @@ public:
 	virtual bool BulkPrivatizeItems(TArrayView<const FContentBrowserItemData> InItems) override;
 
 	virtual bool AppendItemReference(const FContentBrowserItemData& InItem, FString& InOutStr) override;
+	virtual bool AppendItemObjectPath(const FContentBrowserItemData& InItem, FString& InOutStr) override;
+	virtual bool AppendItemPackageName(const FContentBrowserItemData& InItem, FString& InOutStr) override;
 	virtual bool UpdateThumbnail(const FContentBrowserItemData& InItem, FAssetThumbnail& InThumbnail) override;
 
 	virtual bool TryGetCollectionId(const FContentBrowserItemData& InItem, FSoftObjectPath& OutCollectionId) override;
diff --git a/Engine/Plugins/Editor/ContentBrowser/ContentBrowserAssetDataSource/Source/ContentBrowserAssetDataSource/Private/ContentBrowserAssetDataCore.cpp b/Engine/Plugins/Editor/ContentBrowser/ContentBrowserAssetDataSource/Source/ContentBrowserAssetDataSource/Private/ContentBrowserAssetDataCore.cpp
--- a/Engine/Plugins/Editor/ContentBrowser/ContentBrowserAssetDataSource/Source/ContentBrowserAssetDataSource/Private/ContentBrowserAssetDataCore.cpp
+++ b/Engine/Plugins/Editor/ContentBrowser/ContentBrowserAssetDataSource/Source/ContentBrowserAssetDataSource/Private/ContentBrowserAssetDataCore.cpp
@@ -1328,6 +1328,46 @@ bool AppendItemReference(IAssetRegistry* InAssetRegistry, const UContentBrowserD
 	return false;
 }
 
+bool AppendItemObjectPath(IAssetRegistry* InAssetRegistry, const UContentBrowserDataSource* InOwnerDataSource, const FContentBrowserItemData& InItem, FString& InOutStr)
+{
+	if (TSharedPtr<const FContentBrowserAssetFolderItemDataPayload> FolderPayload = GetAssetFolderItemPayload(InOwnerDataSource, InItem))
+	{
+		return AppendAssetFolderItemReference(InAssetRegistry, *FolderPayload, InOutStr);
+	}
+
+	if (TSharedPtr<const FContentBrowserAssetFileItemDataPayload> AssetPayload = GetAssetFileItemPayload(InOwnerDataSource, InItem))
+	{
+		return AppendAssetFileObjectPath(*AssetPayload, InOutStr);
+	}
+
+	if (TSharedPtr<const FContentBrowserUnsupportedAssetFileItemDataPayload> UnsupportedAssetPayload = GetUnsupportedAssetFileItemPayload(InOwnerDataSource, InItem))
+	{
+		return AppendUnsupportedAssetFileObjectPath(*UnsupportedAssetPayload, InOutStr);
+	}
+
+	return false;
+}
+
+bool AppendItemPackageName(IAssetRegistry* InAssetRegistry, const UContentBrowserDataSource* InOwnerDataSource, const FContentBrowserItemData& InItem, FString& InOutStr)
+{
+	if (TSharedPtr<const FContentBrowserAssetFolderItemDataPayload> FolderPayload = GetAssetFolderItemPayload(InOwnerDataSource, InItem))
+	{
+		return AppendAssetFolderItemReference(InAssetRegistry, *FolderPayload, InOutStr);
+	}
+
+	if (TSharedPtr<const FContentBrowserAssetFileItemDataPayload> AssetPayload = GetAssetFileItemPayload(InOwnerDataSource, InItem))
+	{
+		return AppendAssetFilePackageName(*AssetPayload, InOutStr);
+	}
+
+	if (TSharedPtr<const FContentBrowserUnsupportedAssetFileItemDataPayload> UnsupportedAssetPayload = GetUnsupportedAssetFileItemPayload(InOwnerDataSource, InItem))
+	{
+		return AppendUnsupportedAssetFilePackageName(*UnsupportedAssetPayload, InOutStr);
+	}
+
+	return false;
+}
+
 void AppendAssetExportText(const FAssetData& AssetData, FString& InOutStr)
 {
 	if (InOutStr.IsEmpty())
@@ -1344,6 +1384,39 @@ void AppendAssetExportText(const FAssetData& AssetData, FString& InOutStr)
 	}
 }
 
+void AppendAssetObjectPath(const FAssetData& AssetData, FString& InOutStr)
+{
+	if (InOutStr.IsEmpty())
+	{
+		InOutStr = AssetData.GetObjectPathString();
+	}
+	else
+	{
+		if (InOutStr.Len() > 0)
+		{
+			InOutStr += LINE_TERMINATOR;
+		}
+		InOutStr += AssetData.GetObjectPathString();
+	}
+}
+
+void AppendAssetPackageName(const FAssetData& AssetData, FString& InOutStr)
+{
+	UPackage* Package = AssetData.GetPackage();
+	if (InOutStr.IsEmpty())
+	{
+		InOutStr = Package->GetPathName();
+	}
+	else
+	{
+		if (InOutStr.Len() > 0)
+		{
+			InOutStr += LINE_TERMINATOR;
+		}
+		InOutStr += Package->GetPathName();
+	}
+}
+
 bool AppendAssetFolderItemReference(IAssetRegistry* InAssetRegistry, const FContentBrowserAssetFolderItemDataPayload& InFolderPayload, FString& InOutStr)
 {
 	// Folders gather the assets references from within them
@@ -1367,6 +1440,18 @@ bool AppendAssetFileItemReference(const FContentBrowserAssetFileItemDataPayload&
 	return true;
 }
 
+bool AppendAssetFileObjectPath(const FContentBrowserAssetFileItemDataPayload& InAssetPayload, FString& InOutStr)
+{
+	AppendAssetObjectPath(InAssetPayload.GetAssetData(), InOutStr);
+	return true;
+}
+
+bool AppendAssetFilePackageName(const FContentBrowserAssetFileItemDataPayload& InAssetPayload, FString& InOutStr)
+{
+	AppendAssetPackageName(InAssetPayload.GetAssetData(), InOutStr);
+	return true;
+}
+
 bool AppendUnsupportedAssetFileItemReference(const FContentBrowserUnsupportedAssetFileItemDataPayload& InUnsupportedAssetPayload, FString& InOutStr)
 {
 	if (const FAssetData* AssetData = InUnsupportedAssetPayload.GetAssetDataIfAvailable())
@@ -1378,6 +1463,28 @@ bool AppendUnsupportedAssetFileItemReference(const FContentBrowserUnsupportedAss
 	return false;
 }
 
+bool AppendUnsupportedAssetFileObjectPath(const FContentBrowserUnsupportedAssetFileItemDataPayload& InUnsupportedAssetPayload, FString& InOutStr)
+{
+	if (const FAssetData* AssetData = InUnsupportedAssetPayload.GetAssetDataIfAvailable())
+	{
+		AppendAssetObjectPath(*AssetData, InOutStr);
+		return true;
+	}
+
+	return false;
+}
+
+bool AppendUnsupportedAssetFilePackageName(const FContentBrowserUnsupportedAssetFileItemDataPayload& InUnsupportedAssetPayload, FString& InOutStr)
+{
+	if (const FAssetData* AssetData = InUnsupportedAssetPayload.GetAssetDataIfAvailable())
+	{
+		AppendAssetPackageName(*AssetData, InOutStr);
+		return true;
+	}
+
+	return false;
+}
+
 bool GetItemPhysicalPath(const UContentBrowserDataSource* InOwnerDataSource, const FContentBrowserItemData& InItem, FString& OutDiskPath)
 {
 	if (TSharedPtr<const FContentBrowserAssetFolderItemDataPayload> FolderPayload = GetAssetFolderItemPayload(InOwnerDataSource, InItem))
diff --git a/Engine/Plugins/Editor/ContentBrowser/ContentBrowserAssetDataSource/Source/ContentBrowserAssetDataSource/Private/ContentBrowserAssetDataSource.cpp b/Engine/Plugins/Editor/ContentBrowser/ContentBrowserAssetDataSource/Source/ContentBrowserAssetDataSource/Private/ContentBrowserAssetDataSource.cpp
--- a/Engine/Plugins/Editor/ContentBrowser/ContentBrowserAssetDataSource/Source/ContentBrowserAssetDataSource/Private/ContentBrowserAssetDataSource.cpp
+++ b/Engine/Plugins/Editor/ContentBrowser/ContentBrowserAssetDataSource/Source/ContentBrowserAssetDataSource/Private/ContentBrowserAssetDataSource.cpp
@@ -2880,6 +2880,16 @@ bool UContentBrowserAssetDataSource::AppendItemReference(const FContentBrowserIt
 	return ContentBrowserAssetData::AppendItemReference(AssetRegistry, this, InItem, InOutStr);
 }
 
+bool UContentBrowserAssetDataSource::AppendItemObjectPath(const FContentBrowserItemData& InItem, FString& InOutStr)
+{
+	return ContentBrowserAssetData::AppendItemObjectPath(AssetRegistry, this, InItem, InOutStr);
+}
+
+bool UContentBrowserAssetDataSource::AppendItemPackageName(const FContentBrowserItemData& InItem, FString& InOutStr)
+{
+	return ContentBrowserAssetData::AppendItemPackageName(AssetRegistry, this, InItem, InOutStr);
+}
+
 bool UContentBrowserAssetDataSource::UpdateThumbnail(const FContentBrowserItemData& InItem, FAssetThumbnail& InThumbnail)
 {
 	return ContentBrowserAssetData::UpdateItemThumbnail(this, InItem, InThumbnail);
diff --git a/Engine/Plugins/Editor/ContentBrowser/ContentBrowserAssetDataSource/Source/ContentBrowserAssetDataSource/Public/ContentBrowserAssetDataCore.h b/Engine/Plugins/Editor/ContentBrowser/ContentBrowserAssetDataSource/Source/ContentBrowserAssetDataSource/Public/ContentBrowserAssetDataCore.h
--- a/Engine/Plugins/Editor/ContentBrowser/ContentBrowserAssetDataSource/Source/ContentBrowserAssetDataSource/Public/ContentBrowserAssetDataCore.h
+++ b/Engine/Plugins/Editor/ContentBrowser/ContentBrowserAssetDataSource/Source/ContentBrowserAssetDataSource/Public/ContentBrowserAssetDataCore.h
@@ -170,12 +170,24 @@ CONTENTBROWSERASSETDATASOURCE_API FContentBrowserItemData CreateUnsupportedAsset
 
 	CONTENTBROWSERASSETDATASOURCE_API bool AppendItemReference(IAssetRegistry* InAssetRegistry, const UContentBrowserDataSource* InOwnerDataSource, const FContentBrowserItemData& InItem, FString& InOutStr);
 
+	CONTENTBROWSERASSETDATASOURCE_API bool AppendItemObjectPath(IAssetRegistry* InAssetRegistry, const UContentBrowserDataSource* InOwnerDataSource, const FContentBrowserItemData& InItem, FString& InOutStr);
+
+	CONTENTBROWSERASSETDATASOURCE_API bool AppendItemPackageName(IAssetRegistry* InAssetRegistry, const UContentBrowserDataSource* InOwnerDataSource, const FContentBrowserItemData& InItem, FString& InOutStr);
+
 	CONTENTBROWSERASSETDATASOURCE_API bool AppendAssetFolderItemReference(IAssetRegistry* InAssetRegistry, const FContentBrowserAssetFolderItemDataPayload& InFolderPayload, FString& InOutStr);
 
 	CONTENTBROWSERASSETDATASOURCE_API bool AppendAssetFileItemReference(const FContentBrowserAssetFileItemDataPayload& InAssetPayload, FString& InOutStr);
 
+	CONTENTBROWSERASSETDATASOURCE_API bool AppendAssetFileObjectPath(const FContentBrowserAssetFileItemDataPayload& InAssetPayload, FString& InOutStr);
+
+	CONTENTBROWSERASSETDATASOURCE_API bool AppendAssetFilePackageName(const FContentBrowserAssetFileItemDataPayload& InAssetPayload, FString& InOutStr);
+
 	CONTENTBROWSERASSETDATASOURCE_API bool AppendUnsupportedAssetFileItemReference(const FContentBrowserUnsupportedAssetFileItemDataPayload& InUnsupportedAssetPayload, FString& InOutStr);
 
+	CONTENTBROWSERASSETDATASOURCE_API bool AppendUnsupportedAssetFileObjectPath(const FContentBrowserUnsupportedAssetFileItemDataPayload& InUnsupportedAssetPayload, FString& InOutStr);
+
+	CONTENTBROWSERASSETDATASOURCE_API bool AppendUnsupportedAssetFilePackageName(const FContentBrowserUnsupportedAssetFileItemDataPayload& InUnsupportedAssetPayload, FString& InOutStr);
+
 	CONTENTBROWSERASSETDATASOURCE_API bool GetItemPhysicalPath(const UContentBrowserDataSource* InOwnerDataSource, const FContentBrowserItemData& InItem, FString& OutDiskPath);
 
 	CONTENTBROWSERASSETDATASOURCE_API bool GetAssetFolderItemPhysicalPath(const FContentBrowserAssetFolderItemDataPayload& InFolderPayload, FString& OutDiskPath);
diff --git a/Engine/Plugins/Editor/ContentBrowser/ContentBrowserAssetDataSource/Source/ContentBrowserAssetDataSource/Public/ContentBrowserAssetDataSource.h b/Engine/Plugins/Editor/ContentBrowser/ContentBrowserAssetDataSource/Source/ContentBrowserAssetDataSource/Public/ContentBrowserAssetDataSource.h
--- a/Engine/Plugins/Editor/ContentBrowser/ContentBrowserAssetDataSource/Source/ContentBrowserAssetDataSource/Public/ContentBrowserAssetDataSource.h
+++ b/Engine/Plugins/Editor/ContentBrowser/ContentBrowserAssetDataSource/Source/ContentBrowserAssetDataSource/Public/ContentBrowserAssetDataSource.h
@@ -327,6 +327,10 @@ public:
 
 	virtual bool AppendItemReference(const FContentBrowserItemData& InItem, FString& InOutStr) override;
 
+	virtual bool AppendItemObjectPath(const FContentBrowserItemData& InItem, FString& InOutStr) override;
+
+	virtual bool AppendItemPackageName(const FContentBrowserItemData& InItem, FString& InOutStr) override;
+
 	virtual bool UpdateThumbnail(const FContentBrowserItemData& InItem, FAssetThumbnail& InThumbnail) override;
 
 	virtual bool HandleDragEnterItem(const FContentBrowserItemData& InItem, const FDragDropEvent& InDragDropEvent) override;
diff --git a/Engine/Plugins/Editor/ContentBrowser/ContentBrowserClassDataSource/Source/ContentBrowserClassDataSource/Private/ContentBrowserClassDataCore.cpp b/Engine/Plugins/Editor/ContentBrowser/ContentBrowserClassDataSource/Source/ContentBrowserClassDataSource/Private/ContentBrowserClassDataCore.cpp
--- a/Engine/Plugins/Editor/ContentBrowser/ContentBrowserClassDataSource/Source/ContentBrowserClassDataSource/Private/ContentBrowserClassDataCore.cpp
+++ b/Engine/Plugins/Editor/ContentBrowser/ContentBrowserClassDataSource/Source/ContentBrowserClassDataSource/Private/ContentBrowserClassDataCore.cpp
@@ -298,6 +298,26 @@ bool AppendItemReference(const UContentBrowserDataSource* InOwnerDataSource, con
 	return false;
 }
 
+bool AppendItemObjectPath(const UContentBrowserDataSource* InOwnerDataSource, const FContentBrowserItemData& InItem, FString& InOutStr)
+{
+	if (TSharedPtr<const FContentBrowserClassFileItemDataPayload> ClassPayload = GetClassFileItemPayload(InOwnerDataSource, InItem))
+	{
+		return AppendObjectPathFileItemReference(*ClassPayload, InOutStr);
+	}
+
+	return false;
+}
+
+bool AppendItemPackageName(const UContentBrowserDataSource* InOwnerDataSource, const FContentBrowserItemData& InItem, FString& InOutStr)
+{
+	if (TSharedPtr<const FContentBrowserClassFileItemDataPayload> ClassPayload = GetClassFileItemPayload(InOwnerDataSource, InItem))
+	{
+		return AppendPackageNameItemReference(*ClassPayload, InOutStr);
+	}
+
+	return false;
+}
+
 bool AppendClassFileItemReference(const FContentBrowserClassFileItemDataPayload& InClassPayload, FString& InOutStr)
 {
 	if (InOutStr.Len() > 0)
@@ -308,6 +328,30 @@ bool AppendClassFileItemReference(const FContentBrowserClassFileItemDataPayload&
 	return true;
 }
 
+bool AppendObjectPathFileItemReference(const FContentBrowserClassFileItemDataPayload& InClassPayload, FString& InOutStr)
+{
+	if (InOutStr.Len() > 0)
+	{
+		InOutStr += LINE_TERMINATOR;
+	}
+	InOutStr += InClassPayload.GetAssetData().GetObjectPathString();
+	return true;
+}
+
+bool AppendPackageNameItemReference(const FContentBrowserClassFileItemDataPayload& InClassPayload, FString& InOutStr)
+{
+	if (const UPackage* Package = InClassPayload.GetAssetData().GetPackage())
+	{
+		if (InOutStr.Len() > 0)
+		{
+			InOutStr += LINE_TERMINATOR;
+		}
+		InOutStr += Package->GetPathName();
+		return true;
+	}
+	return false;
+}
+
 bool GetItemPhysicalPath(const UContentBrowserDataSource* InOwnerDataSource, const FContentBrowserItemData& InItem, FString& OutDiskPath)
 {
 	if (TSharedPtr<const FContentBrowserClassFolderItemDataPayload> FolderPayload = GetClassFolderItemPayload(InOwnerDataSource, InItem))
diff --git a/Engine/Plugins/Editor/ContentBrowser/ContentBrowserClassDataSource/Source/ContentBrowserClassDataSource/Private/ContentBrowserClassDataSource.cpp b/Engine/Plugins/Editor/ContentBrowser/ContentBrowserClassDataSource/Source/ContentBrowserClassDataSource/Private/ContentBrowserClassDataSource.cpp
--- a/Engine/Plugins/Editor/ContentBrowser/ContentBrowserClassDataSource/Source/ContentBrowserClassDataSource/Private/ContentBrowserClassDataSource.cpp
+++ b/Engine/Plugins/Editor/ContentBrowser/ContentBrowserClassDataSource/Source/ContentBrowserClassDataSource/Private/ContentBrowserClassDataSource.cpp
@@ -535,6 +535,16 @@ bool UContentBrowserClassDataSource::AppendItemReference(const FContentBrowserIt
 	return ContentBrowserClassData::AppendItemReference(this, InItem, InOutStr);
 }
 
+bool UContentBrowserClassDataSource::AppendItemObjectPath(const FContentBrowserItemData& InItem, FString& InOutStr)
+{
+	return ContentBrowserClassData::AppendItemObjectPath(this, InItem, InOutStr);
+}
+
+bool UContentBrowserClassDataSource::AppendItemPackageName(const FContentBrowserItemData& InItem, FString& InOutStr)
+{
+	return ContentBrowserClassData::AppendItemPackageName(this, InItem, InOutStr);
+}
+
 bool UContentBrowserClassDataSource::UpdateThumbnail(const FContentBrowserItemData& InItem, FAssetThumbnail& InThumbnail)
 {
 	return ContentBrowserClassData::UpdateItemThumbnail(this, InItem, InThumbnail);
diff --git a/Engine/Plugins/Editor/ContentBrowser/ContentBrowserClassDataSource/Source/ContentBrowserClassDataSource/Public/ContentBrowserClassDataCore.h b/Engine/Plugins/Editor/ContentBrowser/ContentBrowserClassDataSource/Source/ContentBrowserClassDataSource/Public/ContentBrowserClassDataCore.h
--- a/Engine/Plugins/Editor/ContentBrowser/ContentBrowserClassDataSource/Source/ContentBrowserClassDataSource/Public/ContentBrowserClassDataCore.h
+++ b/Engine/Plugins/Editor/ContentBrowser/ContentBrowserClassDataSource/Source/ContentBrowserClassDataSource/Public/ContentBrowserClassDataCore.h
@@ -50,8 +50,16 @@ namespace ContentBrowserClassData
 
 	CONTENTBROWSERCLASSDATASOURCE_API bool AppendItemReference(const UContentBrowserDataSource* InOwnerDataSource, const FContentBrowserItemData& InItem, FString& InOutStr);
 
+	CONTENTBROWSERCLASSDATASOURCE_API bool AppendItemObjectPath(const UContentBrowserDataSource* InOwnerDataSource, const FContentBrowserItemData& InItem, FString& InOutStr);
+
+	CONTENTBROWSERCLASSDATASOURCE_API bool AppendItemPackageName(const UContentBrowserDataSource* InOwnerDataSource, const FContentBrowserItemData& InItem, FString& InOutStr);
+
 	CONTENTBROWSERCLASSDATASOURCE_API bool AppendClassFileItemReference(const FContentBrowserClassFileItemDataPayload& InClassPayload, FString& InOutStr);
 
+	CONTENTBROWSERCLASSDATASOURCE_API bool AppendObjectPathFileItemReference(const FContentBrowserClassFileItemDataPayload& InClassPayload, FString& InOutStr);
+
+	CONTENTBROWSERCLASSDATASOURCE_API bool AppendPackageNameItemReference(const FContentBrowserClassFileItemDataPayload& InClassPayload, FString& InOutStr);
+
 	CONTENTBROWSERCLASSDATASOURCE_API bool GetItemPhysicalPath(const UContentBrowserDataSource* InOwnerDataSource, const FContentBrowserItemData& InItem, FString& OutDiskPath);
 
 	CONTENTBROWSERCLASSDATASOURCE_API bool GetClassFolderItemPhysicalPath(const FContentBrowserClassFolderItemDataPayload& InFolderPayload, FString& OutDiskPath);
diff --git a/Engine/Plugins/Editor/ContentBrowser/ContentBrowserClassDataSource/Source/ContentBrowserClassDataSource/Public/ContentBrowserClassDataSource.h b/Engine/Plugins/Editor/ContentBrowser/ContentBrowserClassDataSource/Source/ContentBrowserClassDataSource/Public/ContentBrowserClassDataSource.h
--- a/Engine/Plugins/Editor/ContentBrowser/ContentBrowserClassDataSource/Source/ContentBrowserClassDataSource/Public/ContentBrowserClassDataSource.h
+++ b/Engine/Plugins/Editor/ContentBrowser/ContentBrowserClassDataSource/Source/ContentBrowserClassDataSource/Public/ContentBrowserClassDataSource.h
@@ -65,6 +65,10 @@ public:
 
 	virtual bool AppendItemReference(const FContentBrowserItemData& InItem, FString& InOutStr) override;
 
+	virtual bool AppendItemObjectPath(const FContentBrowserItemData& InItem, FString& InOutStr) override;
+
+	virtual bool AppendItemPackageName(const FContentBrowserItemData& InItem, FString& InOutStr) override;
+
 	virtual bool UpdateThumbnail(const FContentBrowserItemData& InItem, FAssetThumbnail& InThumbnail) override;
 
 	virtual bool TryGetCollectionId(const FContentBrowserItemData& InItem, FSoftObjectPath& OutCollectionId) override;
diff --git a/Engine/Source/Editor/ContentBrowser/Private/AssetContextMenu.cpp b/Engine/Source/Editor/ContentBrowser/Private/AssetContextMenu.cpp
--- a/Engine/Source/Editor/ContentBrowser/Private/AssetContextMenu.cpp
+++ b/Engine/Source/Editor/ContentBrowser/Private/AssetContextMenu.cpp
@@ -3,7 +3,9 @@
 #include "AssetContextMenu.h"
 
 #include "AssetDefinition.h"
+#include "AssetDefinitionRegistry.h"
 #include "AssetRegistry/AssetData.h"
+#include "AssetRegistry/AssetRegistryHelpers.h"
 #include "AssetToolsModule.h"
 #include "AssetViewUtils.h"
 #include "CollectionAssetManagement.h"
@@ -19,6 +21,7 @@
 #include "ContentBrowserModule.h"
 #include "ContentBrowserSingleton.h"
 #include "ContentBrowserUtils.h"
+#include "Engine/AssetManager.h"
 #include "Framework/Commands/GenericCommands.h"
 #include "Framework/Commands/UIAction.h"
 #include "Framework/Commands/UICommandInfo.h"
@@ -39,10 +42,12 @@
 #include "Misc/PathViews.h"
 #include "Misc/Paths.h"
 #include "Misc/StringBuilder.h"
+#include "Misc/WarnIfAssetsLoadedInScope.h"
 #include "Modules/ModuleManager.h"
 #include "SAssetView.h"
 #include "Styling/AppStyle.h"
 #include "Styling/SlateTypes.h"
+#include "TelemetryRouter.h"
 #include "Templates/Casts.h"
 #include "Templates/Tuple.h"
 #include "Templates/UnrealTemplate.h"
@@ -62,11 +67,6 @@
 #include "UObject/UnrealNames.h"
 #include "UObject/WeakObjectPtr.h"
 #include "UObject/WeakObjectPtrTemplates.h"
-#include "AssetDefinitionRegistry.h"
-#include "TelemetryRouter.h"
-#include "AssetRegistry/AssetRegistryHelpers.h"
-#include "Engine/AssetManager.h"
-#include "Misc/WarnIfAssetsLoadedInScope.h"
 
 class FMenuBuilder;
 class SWidget;
@@ -89,6 +89,23 @@ void FAssetContextMenu::BindCommands(TSharedPtr< FUICommandList >& Commands)
 		FExecuteAction::CreateSP(this, &FAssetContextMenu::ExecuteSyncToAssetTree),
 		FCanExecuteAction::CreateSP(this, &FAssetContextMenu::CanExecuteSyncToAssetTree)
 		));
+
+	if (const TSharedPtr<SAssetView>& AssetViewPinned = AssetView.Pin())
+	{
+		const SAssetView* AssetViewPtr = AssetViewPinned.Get();
+
+		Commands->MapAction(FGenericCommands::Get().Copy, FUIAction(
+			FExecuteAction::CreateSP(AssetViewPtr, &SAssetView::ExecuteCopy, EAssetViewCopyType::ExportTextPath)
+		));
+
+		Commands->MapAction(FContentBrowserCommands::Get().AssetViewCopyObjectPath, FUIAction(
+			FExecuteAction::CreateSP(AssetViewPtr, &SAssetView::ExecuteCopy, EAssetViewCopyType::ObjectPath)
+		));
+
+		Commands->MapAction(FContentBrowserCommands::Get().AssetViewCopyPackageName, FUIAction(
+			FExecuteAction::CreateSP(AssetViewPtr, &SAssetView::ExecuteCopy, EAssetViewCopyType::PackageName)
+		));
+	}
 }
 
 TSharedRef<SWidget> FAssetContextMenu::MakeContextMenu(TArrayView<const FContentBrowserItem> InSelectedItems, const FSourcesData& InSourcesData, TSharedPtr< FUICommandList > InCommandList)
@@ -648,13 +665,23 @@ bool FAssetContextMenu::AddReferenceMenuOptions(UToolMenu* Menu)
 		FToolMenuSection& Section = Menu->AddSection("AssetContextReferences", LOCTEXT("ReferencesMenuHeading", "References"));
 
 		Section.AddMenuEntry(
-			"CopyReference",
+			FGenericCommands::Get().Copy,
 			LOCTEXT("CopyReference", "Copy Reference"),
-			LOCTEXT("CopyReferenceTooltip", "Copies reference paths for the selected assets to the clipboard."),
-			FSlateIcon(FAppStyle::GetAppStyleSetName(), "GenericCommands.Copy"),
-			FUIAction( FExecuteAction::CreateSP( this, &FAssetContextMenu::ExecuteCopyReference ) )
-			);
-	
+			GetCopyTooltip(EAssetViewCopyType::ExportTextPath)
+		);
+
+		Section.AddMenuEntry(
+			FContentBrowserCommands::Get().AssetViewCopyObjectPath,
+			LOCTEXT("CopyObjectPath", "Copy Object Path"),
+			GetCopyTooltip(EAssetViewCopyType::ObjectPath)
+		);
+
+		Section.AddMenuEntry(
+			FContentBrowserCommands::Get().AssetViewCopyPackageName,
+			LOCTEXT("CopyPackageName", "Copy Package Path"),
+			GetCopyTooltip(EAssetViewCopyType::PackageName)
+		);
+
 		if (!Context->bHasCookedPackages)
 		{
 			Section.AddMenuEntry(
@@ -670,6 +697,79 @@ bool FAssetContextMenu::AddReferenceMenuOptions(UToolMenu* Menu)
 	return true;
 }
 
+FText FAssetContextMenu::GetCopyTooltip(EAssetViewCopyType InCopyType) const
+{
+	static const FString CopyExportPathBaseTooltip = TEXT("Copies reference paths for the selected asset(s) to the clipboard:");
+	static const FString CopyObjectPathBaseTooltip = TEXT("Copies object paths for the selected asset(s) to the clipboard:");
+	static const FString CopyPackageNameBaseTooltip = TEXT("Copies package paths for the selected asset(s) to the clipboard:");
+
+	FText FinalTooltip = FText::GetEmpty();
+	switch (InCopyType)
+	{
+		case EAssetViewCopyType::ExportTextPath:
+			{
+				FString ExportTextPathTooltip = CopyExportPathBaseTooltip;
+				AppendSelectionInformationForCopy(ExportTextPathTooltip, InCopyType);
+				FinalTooltip = FText::FromString(ExportTextPathTooltip);
+			}
+			break;
+
+		case EAssetViewCopyType::ObjectPath:
+			{
+				FString ObjectPathTooltip = CopyObjectPathBaseTooltip;
+				AppendSelectionInformationForCopy(ObjectPathTooltip, InCopyType);
+				FinalTooltip = FText::FromString(ObjectPathTooltip);
+			}
+			break;
+
+		case EAssetViewCopyType::PackageName:
+			{
+				FString PackageNameTooltip = CopyPackageNameBaseTooltip;
+				AppendSelectionInformationForCopy(PackageNameTooltip, InCopyType);
+				FinalTooltip = FText::FromString(PackageNameTooltip);
+			}
+			break;
+	}
+
+	return FinalTooltip;
+}
+
+void FAssetContextMenu::AppendSelectionInformationForCopy(FString& OutTooltip, EAssetViewCopyType InCopyType) const
+{
+	FString AdditionalInformation = TEXT("");
+
+	if (!SelectedFiles.IsEmpty())
+	{
+		const FContentBrowserItem& FirstSelectedFile = SelectedFiles[0];
+
+		switch (InCopyType)
+		{
+			case EAssetViewCopyType::ExportTextPath:
+				FirstSelectedFile.AppendItemReference(AdditionalInformation);
+				break;
+
+			case EAssetViewCopyType::ObjectPath:
+				FirstSelectedFile.AppendItemObjectPath(AdditionalInformation);
+				break;
+
+			case EAssetViewCopyType::PackageName:
+				FirstSelectedFile.AppendItemPackageName(AdditionalInformation);
+				break;
+		}
+
+		if (SelectedFiles.Num() > 1)
+		{
+			AdditionalInformation += LINE_TERMINATOR;
+			AdditionalInformation += FString::Format(TEXT("+{0} more"), { FString::FormatAsNumber(SelectedFiles.Num() - 1) });
+		}
+	}
+
+	if (!AdditionalInformation.IsEmpty())
+	{
+		OutTooltip += LINE_TERMINATOR + AdditionalInformation;
+	}
+}
+
 bool FAssetContextMenu::AddAssetTypeMenuOptions(UToolMenu* Menu)
 {
 	bool bAnyTypeOptions = false;
@@ -1174,14 +1274,6 @@ bool FAssetContextMenu::IsSelectedAssetPrivate()
 	return !IsSelectedAssetPublic();
 }
 
-void FAssetContextMenu::ExecuteCopyReference()
-{
-	if (SelectedFiles.Num() > 0)
-	{
-		ContentBrowserUtils::CopyItemReferencesToClipboard(SelectedFiles);
-	}
-}
-
 void FAssetContextMenu::ExecuteCopyFilePath()
 {
 	if (SelectedFiles.Num() > 0)
diff --git a/Engine/Source/Editor/ContentBrowser/Private/AssetContextMenu.h b/Engine/Source/Editor/ContentBrowser/Private/AssetContextMenu.h
--- a/Engine/Source/Editor/ContentBrowser/Private/AssetContextMenu.h
+++ b/Engine/Source/Editor/ContentBrowser/Private/AssetContextMenu.h
@@ -22,6 +22,7 @@ class UToolMenu;
 
 enum class ECheckBoxState : uint8;
 enum class EContentBrowserViewContext : uint8;
+enum class EAssetViewCopyType;
 
 class FAssetContextMenu : public TSharedFromThis<FAssetContextMenu>
 {
@@ -103,6 +104,12 @@ private:
 	/** Adds asset reference menu options to a menu builder. Returns true if any options were added. */
 	bool AddReferenceMenuOptions(UToolMenu* Menu);
 
+	/** Return the tooltip based on the copy type */
+	FText GetCopyTooltip(EAssetViewCopyType InCopyType) const;
+
+	/** Append information on the path for the Copy tooltip based on the current selection */
+	void AppendSelectionInformationForCopy(FString& OutTooltip, EAssetViewCopyType InCopyType) const;
+
 	bool AddPublicStateMenuOptions(UToolMenu* Menu);
 
 	/** Adds menu options related to working with collections */
@@ -147,9 +154,6 @@ private:
 	/** Handler for determining the selected asset's Private state */
 	bool IsSelectedAssetPrivate();
 
-	/** Handler for CopyReference */
-	void ExecuteCopyReference();
-
 	/** Handler for CopyFilePath */
 	void ExecuteCopyFilePath();
 
diff --git a/Engine/Source/Editor/ContentBrowser/Private/AssetViewWidgets.cpp b/Engine/Source/Editor/ContentBrowser/Private/AssetViewWidgets.cpp
--- a/Engine/Source/Editor/ContentBrowser/Private/AssetViewWidgets.cpp
+++ b/Engine/Source/Editor/ContentBrowser/Private/AssetViewWidgets.cpp
@@ -89,55 +89,6 @@ static int32 GenericThumbnailSizes[(int32)EThumbnailSize::MAX] = { 24, 32, 64, 1
 // FAssetViewModeUtils
 ///////////////////////////////
 
-FReply FAssetViewModeUtils::OnViewModeKeyDown( const TSet< TSharedPtr<FAssetViewItem> >& SelectedItems, const FKeyEvent& InKeyEvent )
-{
-	// All asset views use Ctrl-C to copy references to assets
-	if ( InKeyEvent.IsControlDown() && InKeyEvent.GetCharacter() == 'C' 
-		&& !InKeyEvent.IsShiftDown() && !InKeyEvent.IsAltDown()
-		)
-	{
-		TArray<FContentBrowserItem> SelectedFiles;
-		TArray<FContentBrowserItem> SelectedFolders;
-		for (const TSharedPtr<FAssetViewItem>& SelectedItem : SelectedItems)
-		{
-			if (SelectedItem->GetItem().IsFile())
-			{
-				SelectedFiles.Add(SelectedItem->GetItem());
-			}
-			else if (SelectedItem->GetItem().IsFolder())
-			{
-				SelectedFolders.Add(SelectedItem->GetItem());
-			}
-		}
-
-		FString ClipboardText;
-
-		if (SelectedFiles.Num() > 0)
-		{
-			ClipboardText += ContentBrowserUtils::GetItemReferencesText(SelectedFiles);
-		}
-
-		if (SelectedFolders.Num() > 0)
-		{
-			if (!ClipboardText.IsEmpty())
-			{
-				ClipboardText += LINE_TERMINATOR;
-			}
-
-			ClipboardText += ContentBrowserUtils::GetFolderReferencesText(SelectedFolders);
-		}
-
-		if (!ClipboardText.IsEmpty())
-		{
-			FPlatformApplicationMisc::ClipboardCopy(*ClipboardText);
-		}
-
-		return FReply::Handled();
-	}
-
-	return FReply::Unhandled();
-}
-
 namespace AssetViewWidgets 
 {
 bool IsTopLevelFolder(const FStringView InFolderPath)
@@ -369,20 +320,6 @@ private:
 // Asset view modes
 ///////////////////////////////
 
-FReply SAssetTileView::OnKeyDown( const FGeometry& InGeometry, const FKeyEvent& InKeyEvent )
-{
-	FReply Reply = FAssetViewModeUtils::OnViewModeKeyDown(SelectedItems, InKeyEvent);
-
-	if ( Reply.IsEventHandled() )
-	{
-		return Reply;
-	}
-	else
-	{
-		return STileView<TSharedPtr<FAssetViewItem>>::OnKeyDown(InGeometry, InKeyEvent);
-	}
-}
-
 void SAssetTileView::Tick(const FGeometry& AllottedGeometry, const double InCurrentTime, const float InDeltaTime)
 {
 	// Refreshing an asset view is an intensive task. Do not do this while a user
@@ -394,20 +331,6 @@ void SAssetTileView::Tick(const FGeometry& AllottedGeometry, const double InCurr
 	}
 }
 
-FReply SAssetListView::OnKeyDown(const FGeometry& InGeometry, const FKeyEvent& InKeyEvent)
-{
-	FReply Reply = FAssetViewModeUtils::OnViewModeKeyDown(SelectedItems, InKeyEvent);
-
-	if ( Reply.IsEventHandled() )
-	{
-		return Reply;
-	}
-	else
-	{
-		return SListView<TSharedPtr<FAssetViewItem>>::OnKeyDown(InGeometry, InKeyEvent);
-	}
-}
-
 void SAssetListView::Tick(const FGeometry& AllottedGeometry, const double InCurrentTime, const float InDeltaTime)
 {
 	// Refreshing an asset view is an intensive task. Do not do this while a user
@@ -419,21 +342,6 @@ void SAssetListView::Tick(const FGeometry& AllottedGeometry, const double InCurr
 	}
 }
 
-FReply SAssetColumnView::OnKeyDown(const FGeometry& InGeometry, const FKeyEvent& InKeyEvent)
-{
-	FReply Reply = FAssetViewModeUtils::OnViewModeKeyDown(SelectedItems, InKeyEvent);
-
-	if ( Reply.IsEventHandled() )
-	{
-		return Reply;
-	}
-	else
-	{
-		return SListView<TSharedPtr<FAssetViewItem>>::OnKeyDown(InGeometry, InKeyEvent);
-	}
-}
-
-
 void SAssetColumnView::Tick(const FGeometry& AllottedGeometry, const double InCurrentTime, const float InDeltaTime)
 {
 	// Refreshing an asset view is an intensive task. Do not do this while a user
diff --git a/Engine/Source/Editor/ContentBrowser/Private/ContentBrowserCommands.cpp b/Engine/Source/Editor/ContentBrowser/Private/ContentBrowserCommands.cpp
--- a/Engine/Source/Editor/ContentBrowser/Private/ContentBrowserCommands.cpp
+++ b/Engine/Source/Editor/ContentBrowser/Private/ContentBrowserCommands.cpp
@@ -13,14 +13,18 @@
 
 void FContentBrowserCommands::RegisterCommands()
 {
+	// ContentBrowser commands
 	UI_COMMAND(OpenAssetsOrFolders, "Open Assets or Folders", "Opens the selected assets or folders, depending on the selection", EUserInterfaceActionType::Button, FInputChord(EKeys::Enter));
 	UI_COMMAND(PreviewAssets, "Preview Assets", "Loads the selected assets and previews them if possible", EUserInterfaceActionType::Button, FInputChord(EKeys::SpaceBar));
 	UI_COMMAND(CreateNewFolder, "Create New Folder", "Creates new folder in selected path", EUserInterfaceActionType::Button, FInputChord(EModifierKey::Control | EModifierKey::Shift, EKeys::N));
 	UI_COMMAND(SaveSelectedAsset, "Save Selected Item", "Save the selected item", EUserInterfaceActionType::Button, FInputChord(EModifierKey::Control, EKeys::S));
 	UI_COMMAND(SaveAllCurrentFolder, "Save All", "Save All in current folder", EUserInterfaceActionType::Button, FInputChord(EModifierKey::Control | EModifierKey::Shift, EKeys::S));
 	UI_COMMAND(ResaveAllCurrentFolder, "Resave All", "Resave all assets contained in the current folder", EUserInterfaceActionType::Button, FInputChord());
-	UI_COMMAND(CopySelectedAssetPath, "Copy Selected Asset Path", "Copy the selected asset path", EUserInterfaceActionType::Button, FInputChord(EModifierKey::Control, EKeys::C));
 	UI_COMMAND(EditPath, "Edit Path", "Edit the current content browser path", EUserInterfaceActionType::Button, FInputChord(EKeys::F4), FInputChord(EModifierKey::Control, EKeys::L));
+
+	// AssetView commands
+	UI_COMMAND(AssetViewCopyObjectPath, "Copy Selected Object Path", "Copy the selected object path", EUserInterfaceActionType::Button, FInputChord(EModifierKey::Control | EModifierKey::Shift, EKeys::C));
+	UI_COMMAND(AssetViewCopyPackageName, "Copy Selected Package Name", "Copy the selected package name", EUserInterfaceActionType::Button, FInputChord(EModifierKey::Control | EModifierKey::Alt, EKeys::C));
 }
 
 
diff --git a/Engine/Source/Editor/ContentBrowser/Private/ContentBrowserCommands.h b/Engine/Source/Editor/ContentBrowser/Private/ContentBrowserCommands.h
--- a/Engine/Source/Editor/ContentBrowser/Private/ContentBrowserCommands.h
+++ b/Engine/Source/Editor/ContentBrowser/Private/ContentBrowserCommands.h
@@ -36,6 +36,7 @@ public:
 	TSharedPtr<FUICommandInfo> SaveSelectedAsset;
 	TSharedPtr<FUICommandInfo> SaveAllCurrentFolder;
 	TSharedPtr<FUICommandInfo> ResaveAllCurrentFolder;
-	TSharedPtr<FUICommandInfo> CopySelectedAssetPath;
 	TSharedPtr<FUICommandInfo> EditPath;
+	TSharedPtr<FUICommandInfo> AssetViewCopyObjectPath;
+	TSharedPtr<FUICommandInfo> AssetViewCopyPackageName;
 };
diff --git a/Engine/Source/Editor/ContentBrowser/Private/ContentBrowserUtils.cpp b/Engine/Source/Editor/ContentBrowser/Private/ContentBrowserUtils.cpp
--- a/Engine/Source/Editor/ContentBrowser/Private/ContentBrowserUtils.cpp
+++ b/Engine/Source/Editor/ContentBrowser/Private/ContentBrowserUtils.cpp
@@ -335,6 +335,46 @@ FString ContentBrowserUtils::GetItemReferencesText(const TArray<FContentBrowserI
 	return Result;
 }
 
+FString ContentBrowserUtils::GetItemObjectPathText(const TArray<FContentBrowserItem>& Items)
+{
+	TArray<FContentBrowserItem> SortedItems = Items;
+	SortedItems.Sort([](const FContentBrowserItem& One, const FContentBrowserItem& Two)
+	{
+		return One.GetVirtualPath().Compare(Two.GetVirtualPath()) < 0;
+	});
+
+	FString Result;
+	for (const FContentBrowserItem& Item : SortedItems)
+	{
+		if (ensure(!Item.IsFolder()))
+		{
+			Item.AppendItemObjectPath(Result);
+		}
+	}
+
+	return Result;
+}
+
+FString ContentBrowserUtils::GetItemPackageNameText(const TArray<FContentBrowserItem>& Items)
+{
+	TArray<FContentBrowserItem> SortedItems = Items;
+	SortedItems.Sort([](const FContentBrowserItem& One, const FContentBrowserItem& Two)
+	{
+		return One.GetVirtualPath().Compare(Two.GetVirtualPath()) < 0;
+	});
+
+	FString Result;
+	for (const FContentBrowserItem& Item : SortedItems)
+	{
+		if (ensure(!Item.IsFolder()))
+		{
+			Item.AppendItemPackageName(Result);
+		}
+	}
+
+	return Result;
+}
+
 FString ContentBrowserUtils::GetFolderReferencesText(const TArray<FContentBrowserItem>& Folders)
 {
 	TArray<FContentBrowserItem> SortedItems = Folders;
@@ -368,6 +408,24 @@ void ContentBrowserUtils::CopyItemReferencesToClipboard(const TArray<FContentBro
 	}
 }
 
+void ContentBrowserUtils::CopyItemObjectPathToClipboard(const TArray<FContentBrowserItem>& ItemsToCopy)
+{
+	FString Text = GetItemObjectPathText(ItemsToCopy);
+	if (!Text.IsEmpty())
+	{
+		FPlatformApplicationMisc::ClipboardCopy(*Text);
+	}
+}
+
+void ContentBrowserUtils::CopyItemPackageNameToClipboard(const TArray<FContentBrowserItem>& ItemsToCopy)
+{
+	FString Text = GetItemPackageNameText(ItemsToCopy);
+	if (!Text.IsEmpty())
+	{
+		FPlatformApplicationMisc::ClipboardCopy(*Text);
+	}
+}
+
 void ContentBrowserUtils::CopyFolderReferencesToClipboard(const TArray<FContentBrowserItem>& FoldersToCopy)
 {
 	FString Text = GetFolderReferencesText(FoldersToCopy);
diff --git a/Engine/Source/Editor/ContentBrowser/Private/ContentBrowserUtils.h b/Engine/Source/Editor/ContentBrowser/Private/ContentBrowserUtils.h
--- a/Engine/Source/Editor/ContentBrowser/Private/ContentBrowserUtils.h
+++ b/Engine/Source/Editor/ContentBrowser/Private/ContentBrowserUtils.h
@@ -45,12 +45,24 @@ namespace ContentBrowserUtils
 	/** Returns references to the specified items */
 	FString GetItemReferencesText(const TArray<FContentBrowserItem>& Items);
 
+	/** Returns object path of the specified items */
+	FString GetItemObjectPathText(const TArray<FContentBrowserItem>& Items);
+
+	/** Returns package name of the specified items */
+	FString GetItemPackageNameText(const TArray<FContentBrowserItem>& Items);
+
 	/** Returns references to the specified folders */
 	FString GetFolderReferencesText(const TArray<FContentBrowserItem>& Folders);
 
 	/** Copies references to the specified items to the clipboard */
 	void CopyItemReferencesToClipboard(const TArray<FContentBrowserItem>& ItemsToCopy);
 
+	/** Copies object path of the specified items to the clipboard */
+	void CopyItemObjectPathToClipboard(const TArray<FContentBrowserItem>& ItemsToCopy);
+
+	/** Copies package name of the specified items to the clipboard */
+	void CopyItemPackageNameToClipboard(const TArray<FContentBrowserItem>& ItemsToCopy);
+
 	/** Copies references to the specified folders to the clipboard */
 	void CopyFolderReferencesToClipboard(const TArray<FContentBrowserItem>& FoldersToCopy);
 
diff --git a/Engine/Source/Editor/ContentBrowser/Private/SAssetView.cpp b/Engine/Source/Editor/ContentBrowser/Private/SAssetView.cpp
--- a/Engine/Source/Editor/ContentBrowser/Private/SAssetView.cpp
+++ b/Engine/Source/Editor/ContentBrowser/Private/SAssetView.cpp
@@ -16,6 +16,7 @@
 #include "Async/ParallelFor.h"
 #include "Async/WordMutex.h"
 #include "Async/UniqueLock.h"
+#include "ContentBrowserCommands.h"
 #include "ContentBrowserConfig.h"
 #include "ContentBrowserDataDragDropOp.h"
 #include "ContentBrowserDataLegacyBridge.h"
@@ -36,6 +37,7 @@
 #include "Engine/Level.h"
 #include "Factories/Factory.h"
 #include "Framework/Application/SlateApplication.h"
+#include "Framework/Commands/GenericCommands.h"
 #include "Framework/Commands/UIAction.h"
 #include "Framework/MultiBox/MultiBoxBuilder.h"
 #include "Framework/Notifications/NotificationManager.h"
@@ -1214,6 +1216,8 @@ void SAssetView::Construct( const FArguments& InArgs )
 
 	TSharedRef<SVerticalBox> VerticalBox = SNew(SVerticalBox);
 
+	BindCommands();
+
 	ChildSlot
 	.Padding(0.0f)
 	[
@@ -2380,93 +2384,12 @@ FReply SAssetView::OnKeyChar( const FGeometry& MyGeometry,const FCharacterEvent&
 	return FReply::Unhandled();
 }
 
-static bool IsValidObjectPath(const FString& Path, FString& OutObjectClassName, FString& OutObjectPath, FString& OutPackageName)
-{
-	if (FPackageName::ParseExportTextPath(Path, &OutObjectClassName, &OutObjectPath))
-	{
-		if (UClass* ObjectClass = UClass::TryFindTypeSlow<UClass>(OutObjectClassName, EFindFirstObjectOptions::ExactClass))
-		{
-			OutPackageName = FPackageName::ObjectPathToPackageName(OutObjectPath);
-			if (FPackageName::IsValidLongPackageName(OutPackageName))
-			{
-				return true;
-			}
-		}
-	}
-
-	return false;
-}
-
-static bool ContainsT3D(const FString& ClipboardText)
-{
-	return (ClipboardText.StartsWith(TEXT("Begin Object")) && ClipboardText.EndsWith(TEXT("End Object")))
-		|| (ClipboardText.StartsWith(TEXT("Begin Map")) && ClipboardText.EndsWith(TEXT("End Map")));
-}
-
 FReply SAssetView::OnKeyDown( const FGeometry& MyGeometry, const FKeyEvent& InKeyEvent )
 {
 	const bool bIsControlOrCommandDown = InKeyEvent.IsControlDown() || InKeyEvent.IsCommandDown();
-	
-	if (bIsControlOrCommandDown && InKeyEvent.GetCharacter() == 'V' && IsAssetPathSelected())
-	{
-		FString AssetPaths;
-		TArray<FString> AssetPathsSplit;
-
-		// Get the copied asset paths
-		FPlatformApplicationMisc::ClipboardPaste(AssetPaths);
-
-		// Make sure the clipboard does not contain T3D
-		AssetPaths.TrimEndInline();
-		if (!ContainsT3D(AssetPaths))
-		{
-			AssetPaths.ParseIntoArrayLines(AssetPathsSplit);
-
-			// Get assets and copy them
-			TArray<UObject*> AssetsToCopy;
-			FAssetToolsModule& AssetToolsModule = FModuleManager::LoadModuleChecked<FAssetToolsModule>("AssetTools");
-			for (const FString& AssetPath : AssetPathsSplit)
-			{
-				// Validate string
-				FString ObjectClassName;
-				FString ObjectPath;
-				FString PackageName;
-				if (IsValidObjectPath(AssetPath, ObjectClassName, ObjectPath, PackageName))
-				{
-					// Only duplicate the objects of the supported classes.
-					if (AssetToolsModule.Get().GetAssetClassPathPermissionList(EAssetClassAction::ViewAsset)->PassesStartsWithFilter(ObjectClassName))
-					{
-						FLinkerInstancingContext InstancingContext({ ULevel::LoadAllExternalObjectsTag });
-						UObject* ObjectToCopy = LoadObject<UObject>(nullptr, *ObjectPath, nullptr, LOAD_None, nullptr, &InstancingContext);
-						if (ObjectToCopy && !ObjectToCopy->IsA(UClass::StaticClass()))
-						{
-							AssetsToCopy.Add(ObjectToCopy);
-						}
-					}
-				}
-			}
-			
-			if (AssetsToCopy.Num())
-			{
-				UContentBrowserDataSubsystem* ContentBrowserData = IContentBrowserDataModule::Get().GetSubsystem();
-				if (ensure(ContentBrowserData))
-				{
-					for (const FName& SelectedVirtualPath : SourcesData.VirtualPaths)
-					{
-						const FContentBrowserItem SelectedItem = ContentBrowserData->GetItemAtPath(SelectedVirtualPath, EContentBrowserItemTypeFilter::IncludeFolders);
-						if (SelectedItem.IsValid())
-						{
-							FName PackagePath;
-							if (SelectedItem.Legacy_TryGetPackagePath(PackagePath))
-							{
-								ContentBrowserUtils::CopyAssets(AssetsToCopy, PackagePath.ToString());
-								break;
-							}
-						}
-					}
-				}
-			}
-		}
 
+	if (Commands->ProcessCommandBindings(InKeyEvent))
+	{
 		return FReply::Handled();
 	}
 	// Swallow the key-presses used by the quick-jump in OnKeyChar to avoid other things (such as the viewport commands) getting them instead
@@ -3209,6 +3132,175 @@ FAssetViewInstanceConfig* SAssetView::GetAssetViewConfig() const
 	return nullptr;
 }
 
+void SAssetView::BindCommands()
+{
+	Commands = TSharedPtr<FUICommandList>(new FUICommandList);
+
+	Commands->MapAction(FGenericCommands::Get().Copy, FUIAction(
+		FExecuteAction::CreateSP(this, &SAssetView::ExecuteCopy, EAssetViewCopyType::ExportTextPath)
+	));
+
+	Commands->MapAction(FContentBrowserCommands::Get().AssetViewCopyObjectPath, FUIAction(
+		FExecuteAction::CreateSP(this, &SAssetView::ExecuteCopy, EAssetViewCopyType::ObjectPath)
+	));
+
+	Commands->MapAction(FContentBrowserCommands::Get().AssetViewCopyPackageName, FUIAction(
+		FExecuteAction::CreateSP(this, &SAssetView::ExecuteCopy, EAssetViewCopyType::PackageName)
+	));
+
+	Commands->MapAction(FGenericCommands::Get().Paste, FUIAction(
+		FExecuteAction::CreateSP(this, &SAssetView::ExecutePaste),
+		FCanExecuteAction::CreateSP(this, &SAssetView::IsAssetPathSelected)
+	));
+
+	FInputBindingManager::Get().RegisterCommandList(FContentBrowserCommands::Get().GetContextName(), Commands.ToSharedRef());
+}
+
+void SAssetView::PopulateSelectedFilesAndFolders(TArray<FContentBrowserItem>& OutSelectedFolders, TArray<FContentBrowserItem>& OutSelectedFiles) const
+{
+	for (const FContentBrowserItem& SelectedItem : GetSelectedItems())
+	{
+		if (SelectedItem.IsFile())
+		{
+			OutSelectedFiles.Add(SelectedItem);
+		}
+		else if (SelectedItem.IsFolder())
+		{
+			OutSelectedFolders.Add(SelectedItem);
+		}
+	}
+}
+
+void SAssetView::ExecuteCopy(EAssetViewCopyType InCopyType) const
+{
+	TArray<FContentBrowserItem> SelectedFiles;
+	TArray<FContentBrowserItem> SelectedFolders;
+
+	PopulateSelectedFilesAndFolders(SelectedFolders, SelectedFiles);
+
+	FString ClipboardText;
+	if (SelectedFiles.Num() > 0)
+	{
+		switch (InCopyType)
+		{
+			case EAssetViewCopyType::ExportTextPath:
+				ClipboardText += ContentBrowserUtils::GetItemReferencesText(SelectedFiles);
+				break;
+
+			case EAssetViewCopyType::ObjectPath:
+				ClipboardText += ContentBrowserUtils::GetItemObjectPathText(SelectedFiles);
+				break;
+
+			case EAssetViewCopyType::PackageName:
+				ClipboardText += ContentBrowserUtils::GetItemPackageNameText(SelectedFiles);
+				break;
+		}
+	}
+
+	ExecuteCopyFolders(SelectedFolders, ClipboardText);
+
+	if (!ClipboardText.IsEmpty())
+	{
+		FPlatformApplicationMisc::ClipboardCopy(*ClipboardText);
+	}
+}
+
+void SAssetView::ExecuteCopyFolders(const TArray<FContentBrowserItem>& InSelectedFolders, FString& OutClipboardText) const
+{
+	if (InSelectedFolders.Num() > 0)
+	{
+		if (!OutClipboardText.IsEmpty())
+		{
+			OutClipboardText += LINE_TERMINATOR;
+		}
+		OutClipboardText += ContentBrowserUtils::GetFolderReferencesText(InSelectedFolders);
+	}
+}
+
+static bool IsValidObjectPath(const FString& Path, FString& OutObjectClassName, FString& OutObjectPath, FString& OutPackageName)
+{
+	if (FPackageName::ParseExportTextPath(Path, &OutObjectClassName, &OutObjectPath))
+	{
+		if (UClass* ObjectClass = UClass::TryFindTypeSlow<UClass>(OutObjectClassName, EFindFirstObjectOptions::ExactClass))
+		{
+			OutPackageName = FPackageName::ObjectPathToPackageName(OutObjectPath);
+			if (FPackageName::IsValidLongPackageName(OutPackageName))
+			{
+				return true;
+			}
+		}
+	}
+
+	return false;
+}
+
+static bool ContainsT3D(const FString& ClipboardText)
+{
+	return (ClipboardText.StartsWith(TEXT("Begin Object")) && ClipboardText.EndsWith(TEXT("End Object")))
+		|| (ClipboardText.StartsWith(TEXT("Begin Map")) && ClipboardText.EndsWith(TEXT("End Map")));
+}
+
+void SAssetView::ExecutePaste()
+{
+	FString AssetPaths;
+	TArray<FString> AssetPathsSplit;
+
+	// Get the copied asset paths
+	FPlatformApplicationMisc::ClipboardPaste(AssetPaths);
+
+	// Make sure the clipboard does not contain T3D
+	AssetPaths.TrimEndInline();
+	if (!ContainsT3D(AssetPaths))
+	{
+		AssetPaths.ParseIntoArrayLines(AssetPathsSplit);
+
+		// Get assets and copy them
+		TArray<UObject*> AssetsToCopy;
+		FAssetToolsModule& AssetToolsModule = FModuleManager::LoadModuleChecked<FAssetToolsModule>("AssetTools");
+		for (const FString& AssetPath : AssetPathsSplit)
+		{
+			// Validate string
+			FString ObjectClassName;
+			FString ObjectPath;
+			FString PackageName;
+			if (IsValidObjectPath(AssetPath, ObjectClassName, ObjectPath, PackageName))
+			{
+				// Only duplicate the objects of the supported classes.
+				if (AssetToolsModule.Get().GetAssetClassPathPermissionList(EAssetClassAction::ViewAsset)->PassesStartsWithFilter(ObjectClassName))
+				{
+					FLinkerInstancingContext InstancingContext({ ULevel::LoadAllExternalObjectsTag });
+					UObject* ObjectToCopy = LoadObject<UObject>(nullptr, *ObjectPath, nullptr, LOAD_None, nullptr, &InstancingContext);
+					if (ObjectToCopy && !ObjectToCopy->IsA(UClass::StaticClass()))
+					{
+						AssetsToCopy.Add(ObjectToCopy);
+					}
+				}
+			}
+		}
+
+		if (AssetsToCopy.Num())
+		{
+			UContentBrowserDataSubsystem* ContentBrowserData = IContentBrowserDataModule::Get().GetSubsystem();
+			if (ensure(ContentBrowserData))
+			{
+				for (const FName& SelectedVirtualPath : SourcesData.VirtualPaths)
+				{
+					const FContentBrowserItem SelectedItem = ContentBrowserData->GetItemAtPath(SelectedVirtualPath, EContentBrowserItemTypeFilter::IncludeFolders);
+					if (SelectedItem.IsValid())
+					{
+						FName PackagePath;
+						if (SelectedItem.Legacy_TryGetPackagePath(PackagePath))
+						{
+							ContentBrowserUtils::CopyAssets(AssetsToCopy, PackagePath.ToString());
+							break;
+						}
+					}
+				}
+			}
+		}
+	}
+}
+
 void SAssetView::ToggleShowAllFolder()
 {
 	const bool bNewValue = !IsShowingAllFolder();
diff --git a/Engine/Source/Editor/ContentBrowser/Private/SContentBrowser.cpp b/Engine/Source/Editor/ContentBrowser/Private/SContentBrowser.cpp
--- a/Engine/Source/Editor/ContentBrowser/Private/SContentBrowser.cpp
+++ b/Engine/Source/Editor/ContentBrowser/Private/SContentBrowser.cpp
@@ -764,10 +764,6 @@ void SContentBrowser::BindCommands()
 		FExecuteAction::CreateSP(this, &SContentBrowser::HandleResaveAllCurrentFolderCommand)
 	));
 
-	Commands->MapAction(FContentBrowserCommands::Get().CopySelectedAssetPath, FUIAction(
-		FExecuteAction::CreateSP(this, &SContentBrowser::CopySelectedAssetPathCommand)
-	));
-	
 	Commands->MapAction(FContentBrowserCommands::Get().EditPath, FUIAction(
 		FExecuteAction::CreateSP(this, &SContentBrowser::EditPathCommand)
 	));
diff --git a/Engine/Source/Editor/ContentBrowser/Private/SPathView.cpp b/Engine/Source/Editor/ContentBrowser/Private/SPathView.cpp
--- a/Engine/Source/Editor/ContentBrowser/Private/SPathView.cpp
+++ b/Engine/Source/Editor/ContentBrowser/Private/SPathView.cpp
@@ -62,6 +62,7 @@
 #include "ToolMenu.h"
 #include "ToolMenus.h"
 #include "ToolMenuSection.h"
+#include "Framework/Commands/GenericCommands.h"
 #include "Trace/Detail/Channel.h"
 #include "Types/WidgetActiveTimerDelegate.h"
 #include "UObject/UObjectGlobals.h"
@@ -982,6 +983,9 @@ void SPathView::Construct( const FArguments& InArgs )
 	FContentBrowserModule& ContentBrowserModule = FModuleManager::GetModuleChecked<FContentBrowserModule>(TEXT("ContentBrowser"));
 	ContentBrowserModule.GetOnContentBrowserSettingChanged().AddSP(this, &SPathView::HandleSettingChanged);
 
+	// Binds the commands for the PathView
+	BindCommands();
+
 	// Setup plugin filters
 	PluginPathFilters = InArgs._PluginPathFilters;
 	if (PluginPathFilters.IsValid())
@@ -1208,6 +1212,15 @@ void SPathView::Tick(const FGeometry& AllottedGeometry, const double InCurrentTi
 	}
 }
 
+FReply SPathView::OnKeyDown(const FGeometry& MyGeometry, const FKeyEvent& InKeyEvent)
+{
+	if (Commands->ProcessCommandBindings(InKeyEvent))
+	{
+		return FReply::Handled();
+	}
+	return FReply::Unhandled();
+}
+
 void SPathView::PopulatePathViewFiltersMenu(UToolMenu* Menu)
 {
 	{
@@ -2667,6 +2680,20 @@ TSharedRef<SWidget> SPathView::GetViewButtonContent()
 	return UToolMenus::Get()->GenerateWidget("ContentBrowser.PathViewOptions", MenuContext);
 }
 
+void SPathView::CopySelectedFolder() const
+{
+	ContentBrowserUtils::CopyFolderReferencesToClipboard(GetSelectedFolderItems());
+}
+
+void SPathView::BindCommands()
+{
+	Commands = TSharedPtr<FUICommandList>(new FUICommandList);
+
+	Commands->MapAction(FGenericCommands::Get().Copy, FUIAction(
+		FExecuteAction::CreateSP(this, &SPathView::CopySelectedFolder)
+	));
+}
+
 void SFavoritePathView::Construct(const FArguments& InArgs)
 {
 	// Bind the favorites menu to update after folder changes
diff --git a/Engine/Source/Editor/ContentBrowser/Private/SPathView.h b/Engine/Source/Editor/ContentBrowser/Private/SPathView.h
--- a/Engine/Source/Editor/ContentBrowser/Private/SPathView.h
+++ b/Engine/Source/Editor/ContentBrowser/Private/SPathView.h
@@ -168,6 +168,9 @@ public:
 	/** Tick to poll attributes */
 	virtual void Tick(const FGeometry& AllottedGeometry, const double InCurrentTime, const float InDeltaTime) override;
 
+	/** Process the Commands of the PathView */
+	virtual FReply OnKeyDown( const FGeometry& MyGeometry, const FKeyEvent& InKeyEvent ) override;
+
 	/** Selects the closest matches to the supplied paths in the tree. "/" delimited */
 	void SetSelectedPaths(const TArray<FName>& Paths);
 
@@ -417,6 +420,12 @@ private:
 	/** Handler for when the view combo button is clicked */
 	TSharedRef<SWidget> GetViewButtonContent();
 
+	/** Callback for the Copy command for the PathView */
+	void CopySelectedFolder() const;
+
+	/** Bind our UI commands */
+	void BindCommands();
+
 protected:
 	/** A helper class to manage PreventTreeItemChangedDelegateCount by incrementing it when constructed (on the stack) and decrementing when destroyed */
 	class FScopedPreventTreeItemChangedDelegate
@@ -537,6 +546,9 @@ private:
 	/** The title of this path view */
 	FText TreeTitle;
 
+	/** Commands handled by this widget */
+	TSharedPtr<FUICommandList> Commands;
+
 	/** The filter collection used to filter plugins */
 	TSharedPtr<FPluginFilterCollectionType> PluginPathFilters;
 
diff --git a/Engine/Source/Editor/ContentBrowser/Public/AssetViewWidgets.h b/Engine/Source/Editor/ContentBrowser/Public/AssetViewWidgets.h
--- a/Engine/Source/Editor/ContentBrowser/Public/AssetViewWidgets.h
+++ b/Engine/Source/Editor/ContentBrowser/Public/AssetViewWidgets.h
@@ -69,11 +69,6 @@ DECLARE_DELEGATE_OneParam( FOnItemDestroyed, const TSharedPtr<FAssetViewItem>& /
 class SAssetListItem;
 class SAssetTileItem;
 
-namespace FAssetViewModeUtils 
-{
-	FReply OnViewModeKeyDown( const TSet< TSharedPtr<FAssetViewItem> >& SelectedItems, const FKeyEvent& InKeyEvent );
-}
-
 struct FAssetViewItemHelper
 {
 public:
@@ -90,7 +85,6 @@ class SAssetTileView : public STileView<TSharedPtr<FAssetViewItem>>
 {
 public:
 	virtual bool SupportsKeyboardFocus() const override { return true; }
-	virtual FReply OnKeyDown( const FGeometry& InGeometry, const FKeyEvent& InKeyEvent ) override;
 	virtual void Tick( const FGeometry& AllottedGeometry, const double InCurrentTime, const float InDeltaTime ) override;
 };
 
@@ -99,7 +93,6 @@ class SAssetListView : public SListView<TSharedPtr<FAssetViewItem>>
 {
 public:
 	virtual bool SupportsKeyboardFocus() const override { return true; }
-	virtual FReply OnKeyDown( const FGeometry& InGeometry, const FKeyEvent& InKeyEvent ) override;
 	virtual void Tick(const FGeometry& AllottedGeometry, const double InCurrentTime, const float InDeltaTime) override;
 };
 
@@ -108,7 +101,6 @@ class SAssetColumnView : public SListView<TSharedPtr<FAssetViewItem>>
 {
 public:
 	virtual bool SupportsKeyboardFocus() const override { return true; }
-	virtual FReply OnKeyDown( const FGeometry& InGeometry, const FKeyEvent& InKeyEvent ) override;
 	virtual void Tick(const FGeometry& AllottedGeometry, const double InCurrentTime, const float InDeltaTime) override;
 };
 
diff --git a/Engine/Source/Editor/ContentBrowser/Public/SAssetView.h b/Engine/Source/Editor/ContentBrowser/Public/SAssetView.h
--- a/Engine/Source/Editor/ContentBrowser/Public/SAssetView.h
+++ b/Engine/Source/Editor/ContentBrowser/Public/SAssetView.h
@@ -46,6 +46,7 @@
 #include "Widgets/Views/STableRow.h"
 #include "Widgets/Views/STableViewBase.h"
 
+class FAssetContextMenu;
 class FAssetTextFilter;
 class FAssetThumbnail;
 class FAssetViewItem;
@@ -93,6 +94,14 @@ DECLARE_DELEGATE(FOnSearchOptionChanged);
 /** Fires whenever asset view options menu is being opened, gives chance for external code to set additional context */
 DECLARE_DELEGATE_OneParam(FOnExtendAssetViewOptionsMenuContext, FToolMenuContext&);
 
+/** Copy types */
+enum class EAssetViewCopyType
+{
+	ExportTextPath,
+	ObjectPath,
+	PackageName
+};
+
 /**
  * A widget to display a list of filtered assets
  */
@@ -300,6 +309,8 @@ public:
 		SLATE_ARGUMENT(TOptional<FName>, AssetViewOptionsProfile)
 	SLATE_END_ARGS()
 
+	friend FAssetContextMenu;
+
 	SAssetView();
 	~SAssetView();
 
@@ -904,6 +915,21 @@ private:
 	/** Get the config struct for this asset view, if one exists. */
 	FAssetViewInstanceConfig* GetAssetViewConfig() const;
 
+	/** Bind our UI commands */
+	void BindCommands();
+
+	/** Populate the given parameters based on the current selection */
+	void PopulateSelectedFilesAndFolders(TArray<FContentBrowserItem>& OutSelectedFolders, TArray<FContentBrowserItem>& OutSelectedFiles) const;
+
+	/** Handler for the CopyReference CopyObjectPath and CopyPackageName */
+	void ExecuteCopy(EAssetViewCopyType InCopyType) const;
+
+	/** Append folders path to the given ClipboardText */
+	void ExecuteCopyFolders(const TArray<FContentBrowserItem>& InSelectedFolders, FString& OutClipboardText) const;
+
+	/** Handler for Paste */
+	void ExecutePaste();
+
 private:
 	friend class FAssetViewFrontendFilterHelper;
 
@@ -1176,6 +1202,9 @@ private:
 	/** Initial set of item categories that this view should show - may be adjusted further by things like CanShowClasses or legacy delegate bindings */
 	EContentBrowserItemCategoryFilter InitialCategoryFilter;
 
+	/** Commands handled by this widget */
+	TSharedPtr<FUICommandList> Commands;
+
 	bool bShowDisallowedAssetClassAsUnsupportedItems = false;
 
 	/** A struct to hold data for the deferred creation of a file or folder item */
diff --git a/Engine/Source/Editor/ContentBrowserData/Private/ContentBrowserDataSource.cpp b/Engine/Source/Editor/ContentBrowserData/Private/ContentBrowserDataSource.cpp
--- a/Engine/Source/Editor/ContentBrowserData/Private/ContentBrowserDataSource.cpp
+++ b/Engine/Source/Editor/ContentBrowserData/Private/ContentBrowserDataSource.cpp
@@ -500,6 +500,16 @@ bool UContentBrowserDataSource::AppendItemReference(const FContentBrowserItemDat
 	return false;
 }
 
+bool UContentBrowserDataSource::AppendItemObjectPath(const FContentBrowserItemData& InItem, FString& InOutStr)
+{
+	return false;
+}
+
+bool UContentBrowserDataSource::AppendItemPackageName(const FContentBrowserItemData& InItem, FString& InOutStr)
+{
+	return false;
+}
+
 bool UContentBrowserDataSource::UpdateThumbnail(const FContentBrowserItemData& InItem, FAssetThumbnail& InThumbnail)
 {
 	return false;
diff --git a/Engine/Source/Editor/ContentBrowserData/Private/ContentBrowserItem.cpp b/Engine/Source/Editor/ContentBrowserData/Private/ContentBrowserItem.cpp
--- a/Engine/Source/Editor/ContentBrowserData/Private/ContentBrowserItem.cpp
+++ b/Engine/Source/Editor/ContentBrowserData/Private/ContentBrowserItem.cpp
@@ -586,6 +586,16 @@ bool FContentBrowserItem::AppendItemReference(FString& InOutStr) const
 	return FContentBrowserItemHelper::CallDataSourceImpl<UContentBrowserDataSource>(*this, UE_PROJECTION_MEMBER(UContentBrowserDataSource, AppendItemReference), InOutStr);
 }
 
+bool FContentBrowserItem::AppendItemObjectPath(FString& InOutStr) const
+{
+	return FContentBrowserItemHelper::CallDataSourceImpl<UContentBrowserDataSource>(*this, UE_PROJECTION_MEMBER(UContentBrowserDataSource, AppendItemObjectPath), InOutStr);
+}
+
+bool FContentBrowserItem::AppendItemPackageName(FString& InOutStr) const
+{
+	return FContentBrowserItemHelper::CallDataSourceImpl<UContentBrowserDataSource>(*this, UE_PROJECTION_MEMBER(UContentBrowserDataSource, AppendItemPackageName), InOutStr);
+}
+
 bool FContentBrowserItem::UpdateThumbnail(FAssetThumbnail& InThumbnail) const
 {
 	return FContentBrowserItemHelper::CallDataSourceImpl<UContentBrowserDataSource>(*this, UE_PROJECTION_MEMBER(UContentBrowserDataSource, UpdateThumbnail), InThumbnail);
diff --git a/Engine/Source/Editor/ContentBrowserData/Public/ContentBrowserDataSource.h b/Engine/Source/Editor/ContentBrowserData/Public/ContentBrowserDataSource.h
--- a/Engine/Source/Editor/ContentBrowserData/Public/ContentBrowserDataSource.h
+++ b/Engine/Source/Editor/ContentBrowserData/Public/ContentBrowserDataSource.h
@@ -658,6 +658,28 @@ public:
 	 */
 	virtual bool AppendItemReference(const FContentBrowserItemData& InItem, FString& InOutStr);
 
+	/**
+	 * Attempt to append any object path for the given item to the given string.
+	 * @note Used when copying item object path to the clipboard.
+	 *
+	 * @param InItem The item to query.
+	 * @param InOutStr The string to append to (LINE_TERMINATOR delimited).
+	 *
+	 * @return True if package name were appended, false otherwise.
+	 */
+	virtual bool AppendItemObjectPath(const FContentBrowserItemData& InItem, FString& InOutStr);
+
+	/**
+	 * Attempt to append any package name for the given item to the given string.
+	 * @note Used when copying item package name to the clipboard.
+	 *
+	 * @param InItem The item to query.
+	 * @param InOutStr The string to append to (LINE_TERMINATOR delimited).
+	 *
+	 * @return True if package name were appended, false otherwise.
+	 */
+	virtual bool AppendItemPackageName(const FContentBrowserItemData& InItem, FString& InOutStr);
+
 	/**
 	 * Attempt to update the thumbnail associated with the given item.
 	 *
diff --git a/Engine/Source/Editor/ContentBrowserData/Public/ContentBrowserItem.h b/Engine/Source/Editor/ContentBrowserData/Public/ContentBrowserItem.h
--- a/Engine/Source/Editor/ContentBrowserData/Public/ContentBrowserItem.h
+++ b/Engine/Source/Editor/ContentBrowserData/Public/ContentBrowserItem.h
@@ -407,6 +407,26 @@ public:
 	 */
 	bool AppendItemReference(FString& InOutStr) const;
 
+	/**
+	 * Attempt to append any object path for this item to the given string.
+	 * @note Used when copying item object path to the clipboard.
+	 *
+	 * @param InOutStr The string to append to (LINE_TERMINATOR delimited).
+	 *
+	 * @return True if references were appended, false otherwise.
+	 */
+	bool AppendItemObjectPath(FString& InOutStr) const;
+
+	/**
+	 * Attempt to append any package name for this item to the given string.
+	 * @note Used when copying item package name to the clipboard.
+	 *
+	 * @param InOutStr The string to append to (LINE_TERMINATOR delimited).
+	 *
+	 * @return True if references were appended, false otherwise.
+	 */
+	bool AppendItemPackageName(FString& InOutStr) const;
+
 	/**
 	 * Attempt to update the thumbnail associated with this item.
 	 *
diff --git a/Engine/Source/Editor/EditorStyle/Private/StarshipStyle.cpp b/Engine/Source/Editor/EditorStyle/Private/StarshipStyle.cpp
--- a/Engine/Source/Editor/EditorStyle/Private/StarshipStyle.cpp
+++ b/Engine/Source/Editor/EditorStyle/Private/StarshipStyle.cpp
@@ -6375,6 +6375,9 @@ void FStarshipEditorStyle::FStyle::SetupContentBrowserStyle()
 		Set( "ContentBrowser.AssetActions.VolumeTexture", new IMAGE_BRUSH_SVG("Starship/AssetActions/volume-texture", Icon16x16));
 		Set( "ContentBrowser.AssetActions.RemoveVertexColors", new IMAGE_BRUSH_SVG("Starship/AssetActions/remove-vertex-colors", Icon16x16));
 
+		// ContentBrowser Commands Icons
+		Set( "ContentBrowser.AssetViewCopyObjectPath", new IMAGE_BRUSH_SVG("../../Slate/Starship/Common/Copy", Icon16x16) );
+		Set( "ContentBrowser.AssetViewCopyPackageName", new IMAGE_BRUSH_SVG("../../Slate/Starship/Common/Copy", Icon16x16) );
 
 		Set( "MediaAsset.AssetActions.Play.Small", new IMAGE_BRUSH( "Icons/icon_SCueEd_PlayCue_16x", Icon16x16 ) );
 		Set( "MediaAsset.AssetActions.Stop.Small", new IMAGE_BRUSH( "Icons/icon_SCueEd_Stop_16x", Icon16x16 ) );
