PR: [ContentBrowser] Store SContentBrowser::bIsLocked by instance in INI

#10168

diff --git a/Engine/Source/Editor/ContentBrowser/Private/SContentBrowser.cpp b/Engine/Source/Editor/ContentBrowser/Private/SContentBrowser.cpp
--- a/Engine/Source/Editor/ContentBrowser/Private/SContentBrowser.cpp
+++ b/Engine/Source/Editor/ContentBrowser/Private/SContentBrowser.cpp
@@ -1981,6 +1981,7 @@ void SContentBrowser::SaveSettings() const
 	const FString& SettingsString = InstanceName.ToString();
 
 	GConfig->SetBool(*SettingsIniSection, *(SettingsString + TEXT(".SourcesExpanded")), bSourcesViewExpanded, GEditorPerProjectIni);
+	GConfig->SetBool(*SettingsIniSection, *(SettingsString + TEXT(".IsLocked")), bIsLocked, GEditorPerProjectIni);
 
 	GConfig->SetBool(*SettingsIniSection, *(SettingsString + TEXT(".FavoritesAreaExpanded")), FavoritesArea->IsExpanded(), GEditorPerProjectIni);
 	GConfig->SetBool(*SettingsIniSection, *(SettingsString + TEXT(".PathAreaExpanded")), PathArea->IsExpanded(), GEditorPerProjectIni);
@@ -2128,6 +2129,9 @@ void SContentBrowser::LoadSettings(const FName& InInstanceName)
 	bSourcesViewExpanded = true;
 	GConfig->GetBool(*SettingsIniSection, *(SettingsString + TEXT(".SourcesExpanded")), bSourcesViewExpanded, GEditorPerProjectIni);
 
+	bIsLocked = false;
+	GConfig->GetBool(*SettingsIniSection, *(SettingsString + TEXT(".IsLocked")), bIsLocked, GEditorPerProjectIni);
+
 	bool bFavoritesAreaExpanded = false;
 	GConfig->GetBool(*SettingsIniSection, *(SettingsString + TEXT(".FavoritesAreaExpanded")), bFavoritesAreaExpanded, GEditorPerProjectIni);
 	FavoritesArea->SetExpanded(bFavoritesAreaExpanded);
