ContentBrowser: Suppress spurious refreshes when a console variable changes other than the AssetTools.EnablePublicAssetFeature that we are listening for.
#rnx

diff --git a/Engine/Source/Editor/ContentBrowser/Private/SContentBrowser.cpp b/Engine/Source/Editor/ContentBrowser/Private/SContentBrowser.cpp
--- a/Engine/Source/Editor/ContentBrowser/Private/SContentBrowser.cpp
+++ b/Engine/Source/Editor/ContentBrowser/Private/SContentBrowser.cpp
@@ -380,8 +380,8 @@ void SContentBrowser::Construct( const FArguments& InArgs, const FName& InInstan
 
 	// Register console variable sink for private content setting changing
 	CVarSinkHandle = IConsoleManager::Get().RegisterConsoleVariableSink_Handle(FConsoleCommandDelegate::CreateSP(this, &SContentBrowser::OnConsoleVariableChanged));
-	static const IConsoleVariable* EnablePublicAssetFeatureCVar = IConsoleManager::Get().FindConsoleVariable(TEXT("AssetTools.EnablePublicAssetFeature"));
-	const bool bEnablePrivateContentFeature = EnablePublicAssetFeatureCVar && EnablePublicAssetFeatureCVar->GetBool();
+	UpdatePrivateContentFeatureEnabled(false /* bUpdateFilterIfChanged */);
+
 	AssetViewPtr = SNew(SAssetView)
 			.ThumbnailLabel(Config != nullptr ? Config->ThumbnailLabel : EThumbnailLabel::ClassName)
 			//.ThumbnailScale(Config != nullptr ? Config->ThumbnailScale : 0.18f)
@@ -391,7 +391,9 @@ void SContentBrowser::Construct( const FArguments& InArgs, const FName& InInstan
 			.OnItemsActivated(this, &SContentBrowser::OnItemsActivated)
 			.OnGetItemContextMenu(this, &SContentBrowser::GetItemContextMenu, EContentBrowserViewContext::AssetView)
 			.OnItemRenameCommitted(this, &SContentBrowser::OnItemRenameCommitted)
-			.OnShouldFilterItem(bEnablePrivateContentFeature ? FOnShouldFilterItem::CreateSP(this, &SContentBrowser::HandlePrivateContentFilter) : FOnShouldFilterItem())
+			.OnShouldFilterItem(bPrivateContentFilterEnabled
+				? FOnShouldFilterItem::CreateSP(this, &SContentBrowser::HandlePrivateContentFilter)
+				: FOnShouldFilterItem())
 			.FrontendFilters(FrontendFilters)
 			.TextFilter(TextFilter)
 			.ShowRedirectors_Lambda([this]() { return ContentBrowserUtils::ShouldShowRedirectors(FilterListPtr); })
@@ -3309,9 +3311,25 @@ void SContentBrowser::OnContentBrowserSettingsChanged(FName PropertyName)
 
 void SContentBrowser::OnConsoleVariableChanged()
 {
-	static const IConsoleVariable* EnablePublicAssetFeatureCVar = IConsoleManager::Get().FindConsoleVariable(TEXT("AssetTools.EnablePublicAssetFeature"));
-	const bool bEnablePrivateContentFeature = EnablePublicAssetFeatureCVar && EnablePublicAssetFeatureCVar->GetBool();
-	AssetViewPtr->SetShouldFilterItem(bEnablePrivateContentFeature ? FOnShouldFilterItem::CreateSP(this, &SContentBrowser::HandlePrivateContentFilter) : FOnShouldFilterItem());
+	UpdatePrivateContentFeatureEnabled(true /* bUpdateFilterIfChanged */);
+}
+
+void SContentBrowser::UpdatePrivateContentFeatureEnabled(bool bUpdateFilterIfChanged)
+{
+	static const IConsoleVariable* EnablePublicAssetFeatureCVar =
+		IConsoleManager::Get().FindConsoleVariable(TEXT("AssetTools.EnablePublicAssetFeature"));
+	const bool bShouldEnablePrivateContentFilter = EnablePublicAssetFeatureCVar
+		&& EnablePublicAssetFeatureCVar->GetBool();
+	if (bPrivateContentFilterEnabled != bShouldEnablePrivateContentFilter)
+	{
+		if (bUpdateFilterIfChanged)
+		{
+			AssetViewPtr->SetShouldFilterItem(bShouldEnablePrivateContentFilter
+				? FOnShouldFilterItem::CreateSP(this, &SContentBrowser::HandlePrivateContentFilter)
+				: FOnShouldFilterItem());
+		}
+		bPrivateContentFilterEnabled = bShouldEnablePrivateContentFilter;
+	}
 }
 
 FReply SContentBrowser::BackClicked()
diff --git a/Engine/Source/Editor/ContentBrowser/Private/SContentBrowser.h b/Engine/Source/Editor/ContentBrowser/Private/SContentBrowser.h
--- a/Engine/Source/Editor/ContentBrowser/Private/SContentBrowser.h
+++ b/Engine/Source/Editor/ContentBrowser/Private/SContentBrowser.h
@@ -640,6 +640,7 @@ private:
 	/** Initialize an editor config for this instance if one does not exist. */
 	FContentBrowserInstanceConfig* CreateEditorConfigIfRequired();
 
+	void UpdatePrivateContentFeatureEnabled(bool bUpdateFilterIfChanged);
 private:
 
 	/** The tab that contains this browser */
@@ -735,6 +736,9 @@ private:
 	/** True if source should not be changed from an outside source */
 	bool bIsLocked = false;
 
+	/** True if bShouldEnablePrivateContentFilter was true and we are filtering this->AssetView. */
+	bool bPrivateContentFilterEnabled = false;
+
 	/** Starting width of the PathView Box */
 	float PathViewBoxWidth = 150.f;
 
