Allowing asset pickers to create their own asset view options menu by name

diff --git a/Engine/Source/Editor/ContentBrowser/Private/SAssetView.cpp b/Engine/Source/Editor/ContentBrowser/Private/SAssetView.cpp
--- a/Engine/Source/Editor/ContentBrowser/Private/SAssetView.cpp
+++ b/Engine/Source/Editor/ContentBrowser/Private/SAssetView.cpp
@@ -3301,6 +3301,18 @@ void SAssetView::ExecutePaste()
 	}
 }
 
+FName SAssetView::GetAssetViewOptionsMenuName() const
+{
+	FName MenuName = "ContentBrowser.AssetViewOptions";
+
+	if(AssetViewOptionsProfile.IsSet())
+	{
+		MenuName = FName(MenuName.ToString() + "." + AssetViewOptionsProfile.GetValue().ToString());
+	}
+
+	return MenuName;
+}
+
 void SAssetView::ToggleShowAllFolder()
 {
 	const bool bNewValue = !IsShowingAllFolder();
@@ -3651,7 +3663,7 @@ FLinearColor SAssetView::GetThumbnailHintColorAndOpacity() const
 
 TSharedRef<SWidget> SAssetView::GetViewButtonContent()
 {
-	SAssetView::RegisterGetViewButtonMenu();
+	RegisterGetViewButtonMenu();
 
 	// Get all menu extenders for this context menu from the content browser module
 	FContentBrowserModule& ContentBrowserModule = FModuleManager::GetModuleChecked<FContentBrowserModule>( TEXT("ContentBrowser") );
@@ -3685,14 +3697,15 @@ TSharedRef<SWidget> SAssetView::GetViewButtonContent()
 		OnExtendAssetViewOptionsMenuContext.Execute(MenuContext);
 	}
 
-	return UToolMenus::Get()->GenerateWidget("ContentBrowser.AssetViewOptions", MenuContext);
+	return UToolMenus::Get()->GenerateWidget(GetAssetViewOptionsMenuName(), MenuContext);
 }
 
 void SAssetView::RegisterGetViewButtonMenu()
 {
-	if (!UToolMenus::Get()->IsMenuRegistered("ContentBrowser.AssetViewOptions"))
+	FName MenuName = GetAssetViewOptionsMenuName();
+	if (!UToolMenus::Get()->IsMenuRegistered(MenuName))
 	{
-		UToolMenu* Menu = UToolMenus::Get()->RegisterMenu("ContentBrowser.AssetViewOptions");
+		UToolMenu* Menu = UToolMenus::Get()->RegisterMenu(MenuName);
 		Menu->bCloseSelfOnly = true;
 		Menu->AddDynamicSection("DynamicContent", FNewToolMenuDelegate::CreateLambda([](UToolMenu* InMenu)
 		{
diff --git a/Engine/Source/Editor/ContentBrowser/Public/IContentBrowserSingleton.h b/Engine/Source/Editor/ContentBrowser/Public/IContentBrowserSingleton.h
--- a/Engine/Source/Editor/ContentBrowser/Public/IContentBrowserSingleton.h
+++ b/Engine/Source/Editor/ContentBrowser/Public/IContentBrowserSingleton.h
@@ -311,8 +311,8 @@ struct FAssetPickerConfig
 
 	/** If we display filters & set to true, we will add sections instead of sub-menus for other filters. Useful if the number of additional filters is small. */
 	bool bUseSectionsForCustomFilterCategories;
-
-	/** If specified, the asset view options menu will use this tool menu profile for customization purposes. */
+	
+	/** If specified, the asset view options menu will register a separate menu of the form "ContentBrowser.AssetViewOptions.[AssetViewOptionsProfile]". It will additionally use this name as a runtime tool menu profile. */
 	TOptional<FName> AssetViewOptionsProfile;
 
 	FAssetPickerConfig()
diff --git a/Engine/Source/Editor/ContentBrowser/Public/SAssetView.h b/Engine/Source/Editor/ContentBrowser/Public/SAssetView.h
--- a/Engine/Source/Editor/ContentBrowser/Public/SAssetView.h
+++ b/Engine/Source/Editor/ContentBrowser/Public/SAssetView.h
@@ -551,7 +551,7 @@ private:
 	FLinearColor GetThumbnailHintColorAndOpacity() const;
 
 	/** Register menu for when the view combo button is clicked */
-	static void RegisterGetViewButtonMenu();
+	void RegisterGetViewButtonMenu();
 
 	/** Fill in menu content for when the view combo button is clicked */
 	void PopulateViewButtonMenu(UToolMenu* Menu);
@@ -930,6 +930,8 @@ private:
 	/** Handler for Paste */
 	void ExecutePaste();
 
+	/** Constructs the name of the asset view options menu. */
+	FName GetAssetViewOptionsMenuName() const;
 private:
 	friend class FAssetViewFrontendFilterHelper;
 
