[Content Browser] Style refresh preparation - local CB StyleSet and CVar toggle.

diff --git a/Engine/Source/Editor/ContentBrowser/Private/AssetContextMenu.cpp b/Engine/Source/Editor/ContentBrowser/Private/AssetContextMenu.cpp
--- a/Engine/Source/Editor/ContentBrowser/Private/AssetContextMenu.cpp
+++ b/Engine/Source/Editor/ContentBrowser/Private/AssetContextMenu.cpp
@@ -20,6 +20,7 @@
 #include "ContentBrowserMenuContexts.h"
 #include "ContentBrowserModule.h"
 #include "ContentBrowserSingleton.h"
+#include "ContentBrowserStyle.h"
 #include "ContentBrowserUtils.h"
 #include "Engine/AssetManager.h"
 #include "Framework/Commands/GenericCommands.h"
@@ -580,7 +581,7 @@ void FAssetContextMenu::AddExploreMenuOptions(UToolMenu* Menu)
 				"FindInExplorer",
 				ContentBrowserUtils::GetExploreFolderText(),
 				LOCTEXT("FindInExplorerTooltip", "Finds this asset on disk"),
-				FSlateIcon(FAppStyle::GetAppStyleSetName(), "ContentBrowser.ShowInExplorer"),
+				FSlateIcon(UE::ContentBrowser::Private::FContentBrowserStyle::Get().GetStyleSetName(), "ContentBrowser.ShowInExplorer"),
 				FUIAction(
 					FExecuteAction::CreateSP(this, &FAssetContextMenu::ExecuteFindInExplorer),
 					FCanExecuteAction::CreateSP(this, &FAssetContextMenu::CanExecuteFindInExplorer)
@@ -635,7 +636,7 @@ bool FAssetContextMenu::AddPublicStateMenuOptions(UToolMenu* Menu)
 				"MarkSelectedAsPublic",
 				LOCTEXT("MarkSelectedAsPublic", "Mark Selected As Public"),
 				LOCTEXT("MarkSelectedAsPublicTooltip", "Sets all selected assets to be publicly available for reference by other plugins"),
-				FSlateIcon(FAppStyle::GetAppStyleSetName(), "ContentBrowser.AssetActions.PublicAssetToggle"),
+				FSlateIcon(UE::ContentBrowser::Private::FContentBrowserStyle::Get().GetStyleSetName(), "ContentBrowser.AssetActions.PublicAssetToggle"),
 				FUIAction(
 					FExecuteAction::CreateSP(this, &FAssetContextMenu::ExecuteBulkSetPublicAsset),
 					FCanExecuteAction::CreateSP(this, &FAssetContextMenu::CanExecuteBulkSetPublicAsset)
@@ -646,7 +647,7 @@ bool FAssetContextMenu::AddPublicStateMenuOptions(UToolMenu* Menu)
 				"MarkSelectedAsPrivate",
 				LOCTEXT("MarkSelectedAsPrivate", "Mark Selected As Private"),
 				LOCTEXT("MarkSelectedAsPrivateTooltip", "Sets all selected assets to be private and unavailable for reference by other plugins"),
-				FSlateIcon(FAppStyle::GetAppStyleSetName(), "ContentBrowser.AssetActions.PublicAssetToggle"),
+				FSlateIcon(UE::ContentBrowser::Private::FContentBrowserStyle::Get().GetStyleSetName(), "ContentBrowser.AssetActions.PublicAssetToggle"),
 				FUIAction(
 					FExecuteAction::CreateSP(this, &FAssetContextMenu::ExecuteBulkUnsetPublicAsset),
 					FCanExecuteAction::CreateSP(this, &FAssetContextMenu::CanExecuteBulkSetPublicAsset)
@@ -958,7 +959,7 @@ bool FAssetContextMenu::AddCollectionMenuOptions(UToolMenu* Menu)
 			FText::Format(LOCTEXT("ManageCollections_ToolTip", "Manage the collections that the selected {0}|plural(one=item belongs, other=items belong) to."), SelectedFiles.Num()),
 			FNewToolMenuDelegate::CreateStatic(&FManageCollectionsContextMenu::CreateManageCollectionsSubMenu, QuickAssetManagement),
 			false, // default value
-			FSlateIcon(FAppStyle::GetAppStyleSetName(), "ContentBrowser.Collections")
+			FSlateIcon(UE::ContentBrowser::Private::FContentBrowserStyle::Get().GetStyleSetName(), "ContentBrowser.Collections")
 			);
 
 		bHasAddedItems = true;
diff --git a/Engine/Source/Editor/ContentBrowser/Private/AssetViewWidgets.cpp b/Engine/Source/Editor/ContentBrowser/Private/AssetViewWidgets.cpp
--- a/Engine/Source/Editor/ContentBrowser/Private/AssetViewWidgets.cpp
+++ b/Engine/Source/Editor/ContentBrowser/Private/AssetViewWidgets.cpp
@@ -20,6 +20,7 @@
 #include "ContentBrowserItem.h"
 #include "ContentBrowserItemData.h"
 #include "ContentBrowserModule.h"
+#include "ContentBrowserStyle.h"
 #include "ContentBrowserUtils.h"
 #include "DragDropHandler.h"
 #include "EditorFramework/AssetImportData.h"
@@ -287,7 +288,7 @@ public:
 		SToolTip::Construct(
 			SToolTip::FArguments()
 			.TextMargin(1.0f)
-			.BorderImage(FAppStyle::GetBrush("ContentBrowser.TileViewTooltip.ToolTipBorder"))
+			.BorderImage(UE::ContentBrowser::Private::FContentBrowserStyle::Get().GetBrush("ContentBrowser.TileViewTooltip.ToolTipBorder"))
 			);
 	}
 
@@ -390,7 +391,7 @@ void SAssetViewItem::Construct( const FArguments& InArgs )
 
 	AssetItem->OnItemDataChanged().AddSP(this, &SAssetViewItem::OnAssetDataChanged);
 
-	AssetDirtyBrush = FAppStyle::GetBrush("ContentBrowser.ContentDirty");
+	AssetDirtyBrush = UE::ContentBrowser::Private::FContentBrowserStyle::Get().GetBrush("ContentBrowser.ContentDirty");
 
 	// Set our tooltip - this will refresh each time it's opened to make sure it's up-to-date
 	SetToolTip(SNew(SAssetViewItemToolTip).AssetViewItem(SharedThis(this)));
@@ -809,7 +810,7 @@ TSharedRef<SWidget> SAssetViewItem::CreateToolTipWidget() const
 				[
 					SNew(SBorder)
 					.Padding(6)
-					.BorderImage(FAppStyle::GetBrush("ContentBrowser.TileViewTooltip.ContentBorder"))
+					.BorderImage(UE::ContentBrowser::Private::FContentBrowserStyle::Get().Get().GetBrush("ContentBrowser.TileViewTooltip.ContentBorder"))
 					[
 						SNew(SVerticalBox)
 
@@ -825,7 +826,7 @@ TSharedRef<SWidget> SAssetViewItem::CreateToolTipWidget() const
 							[
 								SNew(STextBlock)
 								.Text(NameText)
-								.Font(FAppStyle::GetFontStyle("ContentBrowser.TileViewTooltip.NameFont"))
+								.Font(UE::ContentBrowser::Private::FContentBrowserStyle::Get().GetFontStyle("ContentBrowser.TileViewTooltip.NameFont"))
 							]
 
 							+ SHorizontalBox::Slot()
@@ -915,11 +916,11 @@ TSharedRef<SWidget> SAssetViewItem::CreateToolTipWidget() const
 					[
 						SNew(SBorder)
 						.Padding(6)
-						.BorderImage(FAppStyle::GetBrush("ContentBrowser.TileViewTooltip.ContentBorder"))
+						.BorderImage(UE::ContentBrowser::Private::FContentBrowserStyle::Get().GetBrush("ContentBrowser.TileViewTooltip.ContentBorder"))
 						[
 							SNew(STextBlock)
 							.WrapTextAt(700.0f)
-							.Font(FAppStyle::GetFontStyle("ContentBrowser.TileViewTooltip.AssetUserDescriptionFont"))
+							.Font(UE::ContentBrowser::Private::FContentBrowserStyle::Get().GetFontStyle("ContentBrowser.TileViewTooltip.AssetUserDescriptionFont"))
 							.Text(UserDescription)
 						]
 					];
@@ -931,7 +932,7 @@ TSharedRef<SWidget> SAssetViewItem::CreateToolTipWidget() const
 				[
 					SNew(SBorder)
 					.Padding(6)
-					.BorderImage(FAppStyle::GetBrush("ContentBrowser.TileViewTooltip.ContentBorder"))
+					.BorderImage(UE::ContentBrowser::Private::FContentBrowserStyle::Get().GetBrush("ContentBrowser.TileViewTooltip.ContentBorder"))
 					[
 						InfoBox
 					]
@@ -974,7 +975,7 @@ TSharedRef<SWidget> SAssetViewItem::CreateToolTipWidget() const
 						[
 							SNew(SBorder)
 							.Padding(FMargin(6, 2, 6, 6))
-							.BorderImage(FAppStyle::GetBrush("ContentBrowser.TileViewTooltip.ContentBorder"))
+							.BorderImage(UE::ContentBrowser::Private::FContentBrowserStyle::Get().GetBrush("ContentBrowser.TileViewTooltip.ContentBorder"))
 							[
 								CollectionPipsWrapBox
 							]
@@ -984,7 +985,7 @@ TSharedRef<SWidget> SAssetViewItem::CreateToolTipWidget() const
 
 			return SNew(SBorder)
 				.Padding(6)
-				.BorderImage( FAppStyle::GetBrush("ContentBrowser.TileViewTooltip.NonContentBorder") )
+				.BorderImage( UE::ContentBrowser::Private::FContentBrowserStyle::Get().GetBrush("ContentBrowser.TileViewTooltip.NonContentBorder") )
 				[
 					OverallTooltipVBox
 				];
@@ -1020,7 +1021,7 @@ TSharedRef<SWidget> SAssetViewItem::CreateToolTipWidget() const
 
 			return SNew(SBorder)
 				.Padding(6)
-				.BorderImage( FAppStyle::GetBrush("ContentBrowser.TileViewTooltip.NonContentBorder") )
+				.BorderImage( UE::ContentBrowser::Private::FContentBrowserStyle::Get().GetBrush("ContentBrowser.TileViewTooltip.NonContentBorder") )
 				[
 					SNew(SVerticalBox)
 
@@ -1030,7 +1031,7 @@ TSharedRef<SWidget> SAssetViewItem::CreateToolTipWidget() const
 					[
 						SNew(SBorder)
 						.Padding(6)
-						.BorderImage( FAppStyle::GetBrush("ContentBrowser.TileViewTooltip.ContentBorder") )
+						.BorderImage( UE::ContentBrowser::Private::FContentBrowserStyle::Get().GetBrush("ContentBrowser.TileViewTooltip.ContentBorder") )
 						[
 							SNew(SVerticalBox)
 
@@ -1046,7 +1047,7 @@ TSharedRef<SWidget> SAssetViewItem::CreateToolTipWidget() const
 								[
 									SNew(STextBlock)
 									.Text( FolderName )
-									.Font( FAppStyle::GetFontStyle("ContentBrowser.TileViewTooltip.NameFont") )
+									.Font( UE::ContentBrowser::Private::FContentBrowserStyle::Get().GetFontStyle("ContentBrowser.TileViewTooltip.NameFont") )
 								]
 
 								+SHorizontalBox::Slot()
@@ -1065,7 +1066,7 @@ TSharedRef<SWidget> SAssetViewItem::CreateToolTipWidget() const
 					[
 						SNew(SBorder)
 						.Padding(6)
-						.BorderImage( FAppStyle::GetBrush("ContentBrowser.TileViewTooltip.ContentBorder") )
+						.BorderImage( UE::ContentBrowser::Private::FContentBrowserStyle::Get().GetBrush("ContentBrowser.TileViewTooltip.ContentBorder") )
 						[
 							InfoBox
 						]
@@ -1701,7 +1702,7 @@ TSharedRef<SWidget> SAssetListItem::GenerateWidgetForColumn( const FName& Column
 				.Padding(0, 1)
 				[
 					SAssignNew(InlineRenameWidget, SInlineEditableTextBlock)
-					.Font(FAppStyle::GetFontStyle("ContentBrowser.AssetTileViewNameFont"))
+					.Font(UE::ContentBrowser::Private::FContentBrowserStyle::Get().GetFontStyle("ContentBrowser.AssetTileViewNameFont"))
 					.Text( GetNameText() )
 					.OnBeginTextEdit(this, &SAssetListItem::HandleBeginNameChange)
 					.OnTextCommitted(this, &SAssetListItem::HandleNameCommitted)
@@ -1718,7 +1719,7 @@ TSharedRef<SWidget> SAssetListItem::GenerateWidgetForColumn( const FName& Column
 				[
 					// Class
 					SAssignNew(ClassTextWidget, STextBlock)
-					.Font(FAppStyle::GetFontStyle("ContentBrowser.AssetListViewClassFont"))
+					.Font(UE::ContentBrowser::Private::FContentBrowserStyle::Get().GetFontStyle("ContentBrowser.AssetListViewClassFont"))
 					.Text(GetAssetClassText())
 					.HighlightText(HighlightText)
 					.ColorAndOpacity(this, &SAssetListItem::GetColumnTextColor, InIsSelected)
@@ -1896,7 +1897,7 @@ void SAssetTileItem::Construct( const FArguments& InArgs )
 			[
 				SNew(SBorder)
 				.Padding(0)
-				.BorderImage(IsFolder() ? TAttribute<const FSlateBrush*>(this, &SAssetTileItem::GetFolderBackgroundImage) : FAppStyle::Get().GetBrush("ContentBrowser.AssetTileItem.ThumbnailAreaBackground"))
+				.BorderImage(IsFolder() ? TAttribute<const FSlateBrush*>(this, &SAssetTileItem::GetFolderBackgroundImage) : UE::ContentBrowser::Private::FContentBrowserStyle::Get().Get().GetBrush("ContentBrowser.AssetTileItem.ThumbnailAreaBackground"))
 				[
 					SNew(SVerticalBox)
 					// Thumbnail
@@ -1941,7 +1942,7 @@ void SAssetTileItem::Construct( const FArguments& InArgs )
 								[
 									SAssignNew(ClassTextWidget, STextBlock)
 									.Visibility(this, &SAssetTileItem::GetAssetClassLabelVisibility)
-									.TextStyle(FAppStyle::Get(), "ContentBrowser.ClassFont")
+									.TextStyle(UE::ContentBrowser::Private::FContentBrowserStyle::Get().Get(), "ContentBrowser.ClassFont")
 									.OverflowPolicy(ETextOverflowPolicy::Ellipsis)
 									.Text(this, &SAssetTileItem::GetAssetClassText)
 									.ColorAndOpacity(this, &SAssetTileItem::GetAssetClassLabelTextColor)
@@ -2340,11 +2341,11 @@ TSharedRef<SWidget> SAssetColumnItem::GenerateWidgetForColumn( const FName& Colu
 		{
 			if(ContentBrowserUtils::IsItemDeveloperContent(AssetItem->GetItem()))
 			{
-				IconBrush = FAppStyle::GetBrush("ContentBrowser.ColumnViewDeveloperFolderIcon");
+				IconBrush = UE::ContentBrowser::Private::FContentBrowserStyle::Get().GetBrush("ContentBrowser.ColumnViewDeveloperFolderIcon");
 			}
 			else
 			{
-				IconBrush = FAppStyle::GetBrush("ContentBrowser.ColumnViewFolderIcon");
+				IconBrush = UE::ContentBrowser::Private::FContentBrowserStyle::Get().GetBrush("ContentBrowser.ColumnViewFolderIcon");
 			}
 		}
 		else
@@ -2355,7 +2356,7 @@ TSharedRef<SWidget> SAssetColumnItem::GenerateWidgetForColumn( const FName& Colu
 			}
 			else
 			{
-				IconBrush = FAppStyle::GetBrush("ContentBrowser.ColumnViewAssetIcon");
+				IconBrush = UE::ContentBrowser::Private::FContentBrowserStyle::Get().GetBrush("ContentBrowser.ColumnViewAssetIcon");
 			}
 		}
 
diff --git a/Engine/Source/Editor/ContentBrowser/Private/ContentBrowserSingleton.cpp b/Engine/Source/Editor/ContentBrowser/Private/ContentBrowserSingleton.cpp
--- a/Engine/Source/Editor/ContentBrowser/Private/ContentBrowserSingleton.cpp
+++ b/Engine/Source/Editor/ContentBrowser/Private/ContentBrowserSingleton.cpp
@@ -14,6 +14,7 @@
 #include "ContentBrowserCommands.h"
 #include "ContentBrowserDataSubsystem.h"
 #include "ContentBrowserModule.h"
+#include "ContentBrowserStyle.h"
 #include "ContentBrowserUtils.h"
 #include "CoreGlobals.h"
 #include "Delegates/Delegate.h"
@@ -91,7 +92,7 @@ FContentBrowserSingleton::FContentBrowserSingleton()
 	: CollectionAssetRegistryBridge(MakeShared<FCollectionAssetRegistryBridge>())
 	, SettingsStringID(0)
 {
-	const FSlateIcon ContentBrowserIcon(FAppStyle::Get().GetStyleSetName(), "ContentBrowser.TabIcon");
+	const FSlateIcon ContentBrowserIcon(UE::ContentBrowser::Private::FContentBrowserStyle::Get().GetStyleSetName(), "ContentBrowser.TabIcon");
 	const IWorkspaceMenuStructure& MenuStructure = WorkspaceMenu::GetMenuStructure();
 	TSharedRef<FWorkspaceItem> ContentBrowserGroup = MenuStructure.GetLevelEditorCategory()->AddGroup(
 		"ContentBrowser",
@@ -1246,8 +1247,6 @@ void FContentBrowserSingleton::PopulateConfigValues()
 void FContentBrowserSingleton::GetContentBrowserSubMenu(UToolMenu* Menu, TSharedRef<FWorkspaceItem> ContentBrowserGroup)
 {
 	// Register the tab spawners for all content browsers
-	const FSlateIcon ContentBrowserIcon(FAppStyle::Get().GetStyleSetName(), "ContentBrowser.TabIcon");
-
 	FToolMenuSection& Section = Menu->AddSection("ContentBrowser");
 
 	Section.AddMenuEntry("FocusContentBrowser",
diff --git a/Engine/Source/Editor/ContentBrowser/Private/ContentBrowserStyle.cpp b/Engine/Source/Editor/ContentBrowser/Private/ContentBrowserStyle.cpp
new file mode 100644
--- /dev/null
+++ b/Engine/Source/Editor/ContentBrowser/Private/ContentBrowserStyle.cpp
@@ -0,0 +1,338 @@
+// Copyright Epic Games, Inc. All Rights Reserved.
+
+#include "ContentBrowserStyle.h"
+
+#include "Styling/CoreStyle.h"
+#include "Styling/SlateStyleMacros.h"
+#include "Styling/SlateStyleRegistry.h"
+#include "Styling/StarshipCoreStyle.h"
+#include "Styling/StyleColors.h"
+#include "Styling/ToolBarStyle.h"
+
+namespace UE::ContentBrowser::Private
+{
+	static bool bEnableNewStyle = false;
+	static FAutoConsoleVariableRef CVarEnableNewStyle(
+		TEXT("ContentBrowser.EnableNewStyle"),
+		bEnableNewStyle,
+		TEXT("Whether or not to enable the Content Browser restyle (UE 5.5+).")
+	);
+
+	FName FContentBrowserStyle::StyleName("ContentBrowser");
+
+	FContentBrowserStyle& FContentBrowserStyle::Get()
+	{
+		static FContentBrowserStyle Instance;
+		return Instance;
+	}
+
+	const FName& FContentBrowserStyle::GetStyleSetName() const
+	{
+		return StyleName;
+	}
+
+	bool FContentBrowserStyle::IsNewStyleEnabled()
+	{
+		return bEnableNewStyle;
+	}
+
+	FContentBrowserStyle::FContentBrowserStyle()
+		: FSlateStyleSet(StyleName)
+	{
+		SetParentStyleName(FAppStyle::GetAppStyleSetName());
+
+		// Init from ParentStyle
+		{
+			const ISlateStyle* ParentStyle = GetParentStyle();
+
+			DefaultForeground = ParentStyle->GetSlateColor("DefaultForeground");
+			InvertedForeground = ParentStyle->GetSlateColor("InvertedForeground");
+			SelectionColor = ParentStyle->GetSlateColor("SelectionColor");
+			SelectionColor_Pressed = ParentStyle->GetSlateColor("SelectionColor_Pressed");
+
+			NormalText = ParentStyle->GetWidgetStyle<FTextBlockStyle>("NormalText");
+			Button = ParentStyle->GetWidgetStyle<FButtonStyle>("Button");
+		}
+
+		SetContentRoot( FPaths::EngineContentDir() / TEXT("Editor/Slate") );
+		SetCoreContentRoot(FPaths::EngineContentDir() / TEXT("Slate"));
+
+		// Originally in FStarshipEditorStyle::FStyle::SetupContentBrowserStyle()
+		// @todo: Remove, deprecate or redirect from FStarshipEditorStyle::FStyle::SetupContentBrowserStyle()
+
+		// Tab and menu icon
+		Set("ContentBrowser.TabIcon", new IMAGE_BRUSH_SVG("Starship/Common/ContentBrowser", CoreStyleConstants::Icon16x16));
+		Set("ContentBrowser.PrivateContentEdit", new IMAGE_BRUSH("Icons/hiererchy_16x", CoreStyleConstants::Icon16x16));
+
+		// Sources View
+		Set("ContentBrowser.SourceTitleFont", DEFAULT_FONT("Regular", 12));
+
+		Set("ContentBrowser.SourceTreeItemFont", FStarshipCoreStyle::GetCoreStyle().GetFontStyle("NormalFont"));
+		Set("ContentBrowser.SourceTreeRootItemFont", FStarshipCoreStyle::GetCoreStyle().GetFontStyle("NormalFont"));
+
+		Set("ContentBrowser.BreadcrumbPathPickerFolder",	new CORE_IMAGE_BRUSH_SVG("Starship/Common/folder-closed", CoreStyleConstants::Icon16x16));
+
+		Set("ContentBrowser.AssetTreeFolderClosed", new CORE_IMAGE_BRUSH_SVG("Starship/Common/folder-closed", CoreStyleConstants::Icon16x16));
+		Set("ContentBrowser.AssetTreeFolderOpen", new CORE_IMAGE_BRUSH_SVG("Starship/Common/folder-open", CoreStyleConstants::Icon16x16));
+
+		Set("ContentBrowser.AssetTreeFolderClosedVirtual", new CORE_IMAGE_BRUSH_SVG("Starship/Common/folder-virtual-closed", CoreStyleConstants::Icon16x16));
+		Set("ContentBrowser.AssetTreeFolderOpenVirtual",	new CORE_IMAGE_BRUSH_SVG("Starship/Common/folder-virtual-open", CoreStyleConstants::Icon16x16));
+		Set("ContentBrowser.AssetTreeFolderOpenDeveloper", new IMAGE_BRUSH_SVG("Starship/ContentBrowser/folder-developer-open", CoreStyleConstants::Icon16x16));
+		Set("ContentBrowser.AssetTreeFolderClosedDeveloper",	new IMAGE_BRUSH_SVG("Starship/ContentBrowser/folder-developer", CoreStyleConstants::Icon16x16));
+		Set("ContentBrowser.AssetTreeFolderOpenCode", new IMAGE_BRUSH_SVG("Starship/ContentBrowser/folder-code-open", CoreStyleConstants::Icon16x16));
+		Set("ContentBrowser.AssetTreeFolderClosedCode", new IMAGE_BRUSH_SVG("Starship/ContentBrowser/folder-code", CoreStyleConstants::Icon16x16));
+		Set("ContentBrowser.AssetTreeFolderOpenPluginRoot", new IMAGE_BRUSH_SVG("Starship/ContentBrowser/folder-pluginroot-open", CoreStyleConstants::Icon16x16));
+		Set("ContentBrowser.AssetTreeFolderClosedPluginRoot", new IMAGE_BRUSH_SVG("Starship/ContentBrowser/folder-pluginroot", CoreStyleConstants::Icon16x16));
+
+		Set("ContentBrowser.DefaultFolderColor", FStyleColors::AccentFolder);
+
+		// Asset list view
+		Set("ContentBrowser.AssetListViewNameFont", DEFAULT_FONT("Regular", 12));
+		Set("ContentBrowser.AssetListViewNameFontDirty", DEFAULT_FONT("Bold", 12));
+		Set("ContentBrowser.AssetListViewClassFont", DEFAULT_FONT("Light", 10));
+
+		// Asset picker
+		Set("ContentBrowser.NoneButton",
+			FButtonStyle(Button)
+			.SetNormal(FSlateNoResource())
+			.SetHovered(BOX_BRUSH("Common/Selection", 8.0f / 32.0f, SelectionColor))
+			.SetPressed(BOX_BRUSH("Common/Selection", 8.0f / 32.0f, SelectionColor_Pressed))
+		);
+		Set("ContentBrowser.NoneButtonText",
+			FTextBlockStyle(NormalText)
+			.SetFont(DEFAULT_FONT("Regular", 12))
+			.SetColorAndOpacity(FLinearColor::White)
+		);
+
+		// Tile view
+		Set("ContentBrowser.AssetTileViewNameFont", DEFAULT_FONT("Regular", 9));
+		Set("ContentBrowser.AssetTileViewNameFontSmall", DEFAULT_FONT("VeryLight", 8));
+		Set("ContentBrowser.AssetTileViewNameFontVerySmall", DEFAULT_FONT("VeryLight", 7));
+		Set("ContentBrowser.AssetTileViewNameFontDirty", FStyleFonts::Get().SmallBold);
+
+		Set("ContentBrowser.AssetListView.ColumnListTableRow",
+			FTableRowStyle()
+			.SetEvenRowBackgroundBrush(FSlateColorBrush(FStyleColors::Recessed))
+			.SetEvenRowBackgroundHoveredBrush(FSlateColorBrush(FStyleColors::SelectHover))
+			.SetOddRowBackgroundBrush(FSlateColorBrush(FStyleColors::Background))
+			.SetOddRowBackgroundHoveredBrush(FSlateColorBrush(FStyleColors::SelectHover))
+			.SetSelectorFocusedBrush(BORDER_BRUSH("Common/Selector", FMargin(4.0f / 16.0f), FStyleColors::Select))
+			.SetActiveBrush(IMAGE_BRUSH("Common/Selection", CoreStyleConstants::Icon8x8, FStyleColors::Select))
+			.SetActiveHoveredBrush(IMAGE_BRUSH("Common/Selection", CoreStyleConstants::Icon8x8, FStyleColors::Select))
+			.SetInactiveBrush(IMAGE_BRUSH("Common/Selection", CoreStyleConstants::Icon8x8, FStyleColors::SelectInactive))
+			.SetInactiveHoveredBrush(IMAGE_BRUSH("Common/Selection", CoreStyleConstants::Icon8x8, FStyleColors::SelectInactive))
+			.SetTextColor(DefaultForeground)
+			.SetSelectedTextColor(InvertedForeground)
+		);
+
+		Set("ContentBrowser.AssetListView.TileTableRow",
+			FTableRowStyle()
+			.SetEvenRowBackgroundBrush(FSlateNoResource())
+			.SetEvenRowBackgroundHoveredBrush(FSlateNoResource())
+			.SetOddRowBackgroundBrush(FSlateNoResource())
+			.SetOddRowBackgroundHoveredBrush(FSlateNoResource())
+			.SetSelectorFocusedBrush(FSlateNoResource())
+			.SetActiveBrush(FSlateNoResource())
+			.SetActiveHoveredBrush(FSlateNoResource())
+			.SetInactiveBrush(FSlateNoResource())
+			.SetInactiveHoveredBrush(FSlateNoResource())
+			.SetTextColor(DefaultForeground)
+			.SetSelectedTextColor(DefaultForeground)
+		);
+
+		Set("ContentBrowser.TileViewTooltip.ToolTipBorder", new FSlateColorBrush(FLinearColor::Black));
+		Set("ContentBrowser.TileViewTooltip.NonContentBorder", new BOX_BRUSH("Docking/TabContentArea", FMargin(4.0f / 16.0f)));
+		Set("ContentBrowser.TileViewTooltip.ContentBorder", new FSlateColorBrush(FStyleColors::Panel));
+		Set("ContentBrowser.TileViewTooltip.PillBorder",	new FSlateRoundedBoxBrush(FStyleColors::Transparent, 10.0f, FStyleColors::White, 1.0f));
+		Set("ContentBrowser.TileViewTooltip.UnsupportedAssetPillBorder",	new FSlateRoundedBoxBrush(FStyleColors::Transparent, 10.0f, FStyleColors::Warning, 1.0f));
+		Set("ContentBrowser.TileViewTooltip.NameFont", DEFAULT_FONT("Regular", 12));
+		Set("ContentBrowser.TileViewTooltip.AssetUserDescriptionFont", DEFAULT_FONT("Regular", 12));
+
+		// Columns view
+		Set("ContentBrowser.SortUp", new IMAGE_BRUSH("Common/SortUpArrow", CoreStyleConstants::Icon8x4));
+		Set("ContentBrowser.SortDown", new IMAGE_BRUSH("Common/SortDownArrow", CoreStyleConstants::Icon8x4));
+
+		// Filter List - These are aliases for SBasicFilterBar styles in StarshipCoreStyle for backwards compatibility
+		Set("ContentBrowser.FilterImage", new IMAGE_BRUSH_SVG("Starship/CoreWidgets/FilterBar/FilterColorSegment", FVector2D(8, 22)));
+		Set("ContentBrowser.FilterBackground", new FSlateRoundedBoxBrush(FStyleColors::Secondary, 3.0f));
+
+		Set("ContentBrowser.FilterButton", FStarshipCoreStyle::GetCoreStyle().GetWidgetStyle<FCheckBoxStyle>("FilterBar.FilterButton"));
+		Set("ContentBrowser.FilterToolBar",	FStarshipCoreStyle::GetCoreStyle().GetWidgetStyle<FToolBarStyle>("FilterBar.FilterToolBar"));
+
+		// Sources view
+		Set("ContentBrowser.Sources.Paths", new IMAGE_BRUSH("ContentBrowser/Sources_Paths_16x", CoreStyleConstants::Icon16x16));
+		Set("ContentBrowser.Sources.Collections", new IMAGE_BRUSH("ContentBrowser/Sources_Collections_Standard_16x", CoreStyleConstants::Icon16x16));
+		Set("ContentBrowser.Sources.Collections.Compact", new IMAGE_BRUSH("ContentBrowser/Sources_Collections_Compact_16x", CoreStyleConstants::Icon16x16));
+
+		// Asset tags (common)
+		Set("ContentBrowser.AssetTagBackground", new FSlateRoundedBoxBrush(FStyleColors::White, 2.0f));
+
+		// Asset tags (standard)
+		Set("ContentBrowser.AssetTagButton", FCheckBoxStyle()
+			.SetUncheckedImage(IMAGE_BRUSH("ContentBrowser/AssetTagCheckbox_Flat", FVector2D(14.0f, 28.0f)))
+			.SetUncheckedHoveredImage(IMAGE_BRUSH("ContentBrowser/AssetTagCheckbox_Flat", FVector2D(14.0f, 28.0f), FLinearColor(0.5f, 0.5f, 0.5f, 1.0f)))
+			.SetUncheckedPressedImage(IMAGE_BRUSH("ContentBrowser/AssetTagCheckbox_Flat", FVector2D(14.0f, 28.0f), FLinearColor(0.5f, 0.5f, 0.5f, 1.0f)))
+			.SetUndeterminedImage(IMAGE_BRUSH("ContentBrowser/AssetTagCheckbox_Flat", FVector2D(14.0f, 28.0f)))
+			.SetUndeterminedHoveredImage(IMAGE_BRUSH("ContentBrowser/AssetTagCheckbox_Flat", FVector2D(14.0f, 28.0f), FLinearColor(0.5f, 0.5f, 0.5f, 1.0f)))
+			.SetUndeterminedPressedImage(IMAGE_BRUSH("ContentBrowser/AssetTagCheckbox_Flat", FVector2D(14.0f, 28.0f), FLinearColor(0.5f, 0.5f, 0.5f, 1.0f)))
+			.SetCheckedImage(IMAGE_BRUSH("ContentBrowser/AssetTagCheckbox_Flat", FVector2D(14.0f, 28.0f)))
+			.SetCheckedHoveredImage(IMAGE_BRUSH("ContentBrowser/AssetTagCheckbox_Flat", FVector2D(14.0f, 28.0f), FLinearColor(0.5f, 0.5f, 0.5f, 1.0f)))
+			.SetCheckedPressedImage(IMAGE_BRUSH("ContentBrowser/AssetTagCheckbox_Flat", FVector2D(14.0f, 28.0f), FLinearColor(0.5f, 0.5f, 0.5f, 1.0f)))
+			.SetPadding(0.0f)
+		);
+
+		Set("ContentBrowser.AssetTagNamePadding", FMargin(4.0f));
+		Set("ContentBrowser.AssetTagCountPadding", FMargin(4.0f));
+
+		// Asset tags (compact)
+		Set("ContentBrowser.AssetTagButton.Compact", FCheckBoxStyle()
+			.SetUncheckedImage(IMAGE_BRUSH("ContentBrowser/AssetTagCheckbox_Flat_Compact", FVector2D(10.0f, 20.0f)))
+			.SetUncheckedHoveredImage(IMAGE_BRUSH("ContentBrowser/AssetTagCheckbox_Flat_Compact", FVector2D(10.0f, 20.0f), FLinearColor(0.5f, 0.5f, 0.5f, 1.0f)))
+			.SetUncheckedPressedImage(IMAGE_BRUSH("ContentBrowser/AssetTagCheckbox_Flat_Compact", FVector2D(10.0f, 20.0f), FLinearColor(0.5f, 0.5f, 0.5f, 1.0f)))
+			.SetUndeterminedImage(IMAGE_BRUSH("ContentBrowser/AssetTagCheckbox_Flat_Compact", FVector2D(10.0f, 20.0f)))
+			.SetUndeterminedHoveredImage(IMAGE_BRUSH("ContentBrowser/AssetTagCheckbox_Flat_Compact", FVector2D(10.0f, 20.0f), FLinearColor(0.5f, 0.5f, 0.5f, 1.0f)))
+			.SetUndeterminedPressedImage(IMAGE_BRUSH("ContentBrowser/AssetTagCheckbox_Flat_Compact", FVector2D(10.0f, 20.0f), FLinearColor(0.5f, 0.5f, 0.5f, 1.0f)))
+			.SetCheckedImage(IMAGE_BRUSH("ContentBrowser/AssetTagCheckbox_Flat_Compact", FVector2D(10.0f, 20.0f)))
+			.SetCheckedHoveredImage(IMAGE_BRUSH("ContentBrowser/AssetTagCheckbox_Flat_Compact", FVector2D(10.0f, 20.0f), FLinearColor(0.5f, 0.5f, 0.5f, 1.0f)))
+			.SetCheckedPressedImage(IMAGE_BRUSH("ContentBrowser/AssetTagCheckbox_Flat_Compact", FVector2D(10.0f, 20.0f), FLinearColor(0.5f, 0.5f, 0.5f, 1.0f)))
+			.SetPadding(0.0f)
+		);
+
+		Set("ContentBrowser.AssetTagNamePadding.Compact", FMargin(2.0f));
+		Set("ContentBrowser.AssetTagCountPadding.Compact", FMargin(2.0f));
+
+
+		Set("ContentBrowser.PrimitiveCustom", new IMAGE_BRUSH("ContentBrowser/ThumbnailCustom", CoreStyleConstants::Icon32x32));
+		Set("ContentBrowser.PrimitiveSphere", new IMAGE_BRUSH("ContentBrowser/ThumbnailSphere", CoreStyleConstants::Icon32x32));
+		Set("ContentBrowser.PrimitiveCube", new IMAGE_BRUSH("ContentBrowser/ThumbnailCube", CoreStyleConstants::Icon32x32));
+		Set("ContentBrowser.PrimitivePlane", new IMAGE_BRUSH("ContentBrowser/ThumbnailPlane", CoreStyleConstants::Icon32x32));
+		Set("ContentBrowser.PrimitiveCylinder", new IMAGE_BRUSH("ContentBrowser/ThumbnailCylinder", CoreStyleConstants::Icon32x32));
+		Set("ContentBrowser.ResetPrimitiveToDefault", new IMAGE_BRUSH_SVG("Starship/Common/Undo", CoreStyleConstants::Icon20x20));
+
+		Set("ContentBrowser.TopBar.Font", FTextBlockStyle(NormalText)
+			.SetFont(DEFAULT_FONT("Bold", 11))
+			.SetColorAndOpacity(FLinearColor(1.0f, 1.0f, 1.0f))
+			.SetHighlightColor(FLinearColor(1.0f, 1.0f, 1.0f))
+			.SetShadowOffset(FVector2D(1, 1))
+			.SetShadowColorAndOpacity(FLinearColor(0, 0, 0, 0.9f)));
+
+		Set("ContentBrowser.ClassFont", FTextBlockStyle(NormalText)
+			.SetFont(DEFAULT_FONT("Regular", 7)));
+
+
+		Set("ContentBrowser.AddContent", new IMAGE_BRUSH_SVG("Starship/ContentBrowser/ContentPack", CoreStyleConstants::Icon20x20));
+		Set("ContentBrowser.ImportPackage", new IMAGE_BRUSH("Icons/icon_Import_40x", CoreStyleConstants::Icon25x25));
+
+		// Asset Context Menu
+		Set("ContentBrowser.AssetActions", new CORE_IMAGE_BRUSH("Icons/icon_tab_Tools_16x", CoreStyleConstants::Icon16x16));
+		Set("ContentBrowser.AssetActions.Edit", new IMAGE_BRUSH("Icons/Edit/icon_Edit_16x", CoreStyleConstants::Icon16x16));
+		Set("ContentBrowser.AssetActions.Delete", new IMAGE_BRUSH("Icons/icon_delete_16px", CoreStyleConstants::Icon16x16, FLinearColor(0.4f, 0.5f, 0.7f, 1.0f)));
+		Set("ContentBrowser.AssetActions.Rename", new IMAGE_BRUSH("Icons/Icon_Asset_Rename_16x", CoreStyleConstants::Icon16x16));
+		Set("ContentBrowser.AssetActions.Duplicate", new IMAGE_BRUSH("Icons/Edit/icon_Edit_Duplicate_16x", CoreStyleConstants::Icon16x16));
+		Set("ContentBrowser.AssetActions.OpenSourceLocation", new IMAGE_BRUSH("Icons/icon_Asset_Open_Source_Location_16x", CoreStyleConstants::Icon16x16));
+		Set("ContentBrowser.AssetActions.OpenInExternalEditor", new IMAGE_BRUSH("Icons/icon_Asset_Open_In_External_Editor_16x", CoreStyleConstants::Icon16x16));
+		Set("ContentBrowser.AssetActions.PublicAssetToggle", new IMAGE_BRUSH("Icons/hiererchy_16x", CoreStyleConstants::Icon16x16));
+		Set("ContentBrowser.AssetActions.ReimportAsset", new IMAGE_BRUSH("Icons/icon_TextureEd_Reimport_40x", CoreStyleConstants::Icon16x16));
+		Set("ContentBrowser.AssetActions.GoToCodeForAsset", new IMAGE_BRUSH("GameProjectDialog/feature_code_32x", CoreStyleConstants::Icon16x16));
+		Set("ContentBrowser.AssetActions.FindAssetInWorld", new IMAGE_BRUSH("/Icons/icon_Genericfinder_16x", CoreStyleConstants::Icon16x16));
+		Set("ContentBrowser.AssetActions.CreateThumbnail", new IMAGE_BRUSH("Icons/icon_Asset_Create_Thumbnail_16x", CoreStyleConstants::Icon16x16));
+		Set("ContentBrowser.AssetActions.DeleteThumbnail", new IMAGE_BRUSH("Icons/icon_Asset_Delete_Thumbnail_16x", CoreStyleConstants::Icon16x16));
+		Set("ContentBrowser.AssetActions.GenericFind", new IMAGE_BRUSH("Icons/icon_Genericfinder_16x", CoreStyleConstants::Icon16x16));
+		Set("ContentBrowser.AssetLocalization", new IMAGE_BRUSH("Icons/icon_localization_16x", CoreStyleConstants::Icon16x16));
+		Set("ContentBrowser.AssetActions.VolumeTexture", new IMAGE_BRUSH_SVG("Starship/AssetActions/volume-texture", CoreStyleConstants::Icon16x16));
+		Set("ContentBrowser.AssetActions.RemoveVertexColors", new IMAGE_BRUSH_SVG("Starship/AssetActions/remove-vertex-colors", CoreStyleConstants::Icon16x16));
+
+		// ContentBrowser Commands Icons
+		Set("ContentBrowser.AssetViewCopyObjectPath", new IMAGE_BRUSH_SVG("../../Slate/Starship/Common/Copy", CoreStyleConstants::Icon16x16));
+		Set("ContentBrowser.AssetViewCopyPackageName", new IMAGE_BRUSH_SVG("../../Slate/Starship/Common/Copy", CoreStyleConstants::Icon16x16));
+
+		// Misc
+		/** Should be moved, shared */ Set("ContentBrowser.ThumbnailShadow", new BOX_BRUSH("ContentBrowser/ThumbnailShadow", FMargin(4.0f / 64.0f)));
+
+
+
+		Set("ContentBrowser.ColumnViewAssetIcon", new IMAGE_BRUSH("Icons/doc_16x", CoreStyleConstants::Icon16x16));
+
+		Set("ContentBrowser.ColumnViewFolderIcon", new CORE_IMAGE_BRUSH_SVG("Starship/Common/folder-closed", CoreStyleConstants::Icon16x16));
+		Set("ContentBrowser.ColumnViewDeveloperFolderIcon", new IMAGE_BRUSH_SVG("Starship/ContentBrowser/folder-developer", CoreStyleConstants::Icon16x16));
+
+		Set("ContentBrowser.ListViewFolderIcon", new IMAGE_BRUSH_SVG("Starship/ContentBrowser/folder", FVector2D(64, 64)));
+		Set("ContentBrowser.ListViewVirtualFolderIcon", new IMAGE_BRUSH_SVG("Starship/ContentBrowser/FolderLargeVirtual", FVector2D(64, 64)));
+		Set("ContentBrowser.ListViewVirtualFolderShadow", new IMAGE_BRUSH("Starship/ContentBrowser/FolderLargeVirtualShadow", FVector2D(256, 256)));
+		Set("ContentBrowser.ListViewDeveloperFolderIcon", new IMAGE_BRUSH_SVG("Starship/ContentBrowser/FolderLargeDeveloper", FVector2D(64, 64)));
+		Set("ContentBrowser.ListViewCodeFolderIcon", new IMAGE_BRUSH_SVG("Starship/ContentBrowser/FolderLargeCode", FVector2D(64, 64)));
+		Set("ContentBrowser.ListViewPluginFolderIcon", new IMAGE_BRUSH_SVG("Starship/ContentBrowser/FolderLargePlugin", FVector2D(64, 64)));
+
+		Set("ContentBrowser.AssetTileItem.FolderAreaHoveredBackground", new FSlateRoundedBoxBrush(FStyleColors::Secondary, 4.0f));
+		Set("ContentBrowser.AssetTileItem.FolderAreaSelectedBackground", new FSlateRoundedBoxBrush(FStyleColors::Secondary, 4.0f));
+		Set("ContentBrowser.AssetTileItem.FolderAreaSelectedHoverBackground", new FSlateRoundedBoxBrush(FStyleColors::Secondary, 4.0f));
+		Set("ContentBrowser.AssetTileItem.ThumbnailAreaBackground", new FSlateRoundedBoxBrush(FStyleColors::Recessed, 4.0f));
+		Set("ContentBrowser.AssetTileItem.NameAreaBackground", new FSlateRoundedBoxBrush(FStyleColors::Secondary, FVector4(0.0f, 0.0f, 4.0f, 4.0f)));
+		Set("ContentBrowser.AssetTileItem.NameAreaHoverBackground", new FSlateRoundedBoxBrush(FStyleColors::Hover, FVector4(0.0f, 0.0f, 4.0f, 4.0f)));
+		Set("ContentBrowser.AssetTileItem.NameAreaSelectedBackground", new FSlateRoundedBoxBrush(FStyleColors::Primary, FVector4(0.0f, 0.0f, 4.0f, 4.0f)));
+		Set("ContentBrowser.AssetTileItem.NameAreaSelectedHoverBackground", new FSlateRoundedBoxBrush(FStyleColors::PrimaryHover, FVector4(0.0f, 0.0f, 4.0f, 4.0f)));
+
+		{
+			FLinearColor TransparentPrimary = FStyleColors::Primary.GetSpecifiedColor();
+			TransparentPrimary.A = 0.0f;
+			Set("ContentBrowser.AssetTileItem.SelectedBorder", new FSlateRoundedBoxBrush(TransparentPrimary, 4.0f, FStyleColors::Primary, 1.0f));
+
+			FLinearColor TransparentPrimaryHover = FStyleColors::PrimaryHover.GetSpecifiedColor();
+			TransparentPrimaryHover.A = 0.0f;
+			Set("ContentBrowser.AssetTileItem.SelectedHoverBorder", new FSlateRoundedBoxBrush(TransparentPrimaryHover, 4.0f, FStyleColors::PrimaryHover, 1.0f));
+
+			FLinearColor TransparentHover = FStyleColors::Hover.GetSpecifiedColor();
+			TransparentHover.A = 0.0f;
+			Set("ContentBrowser.AssetTileItem.HoverBorder", new FSlateRoundedBoxBrush(TransparentHover, 4.0f, FStyleColors::Hover, 1.0f));
+		}
+
+		Set("ContentBrowser.AssetTileItem.DropShadow", new CORE_BOX_BRUSH("Starship/ContentBrowser/drop-shadow", FMargin(4.0f / 64.0f)));
+		Set("ContentBrowser.FolderItem.DropShadow", new IMAGE_BRUSH("Starship/ContentBrowser/folder-drop-shadow", FVector2D(256, 256)));
+
+		Set("ContentBrowser.ShowSourcesView", new IMAGE_BRUSH_SVG("Starship/ContentBrowser/file-tree", CoreStyleConstants::Icon16x16));
+		Set("ContentBrowser.HideSourcesView", new IMAGE_BRUSH_SVG("Starship/ContentBrowser/file-tree-open", CoreStyleConstants::Icon16x16));
+
+		Set("ContentBrowser.DirectoryUp", new IMAGE_BRUSH("Icons/icon_folder_up_16x", CoreStyleConstants::Icon16x16));
+		Set("ContentBrowser.PathPickerButton", new IMAGE_BRUSH("Icons/ellipsis_12x", CoreStyleConstants::Icon12x12, FLinearColor::Black));
+
+		Set("ContentBrowser.ContentDirty", new IMAGE_BRUSH("ContentBrowser/ContentDirty", CoreStyleConstants::Icon16x16));
+		Set("ContentBrowser.AssetDragDropTooltipBackground", new BOX_BRUSH("Old/Menu_Background", FMargin(8.0f / 64.0f)));
+		Set("ContentBrowser.CollectionTreeDragDropBorder", new BOX_BRUSH("Old/Window/ViewportDebugBorder", 0.8f));
+		Set("ContentBrowser.PopupMessageIcon", new IMAGE_BRUSH("Icons/alert", CoreStyleConstants::Icon32x32));
+		Set("ContentBrowser.NewFolderIcon", new CORE_IMAGE_BRUSH_SVG("Starship/Common/folder-plus", CoreStyleConstants::Icon16x16));
+		Set("ContentBrowser.ShowInExplorer", new IMAGE_BRUSH_SVG("Starship/ContentBrowser/show-in-explorer", CoreStyleConstants::Icon16x16));
+		Set("ContentBrowser.ReferenceViewer", new IMAGE_BRUSH_SVG("Starship/ContentBrowser/reference-viewer", CoreStyleConstants::Icon16x16));
+		Set("ContentBrowser.SizeMap", new IMAGE_BRUSH_SVG("Starship/ContentBrowser/size-map", CoreStyleConstants::Icon16x16));
+		Set("ContentBrowser.Collections", new IMAGE_BRUSH_SVG("Starship/ContentBrowser/collections", CoreStyleConstants::Icon16x16));
+		Set("ContentBrowser.Migrate", new IMAGE_BRUSH_SVG("Starship/ContentBrowser/migrate", CoreStyleConstants::Icon16x16));
+		Set("ContentBrowser.Local", new IMAGE_BRUSH("ContentBrowser/Content_Local_12x", CoreStyleConstants::Icon12x12));
+		Set("ContentBrowser.Local.Small", new IMAGE_BRUSH("ContentBrowser/Content_Local_16x", CoreStyleConstants::Icon16x16));
+		Set("ContentBrowser.Local.Large", new IMAGE_BRUSH("ContentBrowser/Content_Local_64x", CoreStyleConstants::Icon64x64));
+		Set("ContentBrowser.Shared", new IMAGE_BRUSH("ContentBrowser/Content_Shared_12x", CoreStyleConstants::Icon12x12));
+		Set("ContentBrowser.Shared.Small", new IMAGE_BRUSH("ContentBrowser/Content_Shared_16x", CoreStyleConstants::Icon16x16));
+		Set("ContentBrowser.Shared.Large", new IMAGE_BRUSH("ContentBrowser/Content_Shared_64x", CoreStyleConstants::Icon64x64));
+		Set("ContentBrowser.Private", new IMAGE_BRUSH("ContentBrowser/Content_Private_12x", CoreStyleConstants::Icon12x12));
+		Set("ContentBrowser.Private.Small", new IMAGE_BRUSH("ContentBrowser/Content_Private_16x", CoreStyleConstants::Icon16x16));
+		Set("ContentBrowser.Private.Large", new IMAGE_BRUSH("ContentBrowser/Content_Private_64x", CoreStyleConstants::Icon64x64));
+		Set("ContentBrowser.CollectionStatus", new IMAGE_BRUSH("/Icons/CollectionStatus_8x", CoreStyleConstants::Icon8x8));
+
+
+		Set("ContentBrowser.SaveAllCurrentFolder", new IMAGE_BRUSH_SVG("Starship/Common/SaveCurrent", CoreStyleConstants::Icon16x16));
+		Set("ContentBrowser.ResaveAllCurrentFolder", new IMAGE_BRUSH_SVG("Starship/Common/SaveCurrent", CoreStyleConstants::Icon16x16));
+
+		FToolBarStyle ContentBrowserToolBarStyle = FStarshipCoreStyle::GetCoreStyle().GetWidgetStyle<FToolBarStyle>("SlimToolBar");
+
+		ContentBrowserToolBarStyle.SetSeparatorBrush(FSlateNoResource());
+		ContentBrowserToolBarStyle.SetSeparatorPadding(FMargin(4.0f, 0.0f));
+		ContentBrowserToolBarStyle.SetBackgroundPadding(FMargin(4.0f, 2.0f, 0.0f, 2.0f));
+
+		Set("ContentBrowser.ToolBar", ContentBrowserToolBarStyle);
+
+		FSlateStyleRegistry::RegisterSlateStyle(*this);
+	}
+
+	FContentBrowserStyle::~FContentBrowserStyle()
+	{
+		FSlateStyleRegistry::UnRegisterSlateStyle(*this);
+	}
+}
diff --git a/Engine/Source/Editor/ContentBrowser/Private/ContentBrowserStyle.h b/Engine/Source/Editor/ContentBrowser/Private/ContentBrowserStyle.h
new file mode 100644
--- /dev/null
+++ b/Engine/Source/Editor/ContentBrowser/Private/ContentBrowserStyle.h
@@ -0,0 +1,37 @@
+// Copyright Epic Games, Inc. All Rights Reserved.
+
+#pragma once
+
+#include "Styling/SlateColor.h"
+#include "Styling/SlateStyle.h"
+#include "Styling/SlateTypes.h"
+
+namespace UE::ContentBrowser::Private
+{
+	class FContentBrowserStyle
+		: public FSlateStyleSet
+	{
+	public:
+		static FContentBrowserStyle& Get();
+
+		virtual const FName& GetStyleSetName() const override;
+
+		static bool IsNewStyleEnabled();
+
+	private:
+		FContentBrowserStyle();
+		virtual ~FContentBrowserStyle() override;
+
+	private:
+		static FName StyleName;
+
+		// Colors and Styles inherited from the parent style
+		FSlateColor DefaultForeground;
+		FSlateColor InvertedForeground;
+		FSlateColor SelectionColor;
+		FSlateColor SelectionColor_Pressed;
+
+		FTextBlockStyle NormalText;
+		FButtonStyle Button;
+	};
+}
diff --git a/Engine/Source/Editor/ContentBrowser/Private/ContentBrowserUtils.cpp b/Engine/Source/Editor/ContentBrowser/Private/ContentBrowserUtils.cpp
--- a/Engine/Source/Editor/ContentBrowser/Private/ContentBrowserUtils.cpp
+++ b/Engine/Source/Editor/ContentBrowser/Private/ContentBrowserUtils.cpp
@@ -61,6 +61,7 @@
 #include "Widgets/SCompoundWidget.h"
 #include "Widgets/Text/STextBlock.h"
 #include "Algo/Transform.h"
+#include "ContentBrowserStyle.h"
 
 class SWidget;
 struct FGeometry;
@@ -107,7 +108,8 @@ public:
 				.VAlign(VAlign_Center)
 				.Padding(0.f, 0.f, 4.f, 0.f)
 				[
-					SNew(SImage) .Image( FAppStyle::GetBrush("ContentBrowser.PopupMessageIcon") )
+					SNew(SImage)
+					.Image(UE::ContentBrowser::Private::FContentBrowserStyle::Get().GetBrush("ContentBrowser.PopupMessageIcon"))
 				]
 
 				+SHorizontalBox::Slot()
diff --git a/Engine/Source/Editor/ContentBrowser/Private/NewAssetOrClassContextMenu.cpp b/Engine/Source/Editor/ContentBrowser/Private/NewAssetOrClassContextMenu.cpp
--- a/Engine/Source/Editor/ContentBrowser/Private/NewAssetOrClassContextMenu.cpp
+++ b/Engine/Source/Editor/ContentBrowser/Private/NewAssetOrClassContextMenu.cpp
@@ -5,6 +5,7 @@
 #include "ContentBrowserConfig.h"
 #include "ContentBrowserDataMenuContexts.h"
 #include "ContentBrowserDataSubsystem.h"
+#include "ContentBrowserStyle.h"
 #include "Framework/Commands/UIAction.h"
 #include "IContentBrowserDataModule.h"
 #include "Internationalization/Internationalization.h"
@@ -52,7 +53,7 @@ void FNewAssetOrClassContextMenu::MakeContextMenu(
 				"GetContent",
 				LOCTEXT( "GetContentText", "Add Feature or Content Pack..." ),
 				LOCTEXT( "GetContentTooltip", "Add features and content packs to the project." ),
-				FSlateIcon( FAppStyle::GetAppStyleSetName(), "ContentBrowser.AddContent" ),
+				FSlateIcon(UE::ContentBrowser::Private::FContentBrowserStyle::Get().GetStyleSetName(), "ContentBrowser.AddContent" ),
 				FUIAction( FExecuteAction::CreateStatic( &FNewAssetOrClassContextMenu::ExecuteGetContent, InOnGetContentRequested ) )
 				);
 		}
@@ -97,7 +98,7 @@ void FNewAssetOrClassContextMenu::MakeContextMenu(
 				"NewFolder",
 				LOCTEXT("NewFolderLabel", "New Folder"),
 				NewFolderToolTip,
-				FSlateIcon(FAppStyle::GetAppStyleSetName(), "ContentBrowser.NewFolderIcon"),
+				FSlateIcon(UE::ContentBrowser::Private::FContentBrowserStyle::Get().GetStyleSetName(), "ContentBrowser.NewFolderIcon"),
 				FUIAction(
 					FExecuteAction::CreateStatic(&FNewAssetOrClassContextMenu::ExecuteNewFolder, FirstSelectedPath, InOnNewFolderRequested),
 					CanExecuteFolderActionsDelegate
diff --git a/Engine/Source/Editor/ContentBrowser/Private/PathContextMenu.cpp b/Engine/Source/Editor/ContentBrowser/Private/PathContextMenu.cpp
--- a/Engine/Source/Editor/ContentBrowser/Private/PathContextMenu.cpp
+++ b/Engine/Source/Editor/ContentBrowser/Private/PathContextMenu.cpp
@@ -16,6 +16,7 @@
 #include "ContentBrowserMenuContexts.h"
 #include "ContentBrowserModule.h"
 #include "ContentBrowserSingleton.h"
+#include "ContentBrowserStyle.h"
 #include "ContentBrowserUtils.h"
 #include "Framework/Application/SlateApplication.h"
 #include "Framework/Commands/GenericCommands.h"
@@ -244,7 +245,7 @@ void FPathContextMenu::MakePathViewContextMenu(UToolMenu* Menu)
 						"DisallowPrivateContentEditing",
 						LOCTEXT("DisallowPrivateContentEditing", "Disallow Private Content Editing"),
 						LOCTEXT("DisallowPrivateContentEditingTooltip", "Hides Private Content and prevents editing the Public/Private state of content in this folder"),
-						FSlateIcon(FAppStyle::GetAppStyleSetName(), "ContentBrowser.PrivateContentEdit"),
+						FSlateIcon(UE::ContentBrowser::Private::FContentBrowserStyle::Get().GetStyleSetName(), "ContentBrowser.PrivateContentEdit"),
 						FUIAction(FExecuteAction::CreateSP(this, &FPathContextMenu::ExecutePrivateContentEdit))
 					);
 				}
@@ -254,7 +255,7 @@ void FPathContextMenu::MakePathViewContextMenu(UToolMenu* Menu)
 						"AllowPrivateContentEditing",
 						LOCTEXT("AllowPrivateContentEditing", "Allow Private Content Editing"),
 						LOCTEXT("AllowPrivateContentEditingTooltip", "Reveals Private Content and allows editing the Public/Private state of content in this folder"),
-						FSlateIcon(FAppStyle::GetAppStyleSetName(), "ContentBrowser.PrivateContentEdit"),
+						FSlateIcon(UE::ContentBrowser::Private::FContentBrowserStyle::Get().GetStyleSetName(), "ContentBrowser.PrivateContentEdit"),
 						FUIAction(FExecuteAction::CreateSP(this, &FPathContextMenu::ExecutePrivateContentEdit))
 					);
 				}
diff --git a/Engine/Source/Editor/ContentBrowser/Private/SAssetDialog.cpp b/Engine/Source/Editor/ContentBrowser/Private/SAssetDialog.cpp
--- a/Engine/Source/Editor/ContentBrowser/Private/SAssetDialog.cpp
+++ b/Engine/Source/Editor/ContentBrowser/Private/SAssetDialog.cpp
@@ -16,6 +16,7 @@
 #include "ContentBrowserItemData.h"
 #include "ContentBrowserPluginFilters.h"
 #include "ContentBrowserSingleton.h"
+#include "ContentBrowserStyle.h"
 #include "ContentBrowserUtils.h"
 #include "CoreGlobals.h"
 #include "Delegates/Delegate.h"
@@ -686,6 +687,8 @@ TSharedPtr<SWidget> SAssetDialog::OnGetAssetContextMenu(const TArray<FAssetData>
 
 void SAssetDialog::SetupContextMenuContent(FMenuBuilder& MenuBuilder, const TArray<FString>& SelectedPaths)
 {
+	const FName ContentBrowserStyleSetName = UE::ContentBrowser::Private::FContentBrowserStyle::Get().GetStyleSetName();
+
 	FText NewFolderToolTip;
 
 	if (SelectedPaths.Num() > 0)
@@ -706,8 +709,8 @@ void SAssetDialog::SetupContextMenuContent(FMenuBuilder& MenuBuilder, const TArr
 
 	MenuBuilder.BeginSection("AssetDialogOptions", LOCTEXT("AssetDialogMenuHeading", "Options"));
 
-	MenuBuilder.AddMenuEntry(FContentBrowserCommands::Get().CreateNewFolder, NAME_None, LOCTEXT("NewFolder", "New Folder"), NewFolderToolTip, FSlateIcon(FAppStyle::GetAppStyleSetName(), "ContentBrowser.NewFolderIcon"));
-	MenuBuilder.AddMenuEntry(FGenericCommands::Get().Rename, NAME_None, LOCTEXT("RenameFolder", "Rename"), LOCTEXT("RenameFolderTooltip", "Rename the selected folder."), FSlateIcon(FAppStyle::GetAppStyleSetName(), "ContentBrowser.AssetActions.Rename"));
+	MenuBuilder.AddMenuEntry(FContentBrowserCommands::Get().CreateNewFolder, NAME_None, LOCTEXT("NewFolder", "New Folder"), NewFolderToolTip, FSlateIcon(ContentBrowserStyleSetName, "ContentBrowser.NewFolderIcon"));
+	MenuBuilder.AddMenuEntry(FGenericCommands::Get().Rename, NAME_None, LOCTEXT("RenameFolder", "Rename"), LOCTEXT("RenameFolderTooltip", "Rename the selected folder."), FSlateIcon(ContentBrowserStyleSetName, "ContentBrowser.AssetActions.Rename"));
 	MenuBuilder.AddMenuEntry(FGenericCommands::Get().Delete, NAME_None, LOCTEXT("DeleteFolder", "Delete"), LOCTEXT("DeleteFolderTooltip", "Removes this folder and all assets it contains."));
 
 	MenuBuilder.EndSection();
@@ -717,7 +720,7 @@ void SAssetDialog::SetupContextMenuContent(FMenuBuilder& MenuBuilder, const TArr
 		MenuBuilder.BeginSection("AssetDialogExplore", LOCTEXT("AssetDialogExploreHeading", "Explore"));
 		MenuBuilder.AddMenuEntry(ContentBrowserUtils::GetExploreFolderText(),
 			LOCTEXT("ExploreTooltip", "Finds this folder on disk."),
-			FSlateIcon(FAppStyle::GetAppStyleSetName(), "SystemWideCommands.FindInContentBrowser"),
+			FSlateIcon(ContentBrowserStyleSetName, "SystemWideCommands.FindInContentBrowser"),
 			FUIAction(FExecuteAction::CreateSP(this, &SAssetDialog::ExecuteExplore)));
 		MenuBuilder.EndSection();
 	}
diff --git a/Engine/Source/Editor/ContentBrowser/Private/SAssetPicker.cpp b/Engine/Source/Editor/ContentBrowser/Private/SAssetPicker.cpp
--- a/Engine/Source/Editor/ContentBrowser/Private/SAssetPicker.cpp
+++ b/Engine/Source/Editor/ContentBrowser/Private/SAssetPicker.cpp
@@ -8,6 +8,7 @@
 #include "CollectionManagerTypes.h"
 #include "ContentBrowserDataFilter.h"
 #include "ContentBrowserItem.h"
+#include "ContentBrowserStyle.h"
 #include "ContentBrowserUtils.h"
 #include "CoreGlobals.h"
 #include "Delegates/Delegate.h"
@@ -176,7 +177,7 @@ void SAssetPicker::Construct( const FArguments& InArgs )
 			[
 				SNew(SImage)
 				.ColorAndOpacity(FSlateColor::UseForeground())
-				.Image(FAppStyle::GetBrush("ContentBrowser.ColumnViewDeveloperFolderIcon"))
+				.Image(UE::ContentBrowser::Private::FContentBrowserStyle::Get().GetBrush("ContentBrowser.ColumnViewDeveloperFolderIcon"))
 			]
 		];
 	}
@@ -242,8 +243,8 @@ void SAssetPicker::Construct( const FArguments& InArgs )
 				.AutoHeight()
 				[
 					SNew(SButton)
-						.ButtonStyle( FAppStyle::Get(), "ContentBrowser.NoneButton" )
-						.TextStyle( FAppStyle::Get(), "ContentBrowser.NoneButtonText" )
+						.ButtonStyle(UE::ContentBrowser::Private::FContentBrowserStyle::Get(), "ContentBrowser.NoneButton" )
+						.TextStyle(UE::ContentBrowser::Private::FContentBrowserStyle::Get(), "ContentBrowser.NoneButtonText" )
 						.Text( LOCTEXT("NoneButtonText", "( None )") )
 						.ToolTipText( LOCTEXT("NoneButtonTooltip", "Clears the asset selection.") )
 						.HAlign(HAlign_Center)
diff --git a/Engine/Source/Editor/ContentBrowser/Private/SAssetView.cpp b/Engine/Source/Editor/ContentBrowser/Private/SAssetView.cpp
--- a/Engine/Source/Editor/ContentBrowser/Private/SAssetView.cpp
+++ b/Engine/Source/Editor/ContentBrowser/Private/SAssetView.cpp
@@ -27,6 +27,7 @@
 #include "ContentBrowserMenuUtils.h"
 #include "ContentBrowserModule.h"
 #include "ContentBrowserSingleton.h"
+#include "ContentBrowserStyle.h"
 #include "ContentBrowserUtils.h"
 #include "DesktopPlatformModule.h"
 #include "DragAndDrop/AssetDragDropOp.h"
@@ -4764,7 +4765,7 @@ TSharedRef<ITableRow> SAssetView::MakeTileViewWidget(TSharedPtr<FAssetViewItem>
 	{
 		TSharedPtr< STableRow<TSharedPtr<FAssetViewItem>> > TableRowWidget;
 		SAssignNew( TableRowWidget, STableRow<TSharedPtr<FAssetViewItem>>, OwnerTable )
-			.Style( FAppStyle::Get(), "ContentBrowser.AssetListView.TileTableRow" )
+			.Style(UE::ContentBrowser::Private::FContentBrowserStyle::Get(), "ContentBrowser.AssetListView.TileTableRow" )
 			.Cursor( bAllowDragging ? EMouseCursor::GrabHand : EMouseCursor::Default )
 			.OnDragDetected( this, &SAssetView::OnDraggingAssetItem );
 
@@ -4799,7 +4800,7 @@ TSharedRef<ITableRow> SAssetView::MakeTileViewWidget(TSharedPtr<FAssetViewItem>
 
 		TSharedPtr< STableRow<TSharedPtr<FAssetViewItem>> > TableRowWidget;
 		SAssignNew( TableRowWidget, STableRow<TSharedPtr<FAssetViewItem>>, OwnerTable )
-		.Style(FAppStyle::Get(), "ContentBrowser.AssetListView.TileTableRow")
+		.Style(UE::ContentBrowser::Private::FContentBrowserStyle::Get(), "ContentBrowser.AssetListView.TileTableRow")
 		.Cursor( bAllowDragging ? EMouseCursor::GrabHand : EMouseCursor::Default )
 		.OnDragDetected( this, &SAssetView::OnDraggingAssetItem );
 
@@ -4840,7 +4841,7 @@ TSharedRef<ITableRow> SAssetView::MakeColumnViewWidget(TSharedPtr<FAssetViewItem
 	if ( !ensure(AssetItem.IsValid()) )
 	{
 		return SNew( STableRow<TSharedPtr<FAssetViewItem>>, OwnerTable )
-			.Style(FAppStyle::Get(), "ContentBrowser.AssetListView.ColumnListTableRow");
+			.Style(UE::ContentBrowser::Private::FContentBrowserStyle::Get(), "ContentBrowser.AssetListView.ColumnListTableRow");
 	}
 
 	// Update the cached custom data
diff --git a/Engine/Source/Editor/ContentBrowser/Private/SCollectionView.cpp b/Engine/Source/Editor/ContentBrowser/Private/SCollectionView.cpp
--- a/Engine/Source/Editor/ContentBrowser/Private/SCollectionView.cpp
+++ b/Engine/Source/Editor/ContentBrowser/Private/SCollectionView.cpp
@@ -11,6 +11,7 @@
 #include "ContentBrowserDelegates.h"
 #include "ContentBrowserModule.h"
 #include "ContentBrowserPluginFilters.h"
+#include "ContentBrowserStyle.h"
 #include "ContentBrowserUtils.h"
 #include "CoreGlobals.h"
 #include "DragAndDrop/AssetDragDropOp.h"
@@ -1121,7 +1122,7 @@ EVisibility SCollectionView::GetHeaderVisibility() const
 
 const FSlateBrush* SCollectionView::GetCollectionViewDropTargetBorder() const
 {
-	return bDraggedOver ? FAppStyle::GetBrush("ContentBrowser.CollectionTreeDragDropBorder") : FAppStyle::GetBrush("NoBorder");
+	return bDraggedOver ? UE::ContentBrowser::Private::FContentBrowserStyle::Get().GetBrush("ContentBrowser.CollectionTreeDragDropBorder") : FAppStyle::GetBrush("NoBorder");
 }
 
 TSharedRef<ITableRow> SCollectionView::GenerateCollectionRow( TSharedPtr<FCollectionItem> CollectionItem, const TSharedRef<STableViewBase>& OwnerTable )
diff --git a/Engine/Source/Editor/ContentBrowser/Private/SContentBrowser.cpp b/Engine/Source/Editor/ContentBrowser/Private/SContentBrowser.cpp
--- a/Engine/Source/Editor/ContentBrowser/Private/SContentBrowser.cpp
+++ b/Engine/Source/Editor/ContentBrowser/Private/SContentBrowser.cpp
@@ -33,6 +33,7 @@
 #include "ContentBrowserMenuContexts.h"
 #include "ContentBrowserModule.h"
 #include "ContentBrowserSingleton.h"
+#include "ContentBrowserStyle.h"
 #include "ContentBrowserUtils.h"
 #include "ContentBrowserVirtualPathTree.h"
 #include "CoreGlobals.h"
@@ -2752,7 +2753,7 @@ TSharedRef<SWidget> SContentBrowser::OnGetCrumbDelimiterContent(const FString& C
 				MenuBuilder.AddMenuEntry(
 					SubItem.GetDisplayName(),
 					FText::GetEmpty(),
-					FSlateIcon(FAppStyle::GetAppStyleSetName(), "ContentBrowser.BreadcrumbPathPickerFolder"),
+					FSlateIcon(UE::ContentBrowser::Private::FContentBrowserStyle::Get().GetStyleSetName(), "ContentBrowser.BreadcrumbPathPickerFolder"),
 					FUIAction(FExecuteAction::CreateSP(const_cast<SContentBrowser*>(this), &SContentBrowser::OnPathMenuItemClicked, SubItem.GetVirtualPath().ToString()))
 					);
 			}
@@ -4436,7 +4437,7 @@ void SContentBrowser::PopulateFolderContextMenu(UToolMenu* Menu)
 				"NewFolder",
 				LOCTEXT("NewFolder", "New Folder"),
 				NewFolderToolTip,
-				FSlateIcon(FAppStyle::GetAppStyleSetName(), "ContentBrowser.NewFolderIcon"),
+				FSlateIcon(UE::ContentBrowser::Private::FContentBrowserStyle::Get().GetStyleSetName(), "ContentBrowser.NewFolderIcon"),
 				FUIAction(
 					FExecuteAction::CreateSP(this, &SContentBrowser::CreateNewFolder, SelectedFolders.Num() > 0 ? SelectedFolders[0].GetVirtualPath().ToString() : FString(), Context->OnCreateNewFolder),
 					FCanExecuteAction::CreateLambda([bCanCreateNewFolder] { return bCanCreateNewFolder; })
@@ -4448,7 +4449,7 @@ void SContentBrowser::PopulateFolderContextMenu(UToolMenu* Menu)
 			"FolderContext",
 			LOCTEXT("ShowInNewContentBrowser", "Show in New Content Browser"),
 			LOCTEXT("ShowInNewContentBrowserTooltip", "Opens a new Content Browser at this folder location (at least 1 Content Browser window needs to be locked)"),
-			FSlateIcon(FAppStyle::GetAppStyleSetName(), "ContentBrowser.TabIcon"),
+			FSlateIcon(UE::ContentBrowser::Private::FContentBrowserStyle::Get().GetStyleSetName(), "ContentBrowser.TabIcon"),
 			FUIAction(FExecuteAction::CreateSP(this, &SContentBrowser::OpenNewContentBrowser))
 		);
 	}
diff --git a/Engine/Source/Editor/ContentBrowser/Private/SPathPicker.cpp b/Engine/Source/Editor/ContentBrowser/Private/SPathPicker.cpp
--- a/Engine/Source/Editor/ContentBrowser/Private/SPathPicker.cpp
+++ b/Engine/Source/Editor/ContentBrowser/Private/SPathPicker.cpp
@@ -8,6 +8,7 @@
 #include "ContentBrowserItem.h"
 #include "ContentBrowserItemData.h"
 #include "ContentBrowserPluginFilters.h"
+#include "ContentBrowserStyle.h"
 #include "Delegates/Delegate.h"
 #include "Framework/Commands/UIAction.h"
 #include "Framework/MultiBox/MultiBoxBuilder.h"
@@ -162,7 +163,7 @@ TSharedPtr<SWidget> SPathPicker::GetFolderContextMenu(const TArray<FString>& Sel
 	MenuBuilder.AddMenuEntry(
 		LOCTEXT("NewFolder", "New Folder"),
 		NewFolderToolTip,
-		FSlateIcon(FAppStyle::GetAppStyleSetName(), "ContentBrowser.NewFolderIcon"),
+		FSlateIcon(UE::ContentBrowser::Private::FContentBrowserStyle::Get().GetStyleSetName(), "ContentBrowser.NewFolderIcon"),
 		FUIAction(
 			FExecuteAction::CreateSP(this, &SPathPicker::CreateNewFolder, SelectedPaths.Num() > 0 ? SelectedPaths[0] : FString(), InOnCreateNewFolder),
 			FCanExecuteAction::CreateLambda( [bCanCreateNewFolder] { return bCanCreateNewFolder; } )
diff --git a/Engine/Source/Editor/ContentBrowser/Private/SPathView.cpp b/Engine/Source/Editor/ContentBrowser/Private/SPathView.cpp
--- a/Engine/Source/Editor/ContentBrowser/Private/SPathView.cpp
+++ b/Engine/Source/Editor/ContentBrowser/Private/SPathView.cpp
@@ -19,6 +19,7 @@
 #include "ContentBrowserPathViewMenuContexts.h"
 #include "ContentBrowserPluginFilters.h"
 #include "ContentBrowserSingleton.h"
+#include "ContentBrowserStyle.h"
 #include "ContentBrowserUtils.h"
 #include "CoreGlobals.h"
 #include "CoreTypes.h"
@@ -1091,7 +1092,7 @@ void SPathView::Construct( const FArguments& InArgs )
 				.AutoHeight()
 				[
 					SNew(STextBlock)
-					.Font( FAppStyle::GetFontStyle("ContentBrowser.SourceTitleFont") )
+					.Font(UE::ContentBrowser::Private::FContentBrowserStyle::Get().GetFontStyle("ContentBrowser.SourceTitleFont") )
 					.Text(this, &SPathView::GetTreeTitle)
 					.Visibility(InArgs._ShowTreeTitle ? EVisibility::Visible : EVisibility::Collapsed)
 				]
@@ -2807,7 +2808,7 @@ TSharedRef<ITableRow> SFavoritePathView::GenerateTreeRow(TSharedPtr<FTreeItem> T
 			.IsItemExpanded(false)
 			.HighlightText(this, &SFavoritePathView::GetHighlightText)
 			.IsSelected(this, &SFavoritePathView::IsTreeItemSelected, TreeItem)
-			.FontOverride(FAppStyle::GetFontStyle("ContentBrowser.SourceTreeItemFont"))
+			.FontOverride(UE::ContentBrowser::Private::FContentBrowserStyle::Get().GetFontStyle("ContentBrowser.SourceTreeItemFont"))
 		];
 }
 
diff --git a/Engine/Source/Editor/ContentBrowser/Private/SThumbnailEditModeTools.cpp b/Engine/Source/Editor/ContentBrowser/Private/SThumbnailEditModeTools.cpp
--- a/Engine/Source/Editor/ContentBrowser/Private/SThumbnailEditModeTools.cpp
+++ b/Engine/Source/Editor/ContentBrowser/Private/SThumbnailEditModeTools.cpp
@@ -7,6 +7,7 @@
 #include "AssetThumbnail.h"
 #include "AssetToolsModule.h"
 #include "Containers/EnumAsByte.h"
+#include "ContentBrowserStyle.h"
 #include "Delegates/Delegate.h"
 #include "Editor/UnrealEdEngine.h"
 #include "GenericPlatform/ICursor.h"
@@ -87,7 +88,7 @@ void SThumbnailEditModeTools::Construct( const FArguments& InArgs, const TShared
 			.Content()
 			[
 				SNew(SImage)
-				.Image(FAppStyle::GetBrush("ContentBrowser.ResetPrimitiveToDefault"))
+				.Image(UE::ContentBrowser::Private::FContentBrowserStyle::Get().GetBrush("ContentBrowser.ResetPrimitiveToDefault"))
 				.ColorAndOpacity(FSlateColor::UseForeground())
 			]
 		]
@@ -115,6 +116,8 @@ EVisibility SThumbnailEditModeTools::GetPrimitiveToolsResetToDefaultVisibility()
 
 const FSlateBrush* SThumbnailEditModeTools::GetCurrentPrimitiveBrush() const
 {
+	const FSlateStyleSet& ContentBrowserStyle = UE::ContentBrowser::Private::FContentBrowserStyle::Get();
+
 	USceneThumbnailInfoWithPrimitive* ThumbnailInfo = GetSceneThumbnailInfoWithPrimitive();
 	if ( ThumbnailInfo )
 	{
@@ -122,10 +125,10 @@ const FSlateBrush* SThumbnailEditModeTools::GetCurrentPrimitiveBrush() const
 		EThumbnailPrimType PrimType = ThumbnailInfo->bUserModifiedShape ? ThumbnailInfo->PrimitiveType.GetValue() : (EThumbnailPrimType)ThumbnailInfo->DefaultPrimitiveType.Get(EThumbnailPrimType::TPT_Sphere);
 		switch (PrimType)
 		{
-		case TPT_None: return FAppStyle::GetBrush("ContentBrowser.PrimitiveCustom");
-		case TPT_Sphere: return FAppStyle::GetBrush("ContentBrowser.PrimitiveSphere");
-		case TPT_Cube: return FAppStyle::GetBrush("ContentBrowser.PrimitiveCube");
-		case TPT_Cylinder: return FAppStyle::GetBrush("ContentBrowser.PrimitiveCylinder");
+		case TPT_None: return ContentBrowserStyle.GetBrush("ContentBrowser.PrimitiveCustom");
+		case TPT_Sphere: return ContentBrowserStyle.GetBrush("ContentBrowser.PrimitiveSphere");
+		case TPT_Cube: return ContentBrowserStyle.GetBrush("ContentBrowser.PrimitiveCube");
+		case TPT_Cylinder: return ContentBrowserStyle.GetBrush("ContentBrowser.PrimitiveCylinder");
 		case TPT_Plane:
 		default:
 			// Fall through and return a plane
@@ -133,7 +136,7 @@ const FSlateBrush* SThumbnailEditModeTools::GetCurrentPrimitiveBrush() const
 		}
 	}
 
-	return FAppStyle::GetBrush( "ContentBrowser.PrimitivePlane" );
+	return ContentBrowserStyle.GetBrush( "ContentBrowser.PrimitivePlane" );
 }
 
 FReply SThumbnailEditModeTools::ChangePrimitive()
diff --git a/Engine/Source/Editor/ContentBrowser/Private/SourcesViewWidgets.cpp b/Engine/Source/Editor/ContentBrowser/Private/SourcesViewWidgets.cpp
--- a/Engine/Source/Editor/ContentBrowser/Private/SourcesViewWidgets.cpp
+++ b/Engine/Source/Editor/ContentBrowser/Private/SourcesViewWidgets.cpp
@@ -12,6 +12,7 @@
 #include "ContentBrowserItemData.h"
 #include "ContentBrowserModule.h"
 #include "ContentBrowserPluginFilters.h"
+#include "ContentBrowserStyle.h"
 #include "ContentBrowserUtils.h"
 #include "DragAndDrop/AssetDragDropOp.h"
 #include "DragAndDrop/CollectionDragDropOp.h"
@@ -57,16 +58,18 @@ struct FAssetTreeItemBrushes
 
 	FAssetTreeItemBrushes()
 	{
-		FolderOpenBrush = FAppStyle::GetBrush("ContentBrowser.AssetTreeFolderOpen");
-		FolderClosedBrush = FAppStyle::GetBrush("ContentBrowser.AssetTreeFolderClosed");
-		FolderOpenVirtualBrush = FAppStyle::GetBrush("ContentBrowser.AssetTreeFolderOpenVirtual");
-		FolderClosedVirtualBrush = FAppStyle::GetBrush("ContentBrowser.AssetTreeFolderClosedVirtual");
-		FolderOpenCodeBrush = FAppStyle::GetBrush("ContentBrowser.AssetTreeFolderOpenCode");
-		FolderClosedCodeBrush = FAppStyle::GetBrush("ContentBrowser.AssetTreeFolderClosedCode");
-		FolderOpenDeveloperBrush = FAppStyle::GetBrush("ContentBrowser.AssetTreeFolderOpenDeveloper");
-		FolderClosedDeveloperBrush = FAppStyle::GetBrush("ContentBrowser.AssetTreeFolderClosedDeveloper");
-		FolderOpenPluginRootBrush = FAppStyle::GetBrush("ContentBrowser.AssetTreeFolderOpenPluginRoot");
-		FolderClosedPluginRootBrush = FAppStyle::GetBrush("ContentBrowser.AssetTreeFolderClosedPluginRoot");
+		const FSlateStyleSet& ContentBrowserStyle = UE::ContentBrowser::Private::FContentBrowserStyle::Get();
+
+		FolderOpenBrush = ContentBrowserStyle.GetBrush("ContentBrowser.AssetTreeFolderOpen");
+		FolderClosedBrush = ContentBrowserStyle.GetBrush("ContentBrowser.AssetTreeFolderClosed");
+		FolderOpenVirtualBrush = ContentBrowserStyle.GetBrush("ContentBrowser.AssetTreeFolderOpenVirtual");
+		FolderClosedVirtualBrush = ContentBrowserStyle.GetBrush("ContentBrowser.AssetTreeFolderClosedVirtual");
+		FolderOpenCodeBrush = ContentBrowserStyle.GetBrush("ContentBrowser.AssetTreeFolderOpenCode");
+		FolderClosedCodeBrush = ContentBrowserStyle.GetBrush("ContentBrowser.AssetTreeFolderClosedCode");
+		FolderOpenDeveloperBrush = ContentBrowserStyle.GetBrush("ContentBrowser.AssetTreeFolderOpenDeveloper");
+		FolderClosedDeveloperBrush = ContentBrowserStyle.GetBrush("ContentBrowser.AssetTreeFolderClosedDeveloper");
+		FolderOpenPluginRootBrush = ContentBrowserStyle.GetBrush("ContentBrowser.AssetTreeFolderOpenPluginRoot");
+		FolderClosedPluginRootBrush = ContentBrowserStyle.GetBrush("ContentBrowser.AssetTreeFolderClosedPluginRoot");
 	}
 	
 	static FAssetTreeItemBrushes& Get()
@@ -153,7 +156,7 @@ void SAssetTreeItem::Construct( const FArguments& InArgs )
 				SAssignNew(InlineRenameWidget, SInlineEditableTextBlock)
 					.Text(this, &SAssetTreeItem::GetNameText)
 					.ToolTipText(this, &SAssetTreeItem::GetToolTipText)
-					.Font( InArgs._FontOverride.IsSet() ? InArgs._FontOverride : FAppStyle::GetFontStyle(bIsRoot ? "ContentBrowser.SourceTreeRootItemFont" : "ContentBrowser.SourceTreeItemFont") )
+					.Font( InArgs._FontOverride.IsSet() ? InArgs._FontOverride : UE::ContentBrowser::Private::FContentBrowserStyle::Get().GetFontStyle(bIsRoot ? "ContentBrowser.SourceTreeRootItemFont" : "ContentBrowser.SourceTreeItemFont") )
 					.HighlightText( InArgs._HighlightText )
 					.OnTextCommitted(this, &SAssetTreeItem::HandleNameCommitted)
 					.OnVerifyTextChanged(this, &SAssetTreeItem::VerifyNameChanged)
