[ContentBrowser] Content Browser Nav Bar should respect custom folder color

diff --git a/Engine/Source/Editor/ContentBrowser/Private/ContentBrowserUtils.cpp b/Engine/Source/Editor/ContentBrowser/Private/ContentBrowserUtils.cpp
--- a/Engine/Source/Editor/ContentBrowser/Private/ContentBrowserUtils.cpp
+++ b/Engine/Source/Editor/ContentBrowser/Private/ContentBrowserUtils.cpp
@@ -557,6 +557,44 @@ bool ContentBrowserUtils::TryGetFolderBrushAndShadowName(const FContentBrowserIt
 	return true;
 }
 
+bool ContentBrowserUtils::TryGetFolderBrushAndShadowNameSmall(const FContentBrowserItem& InFolder, FName& OutBrushName, FName& OutShadowBrushName)
+{
+	if (!InFolder.IsValid() || !InFolder.IsFolder())
+	{
+		return false;
+	}
+
+	OutShadowBrushName = TEXT("ContentBrowser.FolderItem.DropShadow");
+	const bool bDeveloperFolder = IsItemDeveloperContent(InFolder);
+	const bool bCodeFolder = EnumHasAnyFlags(InFolder.GetItemCategory(), EContentBrowserItemFlags::Category_Class);
+	const FContentBrowserItemDataAttributeValue VirtualAttributeValue = InFolder.GetItemAttribute(ContentBrowserItemAttributes::ItemIsCustomVirtualFolder);
+	const bool bVirtualFolder = VirtualAttributeValue.IsValid() && VirtualAttributeValue.GetValue<bool>();
+	const bool bPluginFolder = IsItemPluginRootFolder(InFolder);
+
+	if (bDeveloperFolder)
+	{
+		OutBrushName = TEXT("ContentBrowser.AssetTreeFolderClosedDeveloper");
+	}
+	else if (bCodeFolder)
+	{
+		OutBrushName = TEXT("ContentBrowser.AssetTreeFolderClosedCode");
+	}
+	else if (bVirtualFolder && ShouldShowCustomVirtualFolderIcon())
+	{
+		OutBrushName = TEXT("ContentBrowser.AssetTreeFolderClosedVirtual");
+		OutShadowBrushName = TEXT("ContentBrowser.ListViewVirtualFolderShadow");
+	}
+	else if (bPluginFolder && ShouldShowPluginFolderIcon())
+	{
+		OutBrushName = TEXT("ContentBrowser.AssetTreeFolderClosedPluginRoot");
+	}
+	else
+	{
+		OutBrushName = TEXT("ContentBrowser.AssetTreeFolderClosed");
+	}
+	return true;
+}
+
 bool ContentBrowserUtils::IsCollectionPath(const FString& InPath, FName* OutCollectionName, ECollectionShareType::Type* OutCollectionShareType)
 {
 	static const FString CollectionsRootPrefix = TEXT("/Collections");
diff --git a/Engine/Source/Editor/ContentBrowser/Private/ContentBrowserUtils.h b/Engine/Source/Editor/ContentBrowser/Private/ContentBrowserUtils.h
--- a/Engine/Source/Editor/ContentBrowser/Private/ContentBrowserUtils.h
+++ b/Engine/Source/Editor/ContentBrowser/Private/ContentBrowserUtils.h
@@ -90,6 +90,9 @@ namespace ContentBrowserUtils
 	/** Given a Folder, work out the BrushName and ShadowBrushName to use for it */
 	bool TryGetFolderBrushAndShadowName(const FContentBrowserItem& InFolder, FName& OutBrushName, FName& OutShadowBrushName);
 
+	/** Given a Folder, work out the small version of the BrushName and ShadowBrushName to use for it */
+	bool TryGetFolderBrushAndShadowNameSmall(const FContentBrowserItem& InFolder, FName& OutBrushName, FName& OutShadowBrushName);
+
 	/** Check to see whether the given path is rooted against a collection directory, optionally extracting the collection name and share type from the path */
 	bool IsCollectionPath(const FString& InPath, FName* OutCollectionName = nullptr, ECollectionShareType::Type* OutCollectionShareType = nullptr);
 
diff --git a/Engine/Source/Editor/ContentBrowser/Private/SContentBrowser.cpp b/Engine/Source/Editor/ContentBrowser/Private/SContentBrowser.cpp
--- a/Engine/Source/Editor/ContentBrowser/Private/SContentBrowser.cpp
+++ b/Engine/Source/Editor/ContentBrowser/Private/SContentBrowser.cpp
@@ -16,6 +16,7 @@
 #include "AssetViewUtils.h"
 #include "CollectionManagerModule.h"
 #include "CollectionManagerTypes.h"
+#include "CollectionViewUtils.h"
 #include "Containers/Map.h"
 #include "Containers/Set.h"
 #include "Containers/StringFwd.h"
@@ -279,6 +280,72 @@ namespace ContentBrowser
 
 		return SubItems;
 	}
+
+	void RetrieveSmallFolderInformation(const FContentBrowserItem& InFolder, FName& OutFolderBrushName, FName& OutFolderShadowName, FLinearColor& OutFolderOverrideColor)
+	{
+		// We are just dragging folders so add information on the FolderBrush to use and the color
+		if (InFolder.IsValid() && InFolder.IsFolder())
+		{
+			ContentBrowserUtils::TryGetFolderBrushAndShadowNameSmall(InFolder, OutFolderBrushName, OutFolderShadowName);
+
+			const bool bCollectionFolder = EnumHasAnyFlags(InFolder.GetItemCategory(), EContentBrowserItemFlags::Category_Collection);
+			if (bCollectionFolder)
+			{
+				FName CollectionName;
+				ECollectionShareType::Type CollectionFolderShareType = ECollectionShareType::CST_All;
+				ContentBrowserUtils::IsCollectionPath(InFolder.GetVirtualPath().ToString(), &CollectionName, &CollectionFolderShareType);
+			
+				if (TOptional<FLinearColor> Color = CollectionViewUtils::GetCustomColor(CollectionName, CollectionFolderShareType))
+				{
+					OutFolderOverrideColor = Color.GetValue();
+				}
+			}
+			else
+			{
+				if (TOptional<FLinearColor> Color = ContentBrowserUtils::GetPathColor(InFolder.GetInvariantPath().ToString()))
+				{
+					OutFolderOverrideColor = Color.GetValue();
+				}
+			}
+		}
+	}
+
+	TSharedPtr<SWidget> GetFolderWidgetForNavigationBar(const FText& InFolderName, const FName& InFolderBrushName, const FLinearColor& InFolderOverrideColor)
+	{
+		const FSlateBrush* FolderBrush = FAppStyle::GetBrush(InFolderBrushName);
+
+		if (FolderBrush != FAppStyle::GetDefaultBrush())
+		{
+			return SNew(SHorizontalBox)
+						+ SHorizontalBox::Slot()
+						.AutoWidth()
+						.Padding(FMargin(2, 0, 6, 0))
+						[
+							SNew(SBox)
+							.WidthOverride(16.f)
+							.HeightOverride(16.f)
+							.HAlign(HAlign_Center)
+							.VAlign(VAlign_Center)
+							[
+								SNew(SImage)
+								.DesiredSizeOverride(FVector2D(16, 16))
+								.Image(FolderBrush)
+								.ColorAndOpacity(InFolderOverrideColor)
+							]
+						]
+
+						+ SHorizontalBox::Slot()
+						.FillWidth(1.f)
+						.Padding(FMargin(2, 0, 6, 0))
+						.VAlign( VAlign_Center )
+						[
+							SNew(STextBlock)
+							.TextStyle(FAppStyle::Get(), TEXT("Menu.Label"))
+							.Text(InFolderName)
+						];
+		}
+		return nullptr;
+	}
 }
 
 SContentBrowser::~SContentBrowser()
@@ -2750,12 +2817,31 @@ TSharedRef<SWidget> SContentBrowser::OnGetCrumbDelimiterContent(const FString& C
 
 			for (const FContentBrowserItem& SubItem : SubItems)
 			{
-				MenuBuilder.AddMenuEntry(
-					SubItem.GetDisplayName(),
-					FText::GetEmpty(),
-					FSlateIcon(UE::ContentBrowser::Private::FContentBrowserStyle::Get().GetStyleSetName(), "ContentBrowser.BreadcrumbPathPickerFolder"),
-					FUIAction(FExecuteAction::CreateSP(const_cast<SContentBrowser*>(this), &SContentBrowser::OnPathMenuItemClicked, SubItem.GetVirtualPath().ToString()))
+				FName FolderBrushName = NAME_None;
+				FName FolderShadowBrushName = NAME_None;
+				FLinearColor FolderColor = ContentBrowserUtils::GetDefaultColor();
+				ContentBrowser::RetrieveSmallFolderInformation(SubItem, FolderBrushName, FolderShadowBrushName, FolderColor);
+
+				FText EntryName = SubItem.GetDisplayName();
+				FUIAction EntryAction = FUIAction(FExecuteAction::CreateSP(const_cast<SContentBrowser*>(this), &SContentBrowser::OnPathMenuItemClicked, SubItem.GetVirtualPath().ToString()));
+
+				if (FolderBrushName != NAME_None)
+				{
+					FMenuEntryParams Params;
+					Params.EntryWidget = ContentBrowser::GetFolderWidgetForNavigationBar(EntryName, FolderBrushName, FolderColor);
+					Params.DirectActions = EntryAction;
+					MenuBuilder.AddMenuEntry(Params);
+				}
+				else
+				{
+					MenuBuilder.AddMenuEntry(
+						EntryName,
+						FText::GetEmpty(),
+						FSlateIcon(FAppStyle::GetAppStyleSetName(), FolderBrushName),
+						EntryAction
 					);
+				}
+
 			}
 
 			MenuWidget = MenuBuilder.MakeWidget();
