* Fixed bad brush reference in ContentBrowserStyle
* Adjusted slate post buffer log from warning to log to allow for easier use on CI.

diff --git a/Engine/Source/Editor/ContentBrowser/Private/ContentBrowserStyle.cpp b/Engine/Source/Editor/ContentBrowser/Private/ContentBrowserStyle.cpp
--- a/Engine/Source/Editor/ContentBrowser/Private/ContentBrowserStyle.cpp
+++ b/Engine/Source/Editor/ContentBrowser/Private/ContentBrowserStyle.cpp
@@ -286,7 +286,7 @@ namespace UE::ContentBrowser::Private
 			Set("ContentBrowser.AssetTileItem.HoverBorder", new FSlateRoundedBoxBrush(TransparentHover, 4.0f, FStyleColors::Hover, 1.0f));
 		}
 
-		Set("ContentBrowser.AssetTileItem.DropShadow", new CORE_BOX_BRUSH("Starship/ContentBrowser/drop-shadow", FMargin(4.0f / 64.0f)));
+		Set("ContentBrowser.AssetTileItem.DropShadow", new BOX_BRUSH("Starship/ContentBrowser/drop-shadow", FMargin(4.0f / 64.0f)));
 		Set("ContentBrowser.FolderItem.DropShadow", new IMAGE_BRUSH("Starship/ContentBrowser/folder-drop-shadow", FVector2D(256, 256)));
 
 		Set("ContentBrowser.ShowSourcesView", new IMAGE_BRUSH_SVG("Starship/ContentBrowser/file-tree", CoreStyleConstants::Icon16x16));
diff --git a/Engine/Source/Editor/UMGEditor/Private/UMGEditorModule.cpp b/Engine/Source/Editor/UMGEditor/Private/UMGEditorModule.cpp
--- a/Engine/Source/Editor/UMGEditor/Private/UMGEditorModule.cpp
+++ b/Engine/Source/Editor/UMGEditor/Private/UMGEditorModule.cpp
@@ -402,7 +402,7 @@ private:
 				{
 					if (IConsoleVariable* PostBuffersEnabled = IConsoleManager::Get().FindConsoleVariable(TEXT("Slate.CopyBackbufferToSlatePostRenderTargets")))
 					{
-						UE_LOG(LogUMGEditor, Warning, TEXT("Disabling Slate Post Buffers for multi-window PIE session, currently not supported."));
+						UE_LOG(LogUMGEditor, Log, TEXT("Disabling Slate Post Buffers for multi-window PIE session, currently not supported."));
 
 						bCachedArePostBuffersEnabled = PostBuffersEnabled->GetBool();
 						PostBuffersEnabled->Set(false);
