[Content Browser] Style refresh - removed CVar (in favor of define UE_CONTENTBROWSER_NEW_STYLE)

diff --git a/Engine/Source/Editor/ContentBrowser/Private/ContentBrowserStyle.cpp b/Engine/Source/Editor/ContentBrowser/Private/ContentBrowserStyle.cpp
--- a/Engine/Source/Editor/ContentBrowser/Private/ContentBrowserStyle.cpp
+++ b/Engine/Source/Editor/ContentBrowser/Private/ContentBrowserStyle.cpp
@@ -11,13 +11,6 @@
 
 namespace UE::ContentBrowser::Private
 {
-	static bool bEnableNewStyle = false;
-	static FAutoConsoleVariableRef CVarEnableNewStyle(
-		TEXT("ContentBrowser.EnableNewStyle"),
-		bEnableNewStyle,
-		TEXT("Whether or not to enable the Content Browser restyle (UE 5.5+).")
-	);
-
 	FName FContentBrowserStyle::StyleName("ContentBrowser");
 
 	FContentBrowserStyle& FContentBrowserStyle::Get()
@@ -31,11 +24,6 @@ namespace UE::ContentBrowser::Private
 		return StyleName;
 	}
 
-	bool FContentBrowserStyle::IsNewStyleEnabled()
-	{
-		return bEnableNewStyle;
-	}
-
 	FContentBrowserStyle::FContentBrowserStyle()
 		: FSlateStyleSet(StyleName)
 	{
diff --git a/Engine/Source/Editor/ContentBrowser/Private/ContentBrowserStyle.h b/Engine/Source/Editor/ContentBrowser/Private/ContentBrowserStyle.h
--- a/Engine/Source/Editor/ContentBrowser/Private/ContentBrowserStyle.h
+++ b/Engine/Source/Editor/ContentBrowser/Private/ContentBrowserStyle.h
@@ -16,8 +16,6 @@ namespace UE::ContentBrowser::Private
 
 		virtual const FName& GetStyleSetName() const override;
 
-		static bool IsNewStyleEnabled();
-
 	private:
 		FContentBrowserStyle();
 		virtual ~FContentBrowserStyle() override;
