[Content Browser][New Style] Limited search bar width to 640px, moved breadcrumb to bottom

diff --git a/Engine/Plugins/Editor/ContentBrowser/ContentBrowserAssetDataSource/Source/ContentBrowserAssetDataSource/ContentBrowserAssetDataSource.Build.cs b/Engine/Plugins/Editor/ContentBrowser/ContentBrowserAssetDataSource/Source/ContentBrowserAssetDataSource/ContentBrowserAssetDataSource.Build.cs
--- a/Engine/Plugins/Editor/ContentBrowser/ContentBrowserAssetDataSource/Source/ContentBrowserAssetDataSource/ContentBrowserAssetDataSource.Build.cs
+++ b/Engine/Plugins/Editor/ContentBrowser/ContentBrowserAssetDataSource/Source/ContentBrowserAssetDataSource/ContentBrowserAssetDataSource.Build.cs
@@ -38,6 +38,7 @@ namespace UnrealBuildTool.Rules
 					"EditorWidgets",
 					"Projects",
 					"ToolMenus",
+					"ToolWidgets",
 					"Slate",
 					"SlateCore",
 					"InputCore",
diff --git a/Engine/Plugins/Editor/ContentBrowser/ContentBrowserAssetDataSource/Source/ContentBrowserAssetDataSource/Private/ContentBrowserAssetDataSource.cpp b/Engine/Plugins/Editor/ContentBrowser/ContentBrowserAssetDataSource/Source/ContentBrowserAssetDataSource/Private/ContentBrowserAssetDataSource.cpp
--- a/Engine/Plugins/Editor/ContentBrowser/ContentBrowserAssetDataSource/Source/ContentBrowserAssetDataSource/Private/ContentBrowserAssetDataSource.cpp
+++ b/Engine/Plugins/Editor/ContentBrowser/ContentBrowserAssetDataSource/Source/ContentBrowserAssetDataSource/Private/ContentBrowserAssetDataSource.cpp
@@ -42,6 +42,9 @@
 #include "ToolMenuSection.h"
 #include "Widgets/SBoxPanel.h"
 #include "Widgets/Text/STextBlock.h"
+#if UE_CONTENTBROWSER_NEW_STYLE
+#include "SActionButton.h"
+#endif
 
 #include UE_INLINE_GENERATED_CPP_BY_NAME(ContentBrowserAssetDataSource)
 
@@ -3543,7 +3546,16 @@ void UContentBrowserAssetDataSource::PopulateContentBrowserToolBar(UToolMenu* In
 	const UContentBrowserToolbarMenuContext* ContextObject = InMenu->FindContext<UContentBrowserToolbarMenuContext>();
 	checkf(ContextObject, TEXT("Required context UContentBrowserToolbarMenuContext was missing!"));
 
-	TSharedRef<SWidget> ImportButton = 
+#if UE_CONTENTBROWSER_NEW_STYLE
+	TSharedRef<SWidget> ImportButton =
+		SNew(SActionButton)
+		.ToolTipText(LOCTEXT("ImportTooltip", "Import assets from files to the currently selected folder"))
+		.OnClicked_UObject(this, &UContentBrowserAssetDataSource::OnImportClicked, ContextObject)
+		.IsEnabled_UObject(this, &UContentBrowserAssetDataSource::IsImportEnabled, ContextObject)
+		.Icon(FAppStyle::Get().GetBrush("Icons.Import"))
+		.Text(LOCTEXT("Import", "Import"));
+#else
+	TSharedRef<SWidget> ImportButton =
 		SNew(SButton)
 		.ButtonStyle(FAppStyle::Get(), "SimpleButton")
 		.ToolTipText(LOCTEXT("ImportTooltip", "Import assets from files to the currently selected folder"))
@@ -3571,6 +3583,7 @@ void UContentBrowserAssetDataSource::PopulateContentBrowserToolBar(UToolMenu* In
 				.Text(LOCTEXT("Import", "Import"))
 			]
 		];
+#endif
 
 	FToolMenuSection& Section = InMenu->FindOrAddSection("New");
 
diff --git a/Engine/Source/Editor/ContentBrowser/ContentBrowser.Build.cs b/Engine/Source/Editor/ContentBrowser/ContentBrowser.Build.cs
--- a/Engine/Source/Editor/ContentBrowser/ContentBrowser.Build.cs
+++ b/Engine/Source/Editor/ContentBrowser/ContentBrowser.Build.cs
@@ -6,6 +6,8 @@ public class ContentBrowser : ModuleRules
 {
 	public ContentBrowser(ReadOnlyTargetRules Target) : base(Target)
 	{
+		PublicDefinitions.Add("UE_CONTENTBROWSER_NEW_STYLE=0");
+
 		UnsafeTypeCastWarningLevel = WarningLevel.Error;
 
 		PrivateIncludePathModuleNames.AddRange(
diff --git a/Engine/Source/Editor/ContentBrowser/Private/SContentBrowser.cpp b/Engine/Source/Editor/ContentBrowser/Private/SContentBrowser.cpp
--- a/Engine/Source/Editor/ContentBrowser/Private/SContentBrowser.cpp
+++ b/Engine/Source/Editor/ContentBrowser/Private/SContentBrowser.cpp
@@ -146,6 +146,7 @@
 #include "Widgets/Layout/SBox.h"
 #include "Widgets/Layout/SExpandableArea.h"
 #include "Widgets/Layout/SSeparator.h"
+#include "Widgets/Layout/SSpacer.h"
 #include "Widgets/Layout/SSplitter.h"
 #include "Widgets/Layout/SWidgetSwitcher.h"
 #include "Widgets/Navigation/SBreadcrumbTrail.h"
@@ -154,6 +155,9 @@
 #include "Widgets/SNullWidget.h"
 #include "Widgets/SWidget.h"
 #include "Widgets/Text/STextBlock.h"
+#if UE_CONTENTBROWSER_NEW_STYLE
+#include "SActionButton.h"
+#endif
 
 class FTreeItem;
 struct FGeometry;
@@ -505,6 +509,7 @@ void SContentBrowser::Construct( const FArguments& InArgs, const FName& InInstan
 					.ColorAndOpacity(FSlateColor::UseForeground())
 					.Image(FAppStyle::Get().GetBrush("Icons.Settings"))
 				]
+#if !UE_CONTENTBROWSER_NEW_STYLE
 				+ SHorizontalBox::Slot()
 				.VAlign(VAlign_Center)
 				.Padding(4.0, 0.0f)
@@ -513,9 +518,244 @@ void SContentBrowser::Construct( const FArguments& InArgs, const FName& InInstan
 					.Text(LOCTEXT("Settings", "Settings"))
 					.ColorAndOpacity(FSlateColor::UseForeground())
 				]
+#endif
 			];
 	}
 
+#if UE_CONTENTBROWSER_NEW_STYLE
+	constexpr float PanelInsetPadding = 4.0f;
+
+	ChildSlot
+	[
+		SNew(SVerticalBox)
+
+		+ SVerticalBox::Slot()
+		.AutoHeight()
+		[
+			SNew(SSeparator)
+			.Thickness(2.0f)
+		]
+
+		// Assets/tree
+		+ SVerticalBox::Slot()
+		.FillHeight(1.0f)
+		.Padding(0.0f)
+		[
+			// The tree/assets splitter
+			SAssignNew(PathAssetSplitterPtr, SSplitter)
+			.PhysicalSplitterHandleSize(2.0f)
+
+			// Sources View
+			+ SSplitter::Slot()
+			.Resizable(true)
+			.SizeRule(SSplitter::SizeToContent)
+			.OnSlotResized(this, &SContentBrowser::OnPathViewBoxColumnResized)
+			[
+				SNew(SBox)
+				.Padding(FMargin(PanelInsetPadding, 0.0f, PanelInsetPadding, PanelInsetPadding))
+				.Visibility(this, &SContentBrowser::GetSourcesViewVisibility)
+				.WidthOverride(this, &SContentBrowser::GetPathViewBoxWidthOverride)
+				[
+					SNew(SBorder)
+					.Padding(FMargin(0))
+					.BorderImage(FAppStyle::GetBrush("Brushes.Recessed"))
+					[
+						// Note: If adding more widgets here, fix ContentBrowserSourcesWidgetSwitcherIndex and the code that uses it!
+						SAssignNew(SourcesWidgetSwitcher, SWidgetSwitcher)
+
+						// Paths View
+						+SWidgetSwitcher::Slot()
+						[
+							SAssignNew(PathFavoriteSplitterPtr, SSplitter)
+							.Clipping(EWidgetClipping::ClipToBounds)
+							.PhysicalSplitterHandleSize(2.0f)
+							.HitDetectionSplitterHandleSize(8.0f)
+							.Orientation(EOrientation::Orient_Vertical)
+							.MinimumSlotHeight(26.0f)
+							.Visibility( this, &SContentBrowser::GetSourcesViewVisibility )
+							+SSplitter::Slot()
+							.SizeRule(TAttribute<SSplitter::ESizeRule>(this, &SContentBrowser::GetFavoritesAreaSizeRule))
+							.MinSize(TAttribute<float>(this, &SContentBrowser::GetFavoritesAreaMinSize))
+							.Value(0.2f)
+							[
+								SNew(SBorder)
+								.BorderImage(FAppStyle::Get().GetBrush("Brushes.Header"))
+								.Padding(0.0f, 2.0f, 0.0f, 0.0f)
+								[
+									CreateFavoritesView(Config)
+								]
+							]
+
+							+SSplitter::Slot()
+							.SizeRule(TAttribute<SSplitter::ESizeRule>(this, &SContentBrowser::GetPathAreaSizeRule))
+							.MinSize(29.0f)
+							.Value(0.8f)
+							[
+								SNew(SBorder)
+								.BorderImage(FAppStyle::Get().GetBrush("Brushes.Header"))
+								.Padding(0.0f, 2.0f, 0.0f, 0.0f)
+								[
+									CreatePathView(Config)
+								]
+							]
+
+							+SSplitter::Slot()
+							.SizeRule(TAttribute<SSplitter::ESizeRule>(this, &SContentBrowser::GetCollectionsAreaSizeRule))
+							.MinSize(TAttribute<float>(this, &SContentBrowser::GetCollectionsAreaMinSize))
+							.Value(0.4f)
+							[
+								SNew(SBorder)
+								.BorderImage(FAppStyle::Get().GetBrush("Brushes.Header"))
+								.Padding(0.0f, 2.0f, 0.0f, 0.0f)
+								[
+									CreateDockedCollectionsView(Config)
+								]
+							]
+						]
+
+						// Collections View
+						+SWidgetSwitcher::Slot()
+						[
+							SNew(SBox)
+							.Padding(FMargin(0.0f, PanelInsetPadding, 0.0f, 0.0f))
+							[
+								CollectionViewPtr.ToSharedRef()
+							]
+						]
+					]
+				]
+			]
+
+			// Asset View
+			+ SSplitter::Slot()
+			.Value(0.75f)
+			[
+				SNew(SBox)
+				.Padding(FMargin(0.0f, 0.0f, 0.0f, PanelInsetPadding))
+				[
+					SNew(SVerticalBox)
+
+					+ SVerticalBox::Slot()
+					.AutoHeight()
+					[
+						SNew(SBorder)
+						.Padding(FMargin(3))
+						.BorderImage(bIsDrawer ? FStyleDefaults::GetNoBrush() : FAppStyle::Get().GetBrush("Brushes.Panel"))
+						[
+							SNew(SHorizontalBox)
+
+							+ SHorizontalBox::Slot()
+							.AutoWidth()
+							.HAlign(HAlign_Left)
+							.Padding(5, 0, 0, 0)
+							[
+								CreateToolBar(Config)
+							]
+
+							+ SHorizontalBox::Slot()
+							.AutoWidth()
+							.HAlign(HAlign_Right)
+							.VAlign(VAlign_Center)
+							[
+								CreateLockButton(Config)
+							]
+
+							+ SHorizontalBox::Slot()
+							.AutoWidth()
+							.HAlign(HAlign_Right)
+							.VAlign(VAlign_Center)
+							[
+								CreateDrawerDockButton(Config)
+							]
+
+							+ SHorizontalBox::Slot()
+							.Padding(5.0f, 0.0f, 0.0f, 0.0f)
+							.HAlign(HAlign_Right)
+							.VAlign(VAlign_Center)
+							[
+								ViewOptions
+							]
+						]
+					]
+
+					+ SVerticalBox::Slot()
+					[
+						CreateAssetView(Config)
+					]
+				]
+			]
+		]
+
+		// Navigation Bar
+		+ SVerticalBox::Slot()
+		.AutoHeight()
+		[
+			SNew(SBorder)
+			.Padding(FMargin(3))
+			.BorderImage(bIsDrawer ? FStyleDefaults::GetNoBrush() : FAppStyle::Get().GetBrush("Brushes.Panel"))
+			[
+				SNew(SHorizontalBox)
+
+				// History Back Button
+				+SHorizontalBox::Slot()
+				.Padding(10, 0, 0, 0)
+				.AutoWidth()
+				[
+					SNew(SButton)
+					.VAlign(EVerticalAlignment::VAlign_Center)
+					.ButtonStyle(FAppStyle::Get(), "SimpleButton")
+					.ToolTipText( this, &SContentBrowser::GetHistoryBackTooltip )
+					.ContentPadding( FMargin(1, 0) )
+					.OnClicked(this, &SContentBrowser::BackClicked)
+					.IsEnabled(this, &SContentBrowser::IsBackEnabled)
+					.AddMetaData<FTagMetaData>(FTagMetaData(TEXT("ContentBrowserHistoryBack")))
+					[
+						SNew(SImage)
+						.Image(FAppStyle::Get().GetBrush("Icons.CircleArrowLeft"))
+						.ColorAndOpacity(FSlateColor::UseForeground())
+					]
+				]
+
+				// History Forward Button
+				+ SHorizontalBox::Slot()
+				.Padding(2, 0, 0, 0)
+				.AutoWidth()
+				[
+					SNew(SButton)
+					.VAlign(EVerticalAlignment::VAlign_Center)
+					.ButtonStyle(FAppStyle::Get(), "SimpleButton")
+					.ToolTipText( this, &SContentBrowser::GetHistoryForwardTooltip )
+					.ContentPadding( FMargin(1, 0) )
+					.OnClicked(this, &SContentBrowser::ForwardClicked)
+					.IsEnabled(this, &SContentBrowser::IsForwardEnabled)
+					.AddMetaData<FTagMetaData>(FTagMetaData(TEXT("ContentBrowserHistoryForward")))
+					[
+						SNew(SImage)
+						.Image(FAppStyle::Get().GetBrush("Icons.CircleArrowRight"))
+						.ColorAndOpacity(FSlateColor::UseForeground())
+					]
+				]
+
+				// Path
+				+ SHorizontalBox::Slot()
+				.VAlign(VAlign_Center)
+				.HAlign(HAlign_Fill)
+				.FillWidth(1.0f)
+				.Padding(2, 0, 0, 0)
+				[
+					SAssignNew(NavigationBar, SNavigationBar)
+					.OnPathClicked(this, &SContentBrowser::OnPathClicked)
+					.GetPathMenuContent(this, &SContentBrowser::OnGetCrumbDelimiterContent)
+					.GetComboOptions(this, &SContentBrowser::GetRecentPaths)
+					.OnNavigateToPath(this, &SContentBrowser::OnNavigateToPath)
+					.OnCompletePrefix(this, &SContentBrowser::OnCompletePathPrefix)
+					.OnCanEditPathAsText(this, &SContentBrowser::OnCanEditPathAsText)
+					.AddMetaData<FTagMetaData>(FTagMetaData(TEXT("ContentBrowserPath")))
+				]
+			]
+		]
+	];
+#else
 	ChildSlot
 	[
 		SNew(SVerticalBox)
@@ -573,7 +813,7 @@ void SContentBrowser::Construct( const FArguments& InArgs, const FName& InInstan
 						.ColorAndOpacity(FSlateColor::UseForeground())
 					]
 				]
-				
+
 				// Path
 				+ SHorizontalBox::Slot()
 				.VAlign(VAlign_Center)
@@ -673,7 +913,7 @@ void SContentBrowser::Construct( const FArguments& InArgs, const FName& InInstan
 									CreateFavoritesView(Config)
 								]
 							]
-								
+
 							+SSplitter::Slot()
 							.SizeRule(TAttribute<SSplitter::ESizeRule>(this, &SContentBrowser::GetPathAreaSizeRule))
 							.MinSize(29.0f)
@@ -722,6 +962,7 @@ void SContentBrowser::Construct( const FArguments& InArgs, const FName& InInstan
 			]
 		]
 	];
+#endif
 
 	ExtendViewOptionsMenu(Config);
 
@@ -985,9 +1226,15 @@ TSharedRef<SWidget> SContentBrowser::CreateLockButton(const FContentBrowserConfi
 
 void SContentBrowser::OnFilterBarLayoutChanging(EFilterBarLayout NewLayout)
 {
-	const float SearchBoxDesiredWidth = 500.0f;
+#if UE_CONTENTBROWSER_NEW_STYLE
+	const FOptionalSize SearchBoxDesiredWidth;
+	constexpr float SearchBoxMaxWidth = 640.0f;
+#else
+	const FOptionalSize SearchBoxDesiredWidth = 500.0f;
+	constexpr float SearchBoxMaxWidth = 0.0f;
+#endif
 
-	if(NewLayout == EFilterBarLayout::Horizontal)
+	if (NewLayout == EFilterBarLayout::Horizontal)
 	{
 		TSharedPtr<SHorizontalBox> SearchBoxSlot =
 			SNew(SHorizontalBox)
@@ -1006,13 +1253,13 @@ void SContentBrowser::OnFilterBarLayoutChanging(EFilterBarLayout NewLayout)
 			.Padding(0, 4, 0, 0)
 			[
 				SNew(SBox)
-				 .VAlign(VAlign_Center)
-				 .WidthOverride(SearchBoxDesiredWidth)
-				 [
-					 SearchBoxPtr.ToSharedRef()	
-				 ]
+				.VAlign(VAlign_Center)
+				.WidthOverride(SearchBoxDesiredWidth)
+				[
+					SearchBoxPtr.ToSharedRef()
+				]
 			];
-			
+
 		/** We add the Combo Button and the Search Box to the FilterList itself, so that the filters wrap with them
 		 *	properly in the Horizontal Layout
 		 */
@@ -1030,14 +1277,13 @@ void SContentBrowser::OnFilterBarLayoutChanging(EFilterBarLayout NewLayout)
 			]
 
 			+ SVerticalBox::Slot()
-			.FillHeight( 1.0f )
-			.Padding( 0, 0 )
+			.FillHeight(1.0f)
+			.Padding(0, 0)
 			[
 				AssetViewPtr.ToSharedRef()
 			]
 
 		);
-		
 	}
 	else
 	{
@@ -1068,7 +1314,7 @@ void SContentBrowser::OnFilterBarLayoutChanging(EFilterBarLayout NewLayout)
 					.Padding(0.0f, 2.0f)
 					.AutoHeight()
 					[
-						
+
 						SNew(SBorder)
 						.BorderImage(FAppStyle::Get().GetBrush("Brushes.Header"))
 						.Padding(FMargin(8.0f, 6.0f))
@@ -1090,7 +1336,7 @@ void SContentBrowser::OnFilterBarLayoutChanging(EFilterBarLayout NewLayout)
 					]
 				]
 			]
-		
+
 			+ SSplitter::Slot()
 			.Value(0.88f)
 			[
@@ -1111,19 +1357,20 @@ void SContentBrowser::OnFilterBarLayoutChanging(EFilterBarLayout NewLayout)
 
 					+ SHorizontalBox::Slot()
 					.AutoWidth()
+					.MaxWidth(SearchBoxMaxWidth)
 					[
 						SNew(SBox)
-						 .VAlign(VAlign_Center)
-						 .WidthOverride(SearchBoxDesiredWidth)
-						 [
-							 SearchBoxPtr.ToSharedRef()	
-						 ]
+						.VAlign(VAlign_Center)
+						.WidthOverride(SearchBoxDesiredWidth)
+						[
+							SearchBoxPtr.ToSharedRef()
+						]
 					]
 				]
 
 				+ SVerticalBox::Slot()
-				.FillHeight( 1.0f )
-				.Padding( 0, 0 )
+				.FillHeight(1.0f)
+				.Padding(0, 0)
 				[
 					AssetViewPtr.ToSharedRef()
 				]
@@ -1164,17 +1411,17 @@ TSharedRef<SWidget> SContentBrowser::CreateAssetView(const FContentBrowserConfig
 	AssetViewPtr->SetFilterBar(FilterListPtr);
 
 	SearchBoxPtr = SNew(SAssetSearchBox)
-							.HintText(this, &SContentBrowser::GetSearchAssetsHintText)
-							.ShowSearchHistory(true)
-							.OnTextChanged(this, &SContentBrowser::OnSearchBoxChanged)
-							.OnTextCommitted(this, &SContentBrowser::OnSearchBoxCommitted)
-							.OnKeyDownHandler(this, &SContentBrowser::OnSearchKeyDown)
-							.OnSaveSearchClicked(this, &SContentBrowser::OnSaveSearchButtonClicked)
-							.OnAssetSearchBoxSuggestionFilter(this, &SContentBrowser::OnAssetSearchSuggestionFilter)
-							.OnAssetSearchBoxSuggestionChosen(this, &SContentBrowser::OnAssetSearchSuggestionChosen)
-							.DelayChangeNotificationsWhileTyping(true)
-							.Visibility((Config != nullptr ? Config->bCanShowAssetSearch : true) ? EVisibility::Visible : EVisibility::Collapsed)
-							.AddMetaData<FTagMetaData>(FTagMetaData(TEXT("ContentBrowserSearchAssets")));
+				   .HintText(this, &SContentBrowser::GetSearchAssetsHintText)
+				   .ShowSearchHistory(true)
+				   .OnTextChanged(this, &SContentBrowser::OnSearchBoxChanged)
+				   .OnTextCommitted(this, &SContentBrowser::OnSearchBoxCommitted)
+				   .OnKeyDownHandler(this, &SContentBrowser::OnSearchKeyDown)
+				   .OnSaveSearchClicked(this, &SContentBrowser::OnSaveSearchButtonClicked)
+				   .OnAssetSearchBoxSuggestionFilter(this, &SContentBrowser::OnAssetSearchSuggestionFilter)
+				   .OnAssetSearchBoxSuggestionChosen(this, &SContentBrowser::OnAssetSearchSuggestionChosen)
+				   .DelayChangeNotificationsWhileTyping(true)
+				   .Visibility((Config != nullptr ? Config->bCanShowAssetSearch : true) ? EVisibility::Visible : EVisibility::Collapsed)
+				   .AddMetaData<FTagMetaData>(FTagMetaData(TEXT("ContentBrowserSearchAssets")));
 
 	/* Create the Border that the Asset View will live in, the actual layout is populated in OnFilterBarLayoutChanging, 
 	 * which is called initially called through SAssetFilterBar::LoadSettings through SContentBrowser::LoadSettings()
@@ -1524,13 +1771,14 @@ void SContentBrowser::RegisterContentBrowserToolBar()
 
 	{
 		FToolMenuSection& Section = ToolBar->AddSection("New");
-		
+
 		Section.AddDynamicEntry("New", FNewToolMenuSectionDelegate::CreateLambda([](FToolMenuSection& InSection)
 		{
 			UContentBrowserToolbarMenuContext* Context = InSection.FindContext<UContentBrowserToolbarMenuContext>();
 			TSharedRef<SContentBrowser> ContentBrowser = Context->ContentBrowser.Pin().ToSharedRef();
 		
-			TSharedRef<SPositiveActionButton> NewButton = SNew(SPositiveActionButton)
+			TSharedRef<SPositiveActionButton> NewButton =
+				SNew(SPositiveActionButton)
 				.OnGetMenuContent_Lambda([Context] { return Context->ContentBrowser.Pin()->MakeAddNewContextMenu(EContentBrowserDataMenuContext_AddNewMenuDomain::Toolbar, Context); })
 				.ToolTipText(ContentBrowser, &SContentBrowser::GetAddNewToolTipText)
 				.IsEnabled(ContentBrowser, &SContentBrowser::IsAddNewEnabled)
@@ -1556,38 +1804,47 @@ void SContentBrowser::RegisterContentBrowserToolBar()
 			UContentBrowserToolbarMenuContext* Context = InSection.FindContext<UContentBrowserToolbarMenuContext>();
 			TSharedRef<SContentBrowser> ContentBrowser = Context->ContentBrowser.Pin().ToSharedRef();
 
-
-			TSharedRef<SButton> SaveButton =
-				SNew(SButton)
-				.ButtonStyle(FAppStyle::Get(), "SimpleButton")
-				.ToolTipText(LOCTEXT("SaveDirtyPackagesTooltip", "Save all modified assets."))
-				.ContentPadding(2.f)
-				.OnClicked(ContentBrowser, &SContentBrowser::OnSaveClicked)
-				.AddMetaData<FTagMetaData>(FTagMetaData(TEXT("ContentBrowserSaveDirtyPackages")))
-				[
-					SNew(SHorizontalBox)
-					// Save All Icon
-					+ SHorizontalBox::Slot()
-					.AutoWidth()
-					.HAlign(HAlign_Center)
-					.VAlign(VAlign_Center)
+#if UE_CONTENTBROWSER_NEW_STYLE
+				TSharedRef<SWidget> SaveButton =
+					SNew(SActionButton)
+					.ToolTipText(LOCTEXT("SaveDirtyPackagesTooltip", "Save all modified assets."))
+					.OnClicked(ContentBrowser, &SContentBrowser::OnSaveClicked)
+					.AddMetaData<FTagMetaData>(FTagMetaData(TEXT("ContentBrowserSaveDirtyPackages")))
+					.Icon(FAppStyle::Get().GetBrush("MainFrame.SaveAll"))
+					.Text(LOCTEXT("SaveAll", "Save All"));
+#else
+				TSharedRef<SWidget> SaveButton =
+					SNew(SButton)
+					.ButtonStyle(FAppStyle::Get(), "SimpleButton")
+					.ToolTipText(LOCTEXT("SaveDirtyPackagesTooltip", "Save all modified assets."))
+					.ContentPadding(2.f)
+					.OnClicked(ContentBrowser, &SContentBrowser::OnSaveClicked)
+					.AddMetaData<FTagMetaData>(FTagMetaData(TEXT("ContentBrowserSaveDirtyPackages")))
 					[
-						SNew(SImage)
-						.Image(FAppStyle::Get().GetBrush("MainFrame.SaveAll"))
-						.ColorAndOpacity(FSlateColor::UseForeground())
-					]
+						SNew(SHorizontalBox)
+						// Save All Icon
+						+ SHorizontalBox::Slot()
+						.AutoWidth()
+						.HAlign(HAlign_Center)
+						.VAlign(VAlign_Center)
+						[
+							SNew(SImage)
+							.Image(FAppStyle::Get().GetBrush("MainFrame.SaveAll"))
+							.ColorAndOpacity(FSlateColor::UseForeground())
+						]
 
-					// Save All Text
-					+ SHorizontalBox::Slot()
-					.Padding(FMargin(3, 0, 0, 0))
-					.VAlign(VAlign_Center)
-					.AutoWidth()
-					[
-						SNew(STextBlock)
-						.TextStyle(FAppStyle::Get(), "NormalText")
-						.Text(LOCTEXT("SaveAll", "Save All"))
-					]
-				];
+						// Save All Text
+						+ SHorizontalBox::Slot()
+						.Padding(FMargin(3, 0, 0, 0))
+						.VAlign(VAlign_Center)
+						.AutoWidth()
+						[
+							SNew(STextBlock)
+							.TextStyle(FAppStyle::Get(), "NormalText")
+							.Text(LOCTEXT("SaveAll", "Save All"))
+						]
+					];
+#endif
 
 			InSection.AddEntry(
 				FToolMenuEntry::InitWidget(
