[ContentBrowser] Thumbnail Sizes

diff --git a/Engine/Source/Editor/ContentBrowser/Private/AssetViewWidgets.cpp b/Engine/Source/Editor/ContentBrowser/Private/AssetViewWidgets.cpp
--- a/Engine/Source/Editor/ContentBrowser/Private/AssetViewWidgets.cpp
+++ b/Engine/Source/Editor/ContentBrowser/Private/AssetViewWidgets.cpp
@@ -84,7 +84,11 @@ class FDragDropEvent;
 
 #define LOCTEXT_NAMESPACE "ContentBrowser"
 
-static int32 GenericThumbnailSizes[(int32)EThumbnailSize::MAX] = { 24, 32, 64, 128, 200 };
+#if UE_CONTENTBROWSER_NEW_STYLE
+static int32 GenericThumbnailSizes[(int32)EThumbnailSize::MAX] = { 80, 96, 112, 128, 136, 160 };
+#else
+static int32 GenericThumbnailSizes[(int32)EThumbnailSize::MAX] = { 24, 32, 64, 128, 160, 200 };
+#endif
 
 ///////////////////////////////
 // FAssetViewModeUtils
@@ -1804,6 +1808,7 @@ void SAssetTileItem::Construct( const FArguments& InArgs )
 	AssetThumbnail = InArgs._AssetThumbnail;
 	ItemWidth = InArgs._ItemWidth;
 	ThumbnailPadding = InArgs._ThumbnailPadding;
+	ThumbnailDimension = InArgs._ThumbnailDimension;
 
 	CurrentThumbnailSize = InArgs._CurrentThumbnailSize;
 
@@ -1883,12 +1888,19 @@ void SAssetTileItem::Construct( const FArguments& InArgs )
 			.ColorAndOpacity(this, &SAssetTileItem::GetNameAreaTextColor);
 	}
 
+	constexpr float AssetViewWidgetsBorderPadding = 4.f;
+	constexpr float AssetViewWidgetsShadowPadding = 5.f;
+	constexpr float AssetViewWidgetsRenameWidgetPadding = 2.f;
+	constexpr float AssetViewWidgetsClassTextPadding = 2.f;
+	constexpr float AssetViewWidgetSourceControlSize = 16.f;
+
+#if UE_CONTENTBROWSER_NEW_STYLE
 	ChildSlot
-	.Padding(FMargin(0.0f, 0.0f, 4.0f, 4.0f))
+	.Padding(FMargin(0.0f, 0.0f, AssetViewWidgetsBorderPadding, AssetViewWidgetsBorderPadding))
 	[				
 		// Drop shadow border
 		SNew(SBorder)
-		.Padding(FMargin(0.0f, 0.0f, 5.0f, 5.0f))
+		.Padding(FMargin(0.0f, 0.0f, AssetViewWidgetsShadowPadding, AssetViewWidgetsShadowPadding))
 		.BorderImage(IsFolder() ? TAttribute<const FSlateBrush*>(this, &SAssetTileItem::GetFolderBackgroundShadowImage) : FAppStyle::Get().GetBrush(ItemShadowBorderName))
 		[
 			SNew(SOverlay)
@@ -1917,12 +1929,13 @@ void SAssetTileItem::Construct( const FArguments& InArgs )
 					+SVerticalBox::Slot()
 					[
 						SNew(SBorder)
+						.Visibility(this, &SAssetTileItem::GetNameAreaVisibility)
 						.Padding(FMargin(2.0f, 3.0f))
 						.BorderImage(this, &SAssetTileItem::GetNameAreaBackgroundImage)
 						[
 							SNew(SVerticalBox)
 							+ SVerticalBox::Slot()
-							.Padding(2.0f,2.0f,0.0f,0.0f)
+							.Padding(AssetViewWidgetsRenameWidgetPadding,AssetViewWidgetsRenameWidgetPadding,0.0f,0.0f)
 							.VAlign(VAlign_Top)
 							.HAlign(IsFolder() ? HAlign_Center : HAlign_Left)
 							[
@@ -1935,7 +1948,7 @@ void SAssetTileItem::Construct( const FArguments& InArgs )
 							+ SVerticalBox::Slot()
 							.VAlign(VAlign_Bottom)
 							.AutoHeight()
-							.Padding(2.0f,0.0f, 0.0f, 2.0f)
+							.Padding(AssetViewWidgetsClassTextPadding,0.0f, 0.0f, AssetViewWidgetsClassTextPadding)
 							[
 								SNew(SHorizontalBox)
 								+SHorizontalBox::Slot()
@@ -1954,8 +1967,8 @@ void SAssetTileItem::Construct( const FArguments& InArgs )
 								.HAlign(HAlign_Right)
 								[
 									SNew(SBox)
-									.WidthOverride(16.0f)
-									.HeightOverride(16.0f)
+									.WidthOverride(AssetViewWidgetSourceControlSize)
+									.HeightOverride(AssetViewWidgetSourceControlSize)
 									.Visibility(this, &SAssetTileItem::GetSCCIconVisibility)
 									[
 										GenerateSourceControlIconWidget()
@@ -1975,6 +1988,100 @@ void SAssetTileItem::Construct( const FArguments& InArgs )
 			]
 		]
 	];
+#else
+	ChildSlot
+	.Padding(FMargin(0.0f, 0.0f, AssetViewWidgetsBorderPadding, AssetViewWidgetsBorderPadding))
+	[				
+		// Drop shadow border
+		SNew(SBorder)
+		.Padding(FMargin(0.0f, 0.0f, AssetViewWidgetsShadowPadding, AssetViewWidgetsShadowPadding))
+		.BorderImage(IsFolder() ? TAttribute<const FSlateBrush*>(this, &SAssetTileItem::GetFolderBackgroundShadowImage) : FAppStyle::Get().GetBrush(ItemShadowBorderName))
+		[
+			SNew(SOverlay)
+			.AddMetaData<FTagMetaData>(FTagMetaData(AssetItem->GetItem().GetVirtualPath()))
+			+SOverlay::Slot()
+			[
+				SNew(SBorder)
+				.Padding(0)
+				.BorderImage(IsFolder() ? TAttribute<const FSlateBrush*>(this, &SAssetTileItem::GetFolderBackgroundImage) : UE::ContentBrowser::Private::FContentBrowserStyle::Get().Get().GetBrush("ContentBrowser.AssetTileItem.ThumbnailAreaBackground"))
+				[
+					SNew(SVerticalBox)
+					// Thumbnail
+					+SVerticalBox::Slot()
+					.AutoHeight()
+					[
+						// The remainder of the space is reserved for the name.
+						SNew(SBox)
+						.Padding(0)
+						.WidthOverride(this, &SAssetTileItem::GetThumbnailBoxSize)
+						.HeightOverride(this, &SAssetTileItem::GetThumbnailBoxSize)
+						[
+							ItemContents
+						]
+					]
+
+					+SVerticalBox::Slot()
+					[
+						SNew(SBorder)
+						.Padding(FMargin(2.0f, 3.0f))
+						.BorderImage(this, &SAssetTileItem::GetNameAreaBackgroundImage)
+						[
+							SNew(SVerticalBox)
+							+ SVerticalBox::Slot()
+							.Padding(AssetViewWidgetsRenameWidgetPadding, AssetViewWidgetsRenameWidgetPadding, 0.0f, 0.0f)
+							.VAlign(VAlign_Top)
+							.HAlign(IsFolder() ? HAlign_Center : HAlign_Left)
+							[
+								SNew(SBox)
+								.MaxDesiredHeight(this, &SAssetTileItem::GetNameAreaMaxDesiredHeight)
+								[
+									InlineRenameWidget.ToSharedRef()
+								]
+							]
+							+ SVerticalBox::Slot()
+							.VAlign(VAlign_Bottom)
+							.AutoHeight()
+							.Padding(AssetViewWidgetsClassTextPadding, 0.0f, 0.0f, AssetViewWidgetsClassTextPadding)
+							[
+								SNew(SHorizontalBox)
+								+SHorizontalBox::Slot()
+								[
+									SAssignNew(ClassTextWidget, STextBlock)
+									.Visibility(this, &SAssetTileItem::GetAssetClassLabelVisibility)
+									.TextStyle(UE::ContentBrowser::Private::FContentBrowserStyle::Get().Get(), "ContentBrowser.ClassFont")
+									.OverflowPolicy(ETextOverflowPolicy::Ellipsis)
+									.Text(this, &SAssetTileItem::GetAssetClassText)
+									.ColorAndOpacity(this, &SAssetTileItem::GetAssetClassLabelTextColor)
+								]
+
+								+SHorizontalBox::Slot()
+								.Padding(FMargin(3.0f, 0.0f, 0.0f, 1.0f))
+								.AutoWidth()
+								.HAlign(HAlign_Right)
+								[
+									SNew(SBox)
+									.WidthOverride(AssetViewWidgetSourceControlSize)
+									.HeightOverride(AssetViewWidgetSourceControlSize)
+									.Visibility(this, &SAssetTileItem::GetSCCIconVisibility)
+									[
+										GenerateSourceControlIconWidget()
+									]
+								]
+							]
+							
+						]
+					]
+				]
+			]
+			+SOverlay::Slot()
+			[
+				SNew(SImage)
+				.Image(this, &SAssetViewItem::GetBorderImage)
+				.Visibility(EVisibility::HitTestInvisible)	
+			]
+		]
+	];
+#endif
 
 	HandleSourceControlStateChanged();
 
@@ -2058,13 +2165,18 @@ FOptionalSize SAssetTileItem::GetExtraStateIconMaxWidth() const
 
 FOptionalSize SAssetTileItem::GetStateIconImageSize() const
 {
+	constexpr int32 SourceControlImageMinSize = 12;
 	float IconSize = FMath::TruncToFloat(GetThumbnailBoxSize().Get() * 0.2f);
-	return IconSize > 12 ? IconSize : 12;
+	return IconSize > SourceControlImageMinSize ? IconSize : SourceControlImageMinSize;
 }
 
 FOptionalSize SAssetTileItem::GetThumbnailBoxSize() const
 {
+#if UE_CONTENTBROWSER_NEW_STYLE
+	return FOptionalSize(ThumbnailDimension.Get());
+#else
 	return FOptionalSize(ItemWidth.Get()- ThumbnailPadding);
+#endif
 }
 
 EVisibility SAssetTileItem::GetAssetClassLabelVisibility() const
@@ -2202,6 +2314,11 @@ EVisibility SAssetTileItem::GetSCCIconVisibility() const
 	return bHasCCStateBrush &&  CurrentThumbnailSize.Get() != EThumbnailSize::Tiny && ISourceControlModule::Get().IsEnabled() && ISourceControlModule::Get().GetProvider().IsAvailable() ? EVisibility::Visible : EVisibility::Collapsed;
 }
 
+EVisibility SAssetTileItem::GetNameAreaVisibility() const
+{
+	return CurrentThumbnailSize.Get() == EThumbnailSize::Tiny ? EVisibility::Collapsed : EVisibility::Visible;
+}
+
 
 void SAssetTileItem::InitializeAssetNameHeights()
 {
@@ -2217,7 +2334,9 @@ void SAssetTileItem::InitializeAssetNameHeights()
 			FSlateFontInfo Font = FAppStyle::GetFontStyle(SmallFontName);
 			TSharedRef<FSlateFontMeasure> FontMeasureService = FSlateApplication::Get().GetRenderer()->GetFontMeasureService();
 			SmallFontHeight = FontMeasureService->GetMaxCharacterHeight(Font);
-			AssetNameHeights[(int32)EThumbnailSize::Small] = SmallFontHeight * 2 ;
+
+			constexpr float SmallSizeMultiplier = 2;
+			AssetNameHeights[(int32)EThumbnailSize::Small] = SmallFontHeight * SmallSizeMultiplier;
 		}
 
 
@@ -2226,9 +2345,20 @@ void SAssetTileItem::InitializeAssetNameHeights()
 			FSlateFontInfo Font = FAppStyle::GetFontStyle(SmallFontName);
 			TSharedRef<FSlateFontMeasure> FontMeasureService = FSlateApplication::Get().GetRenderer()->GetFontMeasureService();
 			RegularFontHeight = FontMeasureService->GetMaxCharacterHeight(Font);
-			AssetNameHeights[(int32)EThumbnailSize::Medium] = RegularFontHeight * 3;
-			AssetNameHeights[(int32)EThumbnailSize::Large] = RegularFontHeight * 4;
-			AssetNameHeights[(int32)EThumbnailSize::Huge] = RegularFontHeight * 5;
+
+			constexpr float MediumSizeMultiplier = 3;
+			constexpr float LargeSizeMultiplier = 4;
+			constexpr float XLargeSizeMultiplier = 5;
+			constexpr float HugeSizeMultiplier = 6;
+
+			AssetNameHeights[(int32)EThumbnailSize::Medium] = RegularFontHeight * MediumSizeMultiplier;
+			AssetNameHeights[(int32)EThumbnailSize::Large] = RegularFontHeight * LargeSizeMultiplier;
+#if UE_CONTENTBROWSER_NEW_STYLE
+			AssetNameHeights[(int32)EThumbnailSize::XLarge] = RegularFontHeight * XLargeSizeMultiplier;
+			AssetNameHeights[(int32)EThumbnailSize::Huge] = RegularFontHeight * HugeSizeMultiplier;
+#else
+			AssetNameHeights[(int32)EThumbnailSize::Huge] = RegularFontHeight * XLargeSizeMultiplier;
+#endif
 		}
 
 		bInitializedHeights = true;
diff --git a/Engine/Source/Editor/ContentBrowser/Private/SAssetView.cpp b/Engine/Source/Editor/ContentBrowser/Private/SAssetView.cpp
--- a/Engine/Source/Editor/ContentBrowser/Private/SAssetView.cpp
+++ b/Engine/Source/Editor/ContentBrowser/Private/SAssetView.cpp
@@ -210,6 +210,8 @@ FText SAssetView::ThumbnailSizeToDisplayName(EThumbnailSize InSize)
 		return LOCTEXT("MediumThumbnailSize", "Medium");
 	case EThumbnailSize::Large:
 		return LOCTEXT("LargeThumbnailSize", "Large");
+	case EThumbnailSize::XLarge:
+		return LOCTEXT("XLargeThumbnailSize", "X Large");
 	case EThumbnailSize::Huge:
 		return LOCTEXT("HugeThumbnailSize", "Huge");
 	default:
@@ -1103,11 +1105,21 @@ void SAssetView::Construct( const FArguments& InArgs )
 	ListViewThumbnailSize = 64;
 	ListViewThumbnailPadding = 4;
 	TileViewThumbnailResolution = 256;
-	TileViewThumbnailSize = 150;
+
+	// Max Size for the thumbnail
+#if UE_CONTENTBROWSER_NEW_STYLE
+	constexpr int32 MaxTileViewThumbnailSize = 160;
+#else
+	constexpr int32 MaxTileViewThumbnailSize = 150;
+#endif
+
+	TileViewThumbnailSize = MaxTileViewThumbnailSize;
+
 	TileViewThumbnailPadding = 9;
 
 	TileViewNameHeight = 50;
 
+	UpdateThumbnailSizeValue();
 	MinThumbnailScale = 0.2f * ThumbnailScaleRangeScalar;
 	MaxThumbnailScale = 1.9f * ThumbnailScaleRangeScalar;
 
@@ -1798,7 +1810,8 @@ void SAssetView::SaveSettings(const FString& IniFilename, const FString& IniSect
 {
 	GConfig->SetInt(*IniSection, *GetThumbnailScaleSettingPath(SettingsString), (int32)ThumbnailSize, IniFilename);
 	GConfig->SetInt(*IniSection, *GetCurrentViewTypeSettingPath(SettingsString), CurrentViewType, IniFilename);
-	
+	GConfig->SetFloat(*IniSection, *(SettingsString + TEXT(".ZoomScale")), ZoomScale, IniFilename);
+
 	GConfig->SetArray(*IniSection, *(SettingsString + TEXT(".HiddenColumns")), HiddenColumnNames, IniFilename);
 }
 
@@ -1810,7 +1823,16 @@ void SAssetView::LoadSettings(const FString& IniFilename, const FString& IniSect
 		// Clamp value to normal range and update state
 		ThumbnailSizeConfig = FMath::Clamp<int32>(ThumbnailSizeConfig, 0, (int32)EThumbnailSize::MAX-1);
 
+		// TODO: Remove this afterwards, current CB should hide new size
+#if !UE_CONTENTBROWSER_NEW_STYLE
+		if ((EThumbnailSize)ThumbnailSizeConfig == EThumbnailSize::XLarge)
+		{
+			ThumbnailSizeConfig -= 1;
+		}
+#endif
 		ThumbnailSize = (EThumbnailSize) ThumbnailSizeConfig;
+		// Set the thumbnail value here after the Size enum is retrieved from the config
+		UpdateThumbnailSizeValue();
 	}
 
 	int32 ViewType = EAssetViewType::Tile;
@@ -1823,7 +1845,14 @@ void SAssetView::LoadSettings(const FString& IniFilename, const FString& IniSect
 		}
 		SetCurrentViewType( (EAssetViewType::Type)ViewType );
 	}
-	
+
+	float Zoom = 0;
+	if ( GConfig->GetFloat(*IniSection, *(SettingsString + TEXT(".ZoomScale")), Zoom, IniFilename) )
+	{
+		// Clamp value to normal range and update state
+		ZoomScale = FMath::Clamp<float>(Zoom, 0.f, 1.f);
+	}
+
 	TArray<FString> LoadedHiddenColumnNames;
 	GConfig->GetArray(*IniSection, *(SettingsString + TEXT(".HiddenColumns")), LoadedHiddenColumnNames, IniFilename);
 	if (LoadedHiddenColumnNames.Num() > 0)
@@ -2408,14 +2437,47 @@ FReply SAssetView::OnMouseWheel( const FGeometry& MyGeometry, const FPointerEven
 	// Make sure to not change the thumbnail scaling when we're in Columns view since thumbnail scaling isn't applicable there.
 	if( MouseEvent.IsControlDown() && IsThumbnailScalingAllowed() )
 	{
+#if UE_CONTENTBROWSER_NEW_STYLE
+		const float NewDelta = MouseEvent.GetWheelDelta() * 0.4f;
+		if (ZoomScale == 1.f && NewDelta > 0 || ZoomScale == 0.f && NewDelta < 0)
+		{
+			int32 Step = (int32)FMath::Sign(NewDelta);
+			EThumbnailSize OldSize = ThumbnailSize;
+			ThumbnailSize = (EThumbnailSize)FMath::Clamp<int32>(((int32)ThumbnailSize + Step), 0, (int32)EThumbnailSize::MAX - 1);
+			if (OldSize != ThumbnailSize)
+			{
+				OnThumbnailSizeChanged(ThumbnailSize);
+				ZoomScale = NewDelta > 0.f ? 0.f : 1.f;
+			}
+		}
+		else
+		{
+			ZoomScale = FMath::Clamp(ZoomScale + NewDelta, 0.f, 1.f);
+		}
+#else
 		// Step up/down a level depending on the scroll wheel direction.
 		// Clamp value to enum min/max before updating.
 		const int32 Delta = MouseEvent.GetWheelDelta() > 0 ? 1 : -1;
-		const EThumbnailSize DesiredThumbnailSize = (EThumbnailSize)FMath::Clamp<int32>((int32)ThumbnailSize + Delta, 0, (int32)EThumbnailSize::MAX - 1);
+		EThumbnailSize DesiredThumbnailSize = (EThumbnailSize)FMath::Clamp<int32>((int32)ThumbnailSize + Delta, 0, (int32)EThumbnailSize::MAX - 1);
+
+		// TODO: Remove this afterwards, current CB should hide new size
+		if (DesiredThumbnailSize == EThumbnailSize::XLarge)
+		{
+			if (Delta > 0)
+			{
+				DesiredThumbnailSize = EThumbnailSize::Huge;
+			}
+			else
+			{
+				DesiredThumbnailSize = EThumbnailSize::Large;
+			}
+		}
+
 		if ( DesiredThumbnailSize != ThumbnailSize )
 		{
 			OnThumbnailSizeChanged(DesiredThumbnailSize);
-		}		
+		}
+#endif
 		return FReply::Handled();
 	}
 	return FReply::Unhandled();
@@ -3946,6 +4008,13 @@ void SAssetView::PopulateViewButtonMenu(UToolMenu* Menu)
 			
 			for (int32 EnumValue = (int32)EThumbnailSize::Tiny; EnumValue < (int32)EThumbnailSize::MAX; ++EnumValue)
 			{
+#if !UE_CONTENTBROWSER_NEW_STYLE
+				if ((EThumbnailSize)EnumValue == EThumbnailSize::XLarge)
+				{
+					continue;
+				}
+#endif
+
 				SizeSection.AddMenuEntry(
 					NAME_None,
 					SAssetView::ThumbnailSizeToDisplayName((EThumbnailSize)EnumValue),
@@ -4772,6 +4841,8 @@ TSharedRef<ITableRow> SAssetView::MakeTileViewWidget(TSharedPtr<FAssetViewItem>
 		TSharedRef<SAssetTileItem> Item =
 			SNew(SAssetTileItem)
 			.AssetItem(AssetItem)
+			.CurrentThumbnailSize(this, &SAssetView::GetThumbnailSize)
+			.ThumbnailDimension(this, &SAssetView::GetTileViewThumbnailDimension)
 			.ThumbnailPadding((float)TileViewThumbnailPadding)
 			.ItemWidth(this, &SAssetView::GetTileViewItemWidth)
 			.OnRenameBegin(this, &SAssetView::AssetRenameBegin)
@@ -4810,6 +4881,7 @@ TSharedRef<ITableRow> SAssetView::MakeTileViewWidget(TSharedPtr<FAssetViewItem>
 			.AssetItem(AssetItem)
 			.ThumbnailPadding((float)TileViewThumbnailPadding)
 			.CurrentThumbnailSize(this, &SAssetView::GetThumbnailSize)
+			.ThumbnailDimension(this, &SAssetView::GetTileViewThumbnailDimension)
 			.ItemWidth(this, &SAssetView::GetTileViewItemWidth)
 			.OnRenameBegin(this, &SAssetView::AssetRenameBegin)
 			.OnRenameCommit(this, &SAssetView::AssetRenameCommit)
@@ -4999,7 +5071,14 @@ TSharedPtr<FAssetThumbnail> SAssetView::AddItemToNewThumbnailRelevancyMap(const
 		}
 
 		// The thumbnail newly relevant, create a new thumbnail
-		const int32 ThumbnailResolution = FMath::TruncToInt((float)CurrentThumbnailSize * MaxThumbnailScale);
+		int32 ThumbnailResolution = 0;
+
+#if UE_CONTENTBROWSER_NEW_STYLE
+		ThumbnailResolution = CurrentThumbnailSize;
+#else
+		ThumbnailResolution = FMath::TruncToInt((float)CurrentThumbnailSize * MaxThumbnailScale);
+#endif
+
 		Thumbnail = MakeShared<FAssetThumbnail>(FAssetData(), ThumbnailResolution, ThumbnailResolution, AssetThumbnailPool);
 		Item->GetItem().UpdateThumbnail(*Thumbnail);
 		Thumbnail->GetViewportRenderTargetTexture(); // Access the texture once to trigger it to render
@@ -5397,6 +5476,7 @@ void SAssetView::ToggleThumbnailEditMode()
 void SAssetView::OnThumbnailSizeChanged(EThumbnailSize NewThumbnailSize)
 {
 	ThumbnailSize = NewThumbnailSize;
+	UpdateThumbnailSizeValue();
 
 	if (FAssetViewInstanceConfig* Config = GetAssetViewConfig())
 	{
@@ -5429,6 +5509,9 @@ float SAssetView::GetThumbnailScale() const
 	case EThumbnailSize::Large:
 		BaseScale = 0.75f;
 		break;
+	case EThumbnailSize::XLarge:
+		BaseScale = 0.9f;
+		break;
 	case EThumbnailSize::Huge:
 		BaseScale = 1.0f;
 		break;
@@ -5440,6 +5523,46 @@ float SAssetView::GetThumbnailScale() const
 	return BaseScale * GetTickSpaceGeometry().Scale;
 }
 
+float SAssetView::GetThumbnailSizeValue() const
+{
+	return FMath::Lerp(MinThumbnailSize, MaxThumbnailSize, ZoomScale);
+}
+
+void SAssetView::UpdateThumbnailSizeValue()
+{
+	switch (ThumbnailSize)
+	{
+	case EThumbnailSize::Tiny:
+		MinThumbnailSize = 64.f;
+		MaxThumbnailSize = 80.f;
+		break;
+	case EThumbnailSize::Small:
+		MinThumbnailSize = 80.f;
+		MaxThumbnailSize = 96.f;
+		break;
+	case EThumbnailSize::Medium:
+		MinThumbnailSize = 96.f;
+		MaxThumbnailSize = 112.f;
+		break;
+	case EThumbnailSize::Large:
+		MinThumbnailSize = 112.f;
+		MaxThumbnailSize = 128.f;
+		break;
+	case EThumbnailSize::XLarge:
+		MinThumbnailSize = 128.f;
+		MaxThumbnailSize = 136.f;
+		break;
+	case EThumbnailSize::Huge:
+		MinThumbnailSize = 136.f;
+		MaxThumbnailSize = 160.f;
+		break;
+	default:
+		MinThumbnailSize = 64.f;
+		MaxThumbnailSize = 80.f;
+		break;
+	}
+}
+
 bool SAssetView::IsThumbnailScalingAllowed() const
 {
 	return GetCurrentViewType() != EAssetViewType::Column;
@@ -5447,6 +5570,13 @@ bool SAssetView::IsThumbnailScalingAllowed() const
 
 float SAssetView::GetTileViewTypeNameHeight() const
 {
+#if UE_CONTENTBROWSER_NEW_STYLE
+		if (ThumbnailSize == EThumbnailSize::Tiny)
+		{
+			return 0;
+		}
+		return 67.f;
+#else
 	float TypeNameHeight = 0;
 
 	if (bShowTypeInTileView)
@@ -5473,6 +5603,7 @@ float SAssetView::GetTileViewTypeNameHeight() const
 		}
 	}
 	return TypeNameHeight;
+#endif
 }
 
 float SAssetView::GetSourceControlIconHeight() const
@@ -5487,22 +5618,43 @@ float SAssetView::GetListViewItemHeight() const
 
 float SAssetView::GetTileViewItemHeight() const
 {
+#if UE_CONTENTBROWSER_NEW_STYLE
+	return GetTileViewItemBaseWidth() + GetTileViewTypeNameHeight() + TileViewWidthHeightPadding;
+#else
 	return (((float)TileViewNameHeight + GetTileViewTypeNameHeight()) * FMath::Lerp(MinThumbnailScale, MaxThumbnailScale, GetThumbnailScale())) + GetTileViewItemBaseHeight() * FillScale + GetSourceControlIconHeight();
+#endif
 }
 
 float SAssetView::GetTileViewItemBaseHeight() const
 {
+#if UE_CONTENTBROWSER_NEW_STYLE
+	return GetTileViewItemBaseWidth();
+#else
 	return (float)(TileViewThumbnailSize + TileViewThumbnailPadding * 2) * FMath::Lerp(MinThumbnailScale, MaxThumbnailScale, GetThumbnailScale());
+#endif
 }
 
 float SAssetView::GetTileViewItemWidth() const
 {
+#if UE_CONTENTBROWSER_NEW_STYLE
+	return GetTileViewItemBaseWidth() + TileViewWidthHeightPadding;
+#else
 	return GetTileViewItemBaseWidth() * FillScale;
+#endif
+}
+
+float SAssetView::GetTileViewThumbnailDimension() const
+{
+	return GetThumbnailSizeValue();
 }
 
 float SAssetView::GetTileViewItemBaseWidth() const //-V524
 {
+#if UE_CONTENTBROWSER_NEW_STYLE
+	return GetTileViewThumbnailDimension();
+#else
 	return (float)( TileViewThumbnailSize + TileViewThumbnailPadding * 2 ) * FMath::Lerp( MinThumbnailScale, MaxThumbnailScale, GetThumbnailScale() );
+#endif
 }
 
 EColumnSortMode::Type SAssetView::GetColumnSortMode(const FName ColumnId) const
diff --git a/Engine/Source/Editor/ContentBrowser/Public/AssetViewWidgets.h b/Engine/Source/Editor/ContentBrowser/Public/AssetViewWidgets.h
--- a/Engine/Source/Editor/ContentBrowser/Public/AssetViewWidgets.h
+++ b/Engine/Source/Editor/ContentBrowser/Public/AssetViewWidgets.h
@@ -588,6 +588,8 @@ class SAssetTileItem : public SAssetViewItem
 public:
 	SLATE_BEGIN_ARGS( SAssetTileItem )
 		: _ThumbnailPadding(0)
+		/** Default Size of the Small Thumbnail Size */
+		, _ThumbnailDimension(64)
 		, _ThumbnailLabel( EThumbnailLabel::ClassName )
 		, _ThumbnailHintColorAndOpacity( FLinearColor( 0.0f, 0.0f, 0.0f, 0.0f ) )
 		, _AllowThumbnailHintLabel(true)
@@ -607,10 +609,13 @@ public:
 		/** How much padding to allow around the thumbnail */
 		SLATE_ARGUMENT( float, ThumbnailPadding )
 
-		/** The contents of the label displayed on the thumbnail */
+		/** The Actual Thumbnail dimension based on the  */
+		SLATE_ATTRIBUTE(float, ThumbnailDimension)
+
+		/** The contents of the label displayed on the thumbnail EThumbnailSize */
 		SLATE_ARGUMENT( EThumbnailLabel::Type, ThumbnailLabel )
 
-		/**  */
+		/** Thumbnail Color for Hint */
 		SLATE_ATTRIBUTE( FLinearColor, ThumbnailHintColorAndOpacity )
 
 		/** Whether the thumbnail should ever show it's hint label */
@@ -732,6 +737,8 @@ protected:
 	/** Gets the visibility of the SCC icons */
 	EVisibility GetSCCIconVisibility() const;
 
+	/** Get the visibility for the name area, collapsed if Tiny */
+	EVisibility GetNameAreaVisibility() const;
 
 private:
 	/** If false, the tooltip will not be displayed */
@@ -743,6 +750,9 @@ private:
 	/** The width of the item. Used to enforce a square thumbnail. */
 	TAttribute<float> ItemWidth;
 
+	/** Actual size that the thumbnail should have */
+	TAttribute<float> ThumbnailDimension;
+
 	/** Max name height for each thumbnail size */
 	static float AssetNameHeights[(int32)EThumbnailSize::MAX];
 
diff --git a/Engine/Source/Editor/ContentBrowser/Public/SAssetView.h b/Engine/Source/Editor/ContentBrowser/Public/SAssetView.h
--- a/Engine/Source/Editor/ContentBrowser/Public/SAssetView.h
+++ b/Engine/Source/Editor/ContentBrowser/Public/SAssetView.h
@@ -805,6 +805,12 @@ private:
 	/** Gets the current thumbnail scale */
 	float GetThumbnailScale() const;
 
+	/** Gets the current thumbnail size */
+	float GetThumbnailSizeValue() const;
+
+	/** Set the Min/Max Thumbnail size based on the EThumbnailSize chosen */
+	void UpdateThumbnailSizeValue();
+
 	/** Gets the current thumbnail size enum */
 	EThumbnailSize GetThumbnailSize() const { return ThumbnailSize; }
 
@@ -820,6 +826,9 @@ private:
 	/** Gets the final scaled item height for the tile view */
 	float GetTileViewItemHeight() const;
 
+	/** Get the TileView Thumbnail dimension height and width are the same for the thumbnail itself */
+	float GetTileViewThumbnailDimension() const;
+
 	/** Gets the scaled item height for the tile view before the filler scale is applied */
 	float GetTileViewItemBaseHeight() const;
 
@@ -967,6 +976,14 @@ private:
 	TSharedPtr<FAssetFilterCollectionType> FrontendFilters;
 	TSharedPtr<FAssetTextFilter> TextFilter;
 
+	/** Scale when using CTRL+Wheel, will go from 0.f to 1.f and reset accordingly when ThumbnailSize chosen changes\n
+	 * When going over/under the limit will also change the ThumbnailSize
+	 */
+	float ZoomScale = 0.0f;
+
+	/** Vertical and Horizontal padding for the TileViewItem */
+	const int32 TileViewWidthHeightPadding = 9;
+
 	TAttribute<bool> bShowRedirectors;
 	bool bLastShowRedirectors;
 
@@ -1082,6 +1099,10 @@ private:
 	float MinThumbnailScale;
 	float MaxThumbnailScale;
 
+	/** The max and min thumbnail sizes */
+	float MinThumbnailSize;
+	float MaxThumbnailSize;
+
 	/** Scalar applied to thumbnail sizes so that users thumbnails are scaled based on users display area size*/
 	float ThumbnailScaleRangeScalar;
 
diff --git a/Engine/Source/Editor/UMGEditor/Private/Library/SLibraryView.cpp b/Engine/Source/Editor/UMGEditor/Private/Library/SLibraryView.cpp
--- a/Engine/Source/Editor/UMGEditor/Private/Library/SLibraryView.cpp
+++ b/Engine/Source/Editor/UMGEditor/Private/Library/SLibraryView.cpp
@@ -156,6 +156,12 @@ TSharedRef<SWidget> SLibraryView::ConstructViewOptions()
 		{
 			for (int32 EnumValue = (int32)EThumbnailSize::Tiny; EnumValue < (int32)EThumbnailSize::MAX; ++EnumValue)
 			{
+#if !UE_CONTENTBROWSER_NEW_STYLE
+				if ((EThumbnailSize)EnumValue == EThumbnailSize::XLarge)
+				{
+					continue;
+				}
+#endif
 				SubMenu.AddMenuEntry(
 					SAssetView::ThumbnailSizeToDisplayName((EThumbnailSize)EnumValue),
 					FText::GetEmpty(),
diff --git a/Engine/Source/Editor/UnrealEd/Public/AssetThumbnail.h b/Engine/Source/Editor/UnrealEd/Public/AssetThumbnail.h
--- a/Engine/Source/Editor/UnrealEd/Public/AssetThumbnail.h
+++ b/Engine/Source/Editor/UnrealEd/Public/AssetThumbnail.h
@@ -33,6 +33,7 @@ enum class EThumbnailSize : uint8
 	Small,
 	Medium,
 	Large,
+	XLarge,
 	Huge,
 
 	// Not a size
