[ContentBrowser] Missing keybind for edit asset in the content browser

diff --git a/Engine/Source/Editor/ContentBrowser/Private/AssetContextMenu.cpp b/Engine/Source/Editor/ContentBrowser/Private/AssetContextMenu.cpp
--- a/Engine/Source/Editor/ContentBrowser/Private/AssetContextMenu.cpp
+++ b/Engine/Source/Editor/ContentBrowser/Private/AssetContextMenu.cpp
@@ -327,6 +327,19 @@ TSharedRef<SWidget> FAssetContextMenu::MakeContextMenu(TArrayView<const FContent
 		}
 	}
 
+	ContextObject->bCanView = false;
+	if (!SelectedItems.IsEmpty())
+	{
+		for (const FContentBrowserItem& SelectedItem : SelectedItems)
+		{
+			if (SelectedItem.CanView())
+			{
+				ContextObject->bCanView = true;
+				break;
+			}
+		}
+	}
+
 	FToolMenuContext MenuContext(InCommandList, MenuExtender, ContextObject);
 
 	{
@@ -334,6 +347,7 @@ TSharedRef<SWidget> FAssetContextMenu::MakeContextMenu(TArrayView<const FContent
 		DataContextObject->SelectedItems = SelectedItems;
 		DataContextObject->SelectedCollections = SourcesData.Collections;
 		DataContextObject->bCanBeModified = ContextObject->bCanBeModified;
+		DataContextObject->bCanView = ContextObject->bCanView;
 		DataContextObject->bHasCookedPackages = ContextObject->bHasCookedPackages;
 		DataContextObject->bContainsUnsupportedAssets = ContextObject->bContainsUnsupportedAssets;
 		DataContextObject->ParentWidget = AssetView;
@@ -497,11 +511,6 @@ void FAssetContextMenu::SetOnDuplicateRequested(const FOnDuplicateRequested& InO
 	OnDuplicateRequested = InOnDuplicateRequested;
 }
 
-void FAssetContextMenu::SetOnEditRequested(const FOnEditRequested& InOnEditRequested)
-{
-	OnEditRequested = InOnEditRequested;
-}
-
 void FAssetContextMenu::SetOnAssetViewRefreshRequested(const FOnAssetViewRefreshRequested& InOnAssetViewRefreshRequested)
 {
 	OnAssetViewRefreshRequested = InOnAssetViewRefreshRequested;
@@ -511,25 +520,23 @@ bool FAssetContextMenu::AddCommonMenuOptions(UToolMenu* Menu)
 {
 	UContentBrowserDataMenuContext_FileMenu* Context = Menu->FindContext<UContentBrowserDataMenuContext_FileMenu>();
 	const bool bCanBeModified = !Context || Context->bCanBeModified;
+	const bool bCanBeViewed = !Context || Context->bCanView;
 
 	{
 		FToolMenuSection& Section = Menu->AddSection("CommonAssetActions", LOCTEXT("CommonAssetActionsMenuHeading", "Common"));
 
-		
-		if (bCanBeModified)
+		if (bCanBeModified || bCanBeViewed)
 		{
-			// Edit
-			Section.AddMenuEntry(
-				"EditAsset",
-				LOCTEXT("EditAsset", "Edit..."),
-				LOCTEXT("EditAssetTooltip", "Opens the selected item(s) for edit."),
-				FSlateIcon(FAppStyle::GetAppStyleSetName(), "Icons.Edit"),
-				FUIAction(
-					FExecuteAction::CreateSP(this, &FAssetContextMenu::ExecuteEditItems),
-					FCanExecuteAction::CreateSP(this, &FAssetContextMenu::CanExecuteEditItems)
-				)
+			// Open/Edit Asset
+			Section.AddMenuEntry(FContentBrowserCommands::Get().OpenAssetsOrFolders,
+				GetEditAssetEditorLabel(bCanBeModified, bCanBeViewed),
+				GetEditAssetEditorTooltip(bCanBeModified, bCanBeViewed),
+				GetEditAssetEditorIcon(bCanBeModified, bCanBeViewed)
 			);
+		}
 
+		if (bCanBeModified)
+		{
 			// Rename
 			Section.AddMenuEntry(FGenericCommands::Get().Rename,
 				LOCTEXT("Rename", "Rename"),
@@ -698,6 +705,39 @@ bool FAssetContextMenu::AddReferenceMenuOptions(UToolMenu* Menu)
 	return true;
 }
 
+FText FAssetContextMenu::GetEditAssetEditorLabel(bool bInCanEdit, bool bInCanView) const
+{
+	static const FText EditLabel = LOCTEXT("EditAsset", "Edit...");
+	static const FText OpenLabel = LOCTEXT("OpenReadOnlyAsset", "Open as Read-Only...");
+
+	if (bInCanEdit || bInCanView)
+	{
+		return bInCanEdit ? EditLabel : OpenLabel;
+	}
+	return FText::GetEmpty();
+}
+
+FText FAssetContextMenu::GetEditAssetEditorTooltip(bool bInCanEdit, bool bInCanView) const
+{
+	static const FText EditTooltip = LOCTEXT("EditAssetTooltip", "Opens the selected item(s) for edit");
+	static const FText OpenTooltip = LOCTEXT("OpenAssetTooltip", "Opens the selected item(s)");
+
+	if (bInCanEdit || bInCanView)
+	{
+		return bInCanEdit ? EditTooltip : OpenTooltip;
+	}
+	return FText::GetEmpty();
+}
+
+FSlateIcon FAssetContextMenu::GetEditAssetEditorIcon(bool bInCanEdit, bool bInCanView) const
+{
+	if (bInCanEdit || bInCanView)
+	{
+		return bInCanEdit ? FSlateIcon(FAppStyle::GetAppStyleSetName(), TEXT("Icons.Edit")) : FSlateIcon(FAppStyle::GetAppStyleSetName(), TEXT("ContentBrowser.AssetActions.OpenReadOnly"));
+	}
+	return FSlateIcon();
+}
+
 FText FAssetContextMenu::GetCopyTooltip(EAssetViewCopyType InCopyType) const
 {
 	static const FString CopyExportPathBaseTooltip = TEXT("Copies reference paths for the selected asset(s) to the clipboard:");
@@ -998,24 +1038,6 @@ void FAssetContextMenu::ExecuteFindInExplorer()
 	ContentBrowserUtils::ExploreFolders(SelectedFiles, AssetView.Pin().ToSharedRef());
 }
 
-bool FAssetContextMenu::CanExecuteEditItems() const
-{
-	bool bCanEdit = false;
-	for (const FContentBrowserItem& SelectedItem : SelectedFiles)
-	{
-		bCanEdit |= SelectedItem.CanEdit();
-	}
-	return bCanEdit;
-}
-
-void FAssetContextMenu::ExecuteEditItems()
-{
-	if (SelectedFiles.Num() > 0)
-	{
-		OnEditRequested.ExecuteIfBound(SelectedFiles);
-	}
-}
-
 void FAssetContextMenu::ExecuteSaveAsset()
 {
 	const EContentBrowserItemSaveFlags SaveFlags = EContentBrowserItemSaveFlags::None;
diff --git a/Engine/Source/Editor/ContentBrowser/Private/AssetContextMenu.h b/Engine/Source/Editor/ContentBrowser/Private/AssetContextMenu.h
--- a/Engine/Source/Editor/ContentBrowser/Private/AssetContextMenu.h
+++ b/Engine/Source/Editor/ContentBrowser/Private/AssetContextMenu.h
@@ -53,10 +53,6 @@ public:
 	DECLARE_DELEGATE_OneParam(FOnDuplicateRequested, TArrayView<const FContentBrowserItem> /*OriginalItems*/);
 	void SetOnDuplicateRequested(const FOnDuplicateRequested& InOnDuplicateRequested);
 
-	/** Delegate for when the context menu requests to edit an item */
-	DECLARE_DELEGATE_OneParam(FOnEditRequested, TArrayView<const FContentBrowserItem> /*Items*/);
-	void SetOnEditRequested(const FOnEditRequested& InOnEditRequested);
-
 	/** Delegate for when the context menu requests an asset view refresh */
 	using FOnAssetViewRefreshRequested = UContentBrowserDataMenuContext_FileMenu::FOnRefreshView;
 	void SetOnAssetViewRefreshRequested(const FOnAssetViewRefreshRequested& InOnAssetViewRefreshRequested);
@@ -104,6 +100,15 @@ private:
 	/** Adds asset reference menu options to a menu builder. Returns true if any options were added. */
 	bool AddReferenceMenuOptions(UToolMenu* Menu);
 
+	/** Get the correct Label for the OpenAssetEditor command */
+	FText GetEditAssetEditorLabel(bool bInCanEdit, bool bInCanView) const;
+
+	/** Get the correct Tooltip for the OpenAssetEditor command */
+	FText GetEditAssetEditorTooltip(bool bInCanEdit, bool bInCanView) const;
+
+	/** Get the correct Icon for the OpenAssetEditor command */
+	FSlateIcon GetEditAssetEditorIcon(bool bInCanEdit, bool bInCanView) const;
+	
 	/** Return the tooltip based on the copy type */
 	FText GetCopyTooltip(EAssetViewCopyType InCopyType) const;
 
@@ -121,12 +126,6 @@ private:
 	/** Handler for when find in explorer is selected */
 	void ExecuteFindInExplorer();
 
-	/** Handler to check to see if an edit command is allowed */
-	bool CanExecuteEditItems() const;
-
-	/** Handler for when "Edit" is selected */
-	void ExecuteEditItems();
-
 	/** Handler for confirmation of folder deletion */
 	FReply ExecuteDeleteFolderConfirmed();
 
@@ -187,7 +186,6 @@ private:
 	FOnShowInPathsViewRequested OnShowInPathsViewRequested;
 	FOnRenameRequested OnRenameRequested;
 	FOnDuplicateRequested OnDuplicateRequested;
-	FOnEditRequested OnEditRequested;
 	FOnAssetViewRefreshRequested OnAssetViewRefreshRequested;
 
 	/** Cached CanExecute vars */
diff --git a/Engine/Source/Editor/ContentBrowser/Private/ContentBrowserCommands.cpp b/Engine/Source/Editor/ContentBrowser/Private/ContentBrowserCommands.cpp
--- a/Engine/Source/Editor/ContentBrowser/Private/ContentBrowserCommands.cpp
+++ b/Engine/Source/Editor/ContentBrowser/Private/ContentBrowserCommands.cpp
@@ -14,7 +14,7 @@
 void FContentBrowserCommands::RegisterCommands()
 {
 	// ContentBrowser commands
-	UI_COMMAND(OpenAssetsOrFolders, "Open Assets or Folders", "Opens the selected assets or folders, depending on the selection", EUserInterfaceActionType::Button, FInputChord(EKeys::Enter));
+	UI_COMMAND(OpenAssetsOrFolders, "Open Assets or Folders", "Opens the selected assets or folders, depending on the selection", EUserInterfaceActionType::Button, FInputChord(EKeys::Enter), FInputChord(EModifierKey::Control, EKeys::E));
 	UI_COMMAND(PreviewAssets, "Preview Assets", "Loads the selected assets and previews them if possible", EUserInterfaceActionType::Button, FInputChord(EKeys::SpaceBar));
 	UI_COMMAND(CreateNewFolder, "Create New Folder", "Creates new folder in selected path", EUserInterfaceActionType::Button, FInputChord(EModifierKey::Control | EModifierKey::Shift, EKeys::N));
 	UI_COMMAND(SaveSelectedAsset, "Save Selected Item", "Save the selected item", EUserInterfaceActionType::Button, FInputChord(EModifierKey::Control, EKeys::S));
diff --git a/Engine/Source/Editor/ContentBrowser/Private/SContentBrowser.cpp b/Engine/Source/Editor/ContentBrowser/Private/SContentBrowser.cpp
--- a/Engine/Source/Editor/ContentBrowser/Private/SContentBrowser.cpp
+++ b/Engine/Source/Editor/ContentBrowser/Private/SContentBrowser.cpp
@@ -971,7 +971,6 @@ void SContentBrowser::Construct( const FArguments& InArgs, const FName& InInstan
 	AssetContextMenu->SetOnShowInPathsViewRequested( FAssetContextMenu::FOnShowInPathsViewRequested::CreateSP(this, &SContentBrowser::OnShowInPathsViewRequested) );
 	AssetContextMenu->SetOnRenameRequested( FAssetContextMenu::FOnRenameRequested::CreateSP(this, &SContentBrowser::OnRenameRequested) );
 	AssetContextMenu->SetOnDuplicateRequested( FAssetContextMenu::FOnDuplicateRequested::CreateSP(this, &SContentBrowser::OnDuplicateRequested) );
-	AssetContextMenu->SetOnEditRequested( FAssetContextMenu::FOnEditRequested::CreateSP(this, &SContentBrowser::OnEditRequested) );
 	AssetContextMenu->SetOnAssetViewRefreshRequested( FAssetContextMenu::FOnAssetViewRefreshRequested::CreateSP( this, &SContentBrowser::OnAssetViewRefreshRequested) );
 	FavoritePathViewPtr->SetTreeTitle(LOCTEXT("Favorites", "Favorites"));
 	if( Config != nullptr && Config->SelectedCollectionName.Name != NAME_None )
@@ -4306,11 +4305,6 @@ void SContentBrowser::OnDuplicateRequested(TArrayView<const FContentBrowserItem>
 	}
 }
 
-void SContentBrowser::OnEditRequested(TArrayView<const FContentBrowserItem> Items)
-{
-	OnItemsActivated(Items, EAssetTypeActivationMethod::Opened);
-}
-
 void SContentBrowser::OnAssetViewRefreshRequested()
 {
 	AssetViewPtr->RequestSlowFullListRefresh();
diff --git a/Engine/Source/Editor/ContentBrowser/Private/SContentBrowser.h b/Engine/Source/Editor/ContentBrowser/Private/SContentBrowser.h
--- a/Engine/Source/Editor/ContentBrowser/Private/SContentBrowser.h
+++ b/Engine/Source/Editor/ContentBrowser/Private/SContentBrowser.h
@@ -508,9 +508,6 @@ private:
 	/** Handler for when the asset context menu requests to duplicate an item */
 	void OnDuplicateRequested(TArrayView<const FContentBrowserItem> OriginalItems);
 
-	/** Handler for when the asset context menu requests to edit an item */
-	void OnEditRequested(TArrayView<const FContentBrowserItem> Items);
-
 	/** Handler for when the asset context menu requests to refresh the asset view */
 	void OnAssetViewRefreshRequested();
 
diff --git a/Engine/Source/Editor/ContentBrowser/Public/ContentBrowserMenuContexts.h b/Engine/Source/Editor/ContentBrowser/Public/ContentBrowserMenuContexts.h
--- a/Engine/Source/Editor/ContentBrowser/Public/ContentBrowserMenuContexts.h
+++ b/Engine/Source/Editor/ContentBrowser/Public/ContentBrowserMenuContexts.h
@@ -63,6 +63,9 @@ public:
 	UPROPERTY()
 	bool bCanBeModified;
 
+	UPROPERTY()
+	bool bCanView;
+
 	UPROPERTY()
 	bool bHasCookedPackages;
 
diff --git a/Engine/Source/Editor/ContentBrowserData/Public/ContentBrowserDataMenuContexts.h b/Engine/Source/Editor/ContentBrowserData/Public/ContentBrowserDataMenuContexts.h
--- a/Engine/Source/Editor/ContentBrowserData/Public/ContentBrowserDataMenuContexts.h
+++ b/Engine/Source/Editor/ContentBrowserData/Public/ContentBrowserDataMenuContexts.h
@@ -77,6 +77,9 @@ public:
 	UPROPERTY(BlueprintReadOnly, Category = ContentBrowser)
 	bool bCanBeModified = true;
 
+	UPROPERTY(BlueprintReadOnly, Category = ContentBrowser)
+	bool bCanView = true;
+
 	UPROPERTY(BlueprintReadOnly, Category = ContentBrowser)
 	bool bHasCookedPackages = false;
 
diff --git a/Engine/Source/Editor/EditorStyle/Private/StarshipStyle.cpp b/Engine/Source/Editor/EditorStyle/Private/StarshipStyle.cpp
--- a/Engine/Source/Editor/EditorStyle/Private/StarshipStyle.cpp
+++ b/Engine/Source/Editor/EditorStyle/Private/StarshipStyle.cpp
@@ -6347,6 +6347,7 @@ void FStarshipEditorStyle::FStyle::SetupContentBrowserStyle()
 		// Asset Context Menu
 		Set( "ContentBrowser.AssetActions", new CORE_IMAGE_BRUSH( "Icons/icon_tab_Tools_16x", Icon16x16 ) );
 		Set( "ContentBrowser.AssetActions.Edit", new IMAGE_BRUSH( "Icons/Edit/icon_Edit_16x", Icon16x16 ) );
+		Set( "ContentBrowser.AssetActions.OpenReadOnly", new IMAGE_BRUSH_SVG( "Starship/ContentBrowser/OpenReadOnly_16", Icon16x16 ) );
 		Set( "ContentBrowser.AssetActions.Delete", new IMAGE_BRUSH( "Icons/icon_delete_16px", Icon16x16, FLinearColor( 0.4f, 0.5f, 0.7f, 1.0f ) ) );
 		//Set( "ContentBrowser.AssetActions.Delete", new IMAGE_BRUSH( "Icons/Edit/icon_Edit_Delete_16x", Icon16x16) );
 		Set( "ContentBrowser.AssetActions.Rename", new IMAGE_BRUSH( "Icons/Icon_Asset_Rename_16x", Icon16x16) );
