[ContentBrowser] Fixed search bar sizing

diff --git a/Engine/Source/Editor/ContentBrowser/Private/SContentBrowser.cpp b/Engine/Source/Editor/ContentBrowser/Private/SContentBrowser.cpp
--- a/Engine/Source/Editor/ContentBrowser/Private/SContentBrowser.cpp
+++ b/Engine/Source/Editor/ContentBrowser/Private/SContentBrowser.cpp
@@ -1247,7 +1247,11 @@ void SContentBrowser::OnFilterBarLayoutChanging(EFilterBarLayout NewLayout)
 			]
 
 			+ SHorizontalBox::Slot()
+#if UE_CONTENTBROWSER_NEW_STYLE
 			.FillWidth(1.0)
+#else
+			.AutoWidth()
+#endif
 			.VAlign(VAlign_Top)
 			.Padding(0, 4, 0, 0)
 			[
@@ -1355,7 +1359,7 @@ void SContentBrowser::OnFilterBarLayoutChanging(EFilterBarLayout NewLayout)
 					]
 
 					+ SHorizontalBox::Slot()
-					.AutoWidth()
+					.FillWidth(1.0f)
 					.MaxWidth(SearchBoxMaxWidth)
 					[
 						SNew(SBox)
