[ContentBrowser] Thumbnail Styling

diff --git a/Engine/Source/Editor/ContentBrowser/Private/AssetViewWidgets.cpp b/Engine/Source/Editor/ContentBrowser/Private/AssetViewWidgets.cpp
--- a/Engine/Source/Editor/ContentBrowser/Private/AssetViewWidgets.cpp
+++ b/Engine/Source/Editor/ContentBrowser/Private/AssetViewWidgets.cpp
@@ -589,7 +589,7 @@ void SAssetViewItem::HandleSourceControlStateChanged()
 				if (SCCStateWidget.IsValid())
 				{
 					FSlateIcon SCCIcon = SourceControlState->GetIcon();
-					bHasCCStateBrush = SCCIcon.GetIcon() != nullptr;
+					bHasCCStateBrush = SCCIcon.GetIcon() != FStyleDefaults::GetNoBrush();
 					SCCStateWidget->SetFromSlateIcon(SCCIcon);
 				}
 			}
@@ -1446,7 +1446,11 @@ void SAssetViewItem::CacheDisplayTags()
 
 const FSlateBrush* SAssetViewItem::GetBorderImage() const
 {
+#if UE_CONTENTBROWSER_NEW_STYLE
+	return bDraggedOver ? FAppStyle::GetBrush("ContentBrowser.AssetTileItem.AssetBorder") : FAppStyle::GetBrush("NoBorder");
+#else
 	return bDraggedOver ? FAppStyle::GetBrush("Menu.Background") : FAppStyle::GetBrush("NoBorder");
+#endif
 }
 
 bool SAssetViewItem::IsFolder() const
@@ -1828,6 +1832,10 @@ void SAssetTileItem::Construct( const FArguments& InArgs )
 		ThumbnailConfig.HintColorAndOpacity = InArgs._ThumbnailHintColorAndOpacity;
 		ThumbnailConfig.Padding= FMargin(2.0f);
 		ThumbnailConfig.GenericThumbnailSize = MakeAttributeSP(this, &SAssetTileItem::GetGenericThumbnailSize);
+#if UE_CONTENTBROWSER_NEW_STYLE
+		ThumbnailConfig.bShowAssetChip = true;
+		ThumbnailConfig.AssetChipBorderImageOverride = TAttribute<const FSlateBrush*>::CreateSP(this, &SAssetTileItem::GetAssetAreaOverlayBackgroundImage);
+#endif
 
 		if(!AssetItem->GetItem().IsSupported())
 		{
@@ -1888,19 +1896,22 @@ void SAssetTileItem::Construct( const FArguments& InArgs )
 			.ColorAndOpacity(this, &SAssetTileItem::GetNameAreaTextColor);
 	}
 
-	constexpr float AssetViewWidgetsBorderPadding = 4.f;
-	constexpr float AssetViewWidgetsShadowPadding = 5.f;
-	constexpr float AssetViewWidgetsRenameWidgetPadding = 2.f;
-	constexpr float AssetViewWidgetsClassTextPadding = 2.f;
-	constexpr float AssetViewWidgetSourceControlSize = 16.f;
-
 #if UE_CONTENTBROWSER_NEW_STYLE
+
+	constexpr float BorderPadding = 1.f;
+	constexpr float ShadowLeftTopPadding = 3.f;
+	constexpr float ShadowRightBotPadding = 4.f;
+	constexpr float ThumbnailBorderPadding = 0.f;
+	constexpr float NameAreaBoxLeftRightBotPadding = 4.f;
+	constexpr float NameAreaBoxTopPadding = 6.f;
+	constexpr float ClassNameMaxHeight = 14.f;
+
 	ChildSlot
-	.Padding(FMargin(0.0f, 0.0f, AssetViewWidgetsBorderPadding, AssetViewWidgetsBorderPadding))
-	[				
+	.Padding(FMargin(BorderPadding, BorderPadding, 0.f, 0.f))
+	[
 		// Drop shadow border
 		SNew(SBorder)
-		.Padding(FMargin(0.0f, 0.0f, AssetViewWidgetsShadowPadding, AssetViewWidgetsShadowPadding))
+		.Padding(FMargin(ShadowLeftTopPadding, ShadowLeftTopPadding, ShadowRightBotPadding, ShadowRightBotPadding))
 		.BorderImage(IsFolder() ? TAttribute<const FSlateBrush*>(this, &SAssetTileItem::GetFolderBackgroundShadowImage) : FAppStyle::Get().GetBrush(ItemShadowBorderName))
 		[
 			SNew(SOverlay)
@@ -1908,8 +1919,8 @@ void SAssetTileItem::Construct( const FArguments& InArgs )
 			+SOverlay::Slot()
 			[
 				SNew(SBorder)
-				.Padding(0)
-				.BorderImage(IsFolder() ? TAttribute<const FSlateBrush*>(this, &SAssetTileItem::GetFolderBackgroundImage) : UE::ContentBrowser::Private::FContentBrowserStyle::Get().Get().GetBrush("ContentBrowser.AssetTileItem.ThumbnailAreaBackground"))
+				.Padding(ThumbnailBorderPadding)
+				.BorderImage(IsFolder() ? TAttribute<const FSlateBrush*>(this, &SAssetTileItem::GetFolderBackgroundImage) : TAttribute<const FSlateBrush*>(this, &SAssetTileItem::GetNameAreaBackgroundImage))
 				[
 					SNew(SVerticalBox)
 					// Thumbnail
@@ -1918,7 +1929,6 @@ void SAssetTileItem::Construct( const FArguments& InArgs )
 					[
 						// The remainder of the space is reserved for the name.
 						SNew(SBox)
-						.Padding(0)
 						.WidthOverride(this, &SAssetTileItem::GetThumbnailBoxSize)
 						.HeightOverride(this, &SAssetTileItem::GetThumbnailBoxSize)
 						[
@@ -1927,68 +1937,53 @@ void SAssetTileItem::Construct( const FArguments& InArgs )
 					]
 
 					+SVerticalBox::Slot()
+					.FillHeight(1.f)
 					[
-						SNew(SBorder)
+						SNew(SBox)
+						.Visibility(this, &SAssetTileItem::GetNameAreaVisibility)
+						.Padding(FMargin(NameAreaBoxLeftRightBotPadding, NameAreaBoxTopPadding, NameAreaBoxLeftRightBotPadding, NameAreaBoxLeftRightBotPadding))
 						.Visibility(this, &SAssetTileItem::GetNameAreaVisibility)
-						.Padding(FMargin(2.0f, 3.0f))
-						.BorderImage(this, &SAssetTileItem::GetNameAreaBackgroundImage)
 						[
 							SNew(SVerticalBox)
 							+ SVerticalBox::Slot()
-							.Padding(AssetViewWidgetsRenameWidgetPadding,AssetViewWidgetsRenameWidgetPadding,0.0f,0.0f)
-							.VAlign(VAlign_Top)
-							.HAlign(IsFolder() ? HAlign_Center : HAlign_Left)
+							.HAlign(IsFolder() ? HAlign_Center : HAlign_Fill)
 							[
 								SNew(SBox)
-								.MaxDesiredHeight(this, &SAssetTileItem::GetNameAreaMaxDesiredHeight)
+								.VAlign(VAlign_Top)
+								.HeightOverride(this, &SAssetTileItem::GetNameAreaMaxDesiredHeight)
 								[
 									InlineRenameWidget.ToSharedRef()
 								]
 							]
 							+ SVerticalBox::Slot()
-							.VAlign(VAlign_Bottom)
 							.AutoHeight()
-							.Padding(AssetViewWidgetsClassTextPadding,0.0f, 0.0f, AssetViewWidgetsClassTextPadding)
 							[
-								SNew(SHorizontalBox)
-								+SHorizontalBox::Slot()
+								SNew(SBox)
+								.HeightOverride(ClassNameMaxHeight)
+								.VAlign(VAlign_Bottom)
 								[
 									SAssignNew(ClassTextWidget, STextBlock)
+									.Font(FAppStyle::GetFontStyle("ContentBrowser.AssetTileViewClassNameFont"))
 									.Visibility(this, &SAssetTileItem::GetAssetClassLabelVisibility)
 									.TextStyle(UE::ContentBrowser::Private::FContentBrowserStyle::Get().Get(), "ContentBrowser.ClassFont")
 									.OverflowPolicy(ETextOverflowPolicy::Ellipsis)
 									.Text(this, &SAssetTileItem::GetAssetClassText)
 									.ColorAndOpacity(this, &SAssetTileItem::GetAssetClassLabelTextColor)
 								]
-
-								+SHorizontalBox::Slot()
-								.Padding(FMargin(3.0f, 0.0f, 0.0f, 1.0f))
-								.AutoWidth()
-								.HAlign(HAlign_Right)
-								[
-									SNew(SBox)
-									.WidthOverride(AssetViewWidgetSourceControlSize)
-									.HeightOverride(AssetViewWidgetSourceControlSize)
-									.Visibility(this, &SAssetTileItem::GetSCCIconVisibility)
-									[
-										GenerateSourceControlIconWidget()
-									]
-								]
 							]
-							
 						]
 					]
 				]
 			]
-			+SOverlay::Slot()
-			[
-				SNew(SImage)
-				.Image(this, &SAssetViewItem::GetBorderImage)
-				.Visibility(EVisibility::HitTestInvisible)	
-			]
 		]
 	];
 #else
+	constexpr float AssetViewWidgetsBorderPadding = 4.f;
+	constexpr float AssetViewWidgetsShadowPadding = 5.f;
+	constexpr float AssetViewWidgetsRenameWidgetPadding = 2.f;
+	constexpr float AssetViewWidgetsClassTextPadding = 2.f;
+	constexpr float AssetViewWidgetSourceControlSize = 16.f;
+
 	ChildSlot
 	.Padding(FMargin(0.0f, 0.0f, AssetViewWidgetsBorderPadding, AssetViewWidgetsBorderPadding))
 	[				
@@ -2195,11 +2190,19 @@ FSlateColor SAssetTileItem::GetAssetClassLabelTextColor() const
 		return FStyleColors::White;
 	}
 
+#if UE_CONTENTBROWSER_NEW_STYLE
+	return FStyleColors::Hover2;
+#else
 	return FSlateColor::UseSubduedForeground();
+#endif
 }
 
 FSlateFontInfo SAssetTileItem::GetThumbnailFont() const
 {
+#if UE_CONTENTBROWSER_NEW_STYLE
+	const static FLazyName RegularFont("ContentBrowser.AssetTileViewNameFont");
+	return FAppStyle::GetFontStyle(RegularFont);
+#else
 	FOptionalSize ThumbSize = GetThumbnailBoxSize();
 	if ( ThumbSize.IsSet() )
 	{
@@ -2218,6 +2221,7 @@ FSlateFontInfo SAssetTileItem::GetThumbnailFont() const
 
 	const static FName RegularFont("ContentBrowser.AssetTileViewNameFont");
 	return FAppStyle::GetFontStyle(RegularFont);
+#endif
 }
 
 const FSlateBrush* SAssetTileItem::GetFolderBackgroundImage() const
@@ -2259,33 +2263,96 @@ const FSlateBrush* SAssetTileItem::GetFolderBackgroundShadowImage() const
 }
 
 const FSlateBrush* SAssetTileItem::GetNameAreaBackgroundImage() const
+{
+	if (CurrentThumbnailSize.Get() == EThumbnailSize::Tiny)
+	{
+		return FStyleDefaults::GetNoBrush();
+	}
+
+#if UE_CONTENTBROWSER_NEW_STYLE
+	static const FLazyName SelectedHover("ContentBrowser.AssetTileItem.AssetContentSelectedHoverBackground");
+	static const FLazyName Selected("ContentBrowser.AssetTileItem.AssetContentSelectedBackground");
+	static const FLazyName Hovered("ContentBrowser.AssetTileItem.AssetContentHoverBackground");
+	static const FLazyName Normal("ContentBrowser.AssetTileItem.AssetContent");
+#else
+	static const FLazyName SelectedHover("ContentBrowser.AssetTileItem.NameAreaSelectedHoverBackground");
+	static const FLazyName Selected("ContentBrowser.AssetTileItem.NameAreaSelectedBackground");
+	static const FLazyName Hovered("ContentBrowser.AssetTileItem.NameAreaHoverBackground");
+	static const FLazyName Normal("ContentBrowser.AssetTileItem.NameAreaBackground");
+#endif
+
+	const bool bIsSelected = IsSelected.IsBound() ? IsSelected.Execute() : false;
+	const bool bIsHoveredOrDraggedOver = IsHovered() || bDraggedOver;
+	if (bIsSelected && bIsHoveredOrDraggedOver)
+	{
+		return FAppStyle::Get().GetBrush(SelectedHover);
+	}
+	else if (bIsSelected)
+	{
+		return FAppStyle::Get().GetBrush(Selected);
+	}
+	else if (bIsHoveredOrDraggedOver && !IsFolder())
+	{
+		return FAppStyle::Get().GetBrush(Hovered);
+	}
+	else if (!IsFolder())
+	{
+		return FAppStyle::Get().GetBrush(Normal);
+	}
+	return FStyleDefaults::GetNoBrush();
+}
+
+const FSlateBrush* SAssetTileItem::GetAssetAreaOverlayBackgroundImage() const
 {
 	const bool bIsSelected = IsSelected.IsBound() ? IsSelected.Execute() : false;
 	const bool bIsHoveredOrDraggedOver = IsHovered() || bDraggedOver;
 	if (bIsSelected && bIsHoveredOrDraggedOver)
 	{
-		static const FName SelectedHover("ContentBrowser.AssetTileItem.NameAreaSelectedHoverBackground");
+		static const FLazyName SelectedHover("ContentBrowser.AssetTileItem.AssetBorderSelectedHoverBackground");
 		return FAppStyle::Get().GetBrush(SelectedHover);
 	}
 	else if (bIsSelected)
 	{
-		static const FName Selected("ContentBrowser.AssetTileItem.NameAreaSelectedBackground");
+		static const FLazyName Selected("ContentBrowser.AssetTileItem.AssetBorderSelectedBackground");
 		return FAppStyle::Get().GetBrush(Selected);
 	}
 	else if (bIsHoveredOrDraggedOver && !IsFolder())
 	{
-		static const FName Hovered("ContentBrowser.AssetTileItem.NameAreaHoverBackground");
+		static const FLazyName Hovered("ContentBrowser.AssetTileItem.AssetBorderHoverBackground");
 		return FAppStyle::Get().GetBrush(Hovered);
 	}
 	else if (!IsFolder())
 	{
-		static const FName Normal("ContentBrowser.AssetTileItem.NameAreaBackground");
+		static const FLazyName Normal("ContentBrowser.AssetTileItem.AssetBorder");
 		return FAppStyle::Get().GetBrush(Normal);
 	}
 
 	return FStyleDefaults::GetNoBrush();
 }
 
+FSlateColor SAssetTileItem::GetChipBackgroundColor() const
+{
+	const bool bIsSelected = IsSelected.IsBound() ? IsSelected.Execute() : false;
+	const bool bIsHoveredOrDraggedOver = IsHovered() || bDraggedOver;
+	if (bIsSelected && bIsHoveredOrDraggedOver)
+	{
+		return FStyleColors::PrimaryHover;
+	}
+	else if (bIsSelected)
+	{
+		return FStyleColors::Primary;
+	}
+	else if (bIsHoveredOrDraggedOver && !IsFolder())
+	{
+		return FStyleColors::Hover;
+	}
+	else if (!IsFolder())
+	{
+		return FStyleColors::Secondary;
+	}
+	return FStyleColors::Error;
+}
+
 FSlateColor SAssetTileItem::GetNameAreaTextColor() const
 {
 	const bool bIsSelected = IsSelected.IsBound() ? IsSelected.Execute() : false;
@@ -2295,12 +2362,21 @@ FSlateColor SAssetTileItem::GetNameAreaTextColor() const
 		return FStyleColors::White;
 	}
 
+#if UE_CONTENTBROWSER_NEW_STYLE
+	return FStyleColors::Foreground;
+#else
 	return FSlateColor::UseForeground();
+#endif
 }
 
 FOptionalSize SAssetTileItem::GetNameAreaMaxDesiredHeight() const
 {
+#if UE_CONTENTBROWSER_NEW_STYLE
+	constexpr int32 MaxHeightNameArea = 42;
+	return MaxHeightNameArea;
+#else
 	return AssetNameHeights[(int32)CurrentThumbnailSize.Get()];
+#endif
 }
 
 int32 SAssetTileItem::GetGenericThumbnailSize() const
@@ -2311,7 +2387,11 @@ int32 SAssetTileItem::GetGenericThumbnailSize() const
 EVisibility SAssetTileItem::GetSCCIconVisibility() const
 {
 	// Hide the scc state icon when there is no brush or in tiny size since there isn't enough space
+#if UE_CONTENTBROWSER_NEW_STYLE
+	return bHasCCStateBrush && ISourceControlModule::Get().IsEnabled() && ISourceControlModule::Get().GetProvider().IsAvailable() ? EVisibility::Visible : EVisibility::Collapsed;
+#else
 	return bHasCCStateBrush &&  CurrentThumbnailSize.Get() != EThumbnailSize::Tiny && ISourceControlModule::Get().IsEnabled() && ISourceControlModule::Get().GetProvider().IsAvailable() ? EVisibility::Visible : EVisibility::Collapsed;
+#endif
 }
 
 EVisibility SAssetTileItem::GetNameAreaVisibility() const
diff --git a/Engine/Source/Editor/ContentBrowser/Private/SAssetView.cpp b/Engine/Source/Editor/ContentBrowser/Private/SAssetView.cpp
--- a/Engine/Source/Editor/ContentBrowser/Private/SAssetView.cpp
+++ b/Engine/Source/Editor/ContentBrowser/Private/SAssetView.cpp
@@ -2439,7 +2439,7 @@ FReply SAssetView::OnMouseWheel( const FGeometry& MyGeometry, const FPointerEven
 	{
 #if UE_CONTENTBROWSER_NEW_STYLE
 		const float NewDelta = MouseEvent.GetWheelDelta() * 0.4f;
-		if (ZoomScale == 1.f && NewDelta > 0 || ZoomScale == 0.f && NewDelta < 0)
+		if ((ZoomScale == 1.f && NewDelta > 0) || (ZoomScale == 0.f && NewDelta < 0))
 		{
 			int32 Step = (int32)FMath::Sign(NewDelta);
 			EThumbnailSize OldSize = ThumbnailSize;
@@ -5619,7 +5619,7 @@ float SAssetView::GetListViewItemHeight() const
 float SAssetView::GetTileViewItemHeight() const
 {
 #if UE_CONTENTBROWSER_NEW_STYLE
-	return GetTileViewItemBaseWidth() + GetTileViewTypeNameHeight() + TileViewWidthHeightPadding;
+	return GetTileViewItemBaseWidth() + GetTileViewTypeNameHeight() + TileViewHeightPadding;
 #else
 	return (((float)TileViewNameHeight + GetTileViewTypeNameHeight()) * FMath::Lerp(MinThumbnailScale, MaxThumbnailScale, GetThumbnailScale())) + GetTileViewItemBaseHeight() * FillScale + GetSourceControlIconHeight();
 #endif
@@ -5637,7 +5637,7 @@ float SAssetView::GetTileViewItemBaseHeight() const
 float SAssetView::GetTileViewItemWidth() const
 {
 #if UE_CONTENTBROWSER_NEW_STYLE
-	return GetTileViewItemBaseWidth() + TileViewWidthHeightPadding;
+	return GetTileViewItemBaseWidth() + TileViewWidthPadding;
 #else
 	return GetTileViewItemBaseWidth() * FillScale;
 #endif
diff --git a/Engine/Source/Editor/ContentBrowser/Public/AssetViewWidgets.h b/Engine/Source/Editor/ContentBrowser/Public/AssetViewWidgets.h
--- a/Engine/Source/Editor/ContentBrowser/Public/AssetViewWidgets.h
+++ b/Engine/Source/Editor/ContentBrowser/Public/AssetViewWidgets.h
@@ -728,6 +728,8 @@ protected:
 	const FSlateBrush* GetFolderBackgroundShadowImage() const;
 
 	const FSlateBrush* GetNameAreaBackgroundImage() const;
+	const FSlateBrush* GetAssetAreaOverlayBackgroundImage() const;
+	FSlateColor GetChipBackgroundColor() const;
 	FSlateColor GetNameAreaTextColor() const;
 
 	FOptionalSize GetNameAreaMaxDesiredHeight() const;
diff --git a/Engine/Source/Editor/ContentBrowser/Public/SAssetView.h b/Engine/Source/Editor/ContentBrowser/Public/SAssetView.h
--- a/Engine/Source/Editor/ContentBrowser/Public/SAssetView.h
+++ b/Engine/Source/Editor/ContentBrowser/Public/SAssetView.h
@@ -981,8 +981,11 @@ private:
 	 */
 	float ZoomScale = 0.0f;
 
-	/** Vertical and Horizontal padding for the TileViewItem */
-	const int32 TileViewWidthHeightPadding = 9;
+	/** Vertical padding for the TileViewItem */
+	static constexpr int32 TileViewHeightPadding = 9;
+
+	/** Horizontal padding for the TileViewItem */
+	static constexpr int32 TileViewWidthPadding = 8;
 
 	TAttribute<bool> bShowRedirectors;
 	bool bLastShowRedirectors;
diff --git a/Engine/Source/Editor/EditorStyle/Private/StarshipStyle.cpp b/Engine/Source/Editor/EditorStyle/Private/StarshipStyle.cpp
--- a/Engine/Source/Editor/EditorStyle/Private/StarshipStyle.cpp
+++ b/Engine/Source/Editor/EditorStyle/Private/StarshipStyle.cpp
@@ -6237,7 +6237,10 @@ void FStarshipEditorStyle::FStyle::SetupContentBrowserStyle()
 		);
 
 		// Tile view
+		Set( "ContentBrowser.AssetTileViewWhiteChip", new IMAGE_BRUSH_SVG("Starship/ContentBrowser/CornerChipInner_18", FVector2D(18, 18)));
+		Set( "ContentBrowser.AssetTileViewWhiteChipBorder", new IMAGE_BRUSH_SVG("Starship/ContentBrowser/CornerChipOuter_18", FVector2D(18, 18)));
 		Set( "ContentBrowser.AssetTileViewNameFont", DEFAULT_FONT("Regular", 9));
+		Set( "ContentBrowser.AssetTileViewClassNameFont", DEFAULT_FONT("Regular", 7));
 		Set( "ContentBrowser.AssetTileViewNameFontSmall", DEFAULT_FONT( "VeryLight", 8 ) );
 		Set( "ContentBrowser.AssetTileViewNameFontVerySmall", DEFAULT_FONT( "VeryLight", 7 ) );
 		Set( "ContentBrowser.AssetTileViewNameFontDirty", FStyleFonts::Get().SmallBold);
@@ -6416,7 +6419,23 @@ void FStarshipEditorStyle::FStyle::SetupContentBrowserStyle()
 		Set("ContentBrowser.AssetTileItem.NameAreaSelectedBackground", new FSlateRoundedBoxBrush(FStyleColors::Primary, FVector4(0.0f, 0.0f, 4.0f, 4.0f)));
 		Set("ContentBrowser.AssetTileItem.NameAreaSelectedHoverBackground", new FSlateRoundedBoxBrush(FStyleColors::PrimaryHover, FVector4(0.0f, 0.0f, 4.0f, 4.0f)));
 
+		// Asset Thumbnail Border
+		Set("ContentBrowser.AssetTileItem.AssetBorder", new FSlateRoundedBoxBrush(FStyleColors::Transparent, FVector4(4.0f, 2.0f, 4.0f, 4.0f), FStyleColors::Secondary, 2.f));
+		Set("ContentBrowser.AssetTileItem.AssetBorderHoverBackground", new FSlateRoundedBoxBrush(FStyleColors::Transparent, FVector4(4.0f, 2.0f, 4.0f, 4.0f), FStyleColors::Hover, 2.f));
+		Set("ContentBrowser.AssetTileItem.AssetBorderSelectedBackground", new FSlateRoundedBoxBrush(FStyleColors::Transparent, FVector4(4.0f, 2.0f, 4.0f, 4.0f), FStyleColors::Primary, 2.f));
+		Set("ContentBrowser.AssetTileItem.AssetBorderSelectedHoverBackground", new FSlateRoundedBoxBrush(FStyleColors::Transparent, FVector4(4.0f, 2.0f, 4.0f, 4.0f), FStyleColors::PrimaryHover, 2.f));
+
+		// Tile Item Border
+		Set("ContentBrowser.AssetTileItem.AssetContent", new FSlateRoundedBoxBrush(FStyleColors::Secondary, FVector4(4.0f, 2.0f, 4.0f, 4.0f)));
+		Set("ContentBrowser.AssetTileItem.AssetContentHoverBackground", new FSlateRoundedBoxBrush(FStyleColors::Hover, FVector4(4.0f, 2.0f, 4.0f, 4.0f)));
+		Set("ContentBrowser.AssetTileItem.AssetContentSelectedBackground", new FSlateRoundedBoxBrush(FStyleColors::Primary, FVector4(4.0f, 2.0f, 4.0f, 4.0f)));
+		Set("ContentBrowser.AssetTileItem.AssetContentSelectedHoverBackground", new FSlateRoundedBoxBrush(FStyleColors::PrimaryHover, FVector4(4.0f, 2.0f, 4.0f, 4.0f)));
+
 		{
+			FLinearColor OverlayColor = FStyleColors::Panel.GetSpecifiedColor();
+            OverlayColor.A = 0.75f;
+            Set("ContentBrowser.AssetTileItem.AssetThumbnailOverlayBorder", new FSlateRoundedBoxBrush(OverlayColor, 4.f));
+
 			FLinearColor TransparentPrimary = FStyleColors::Primary.GetSpecifiedColor();
 			TransparentPrimary.A = 0.0;
 			Set("ContentBrowser.AssetTileItem.SelectedBorder", new FSlateRoundedBoxBrush(TransparentPrimary, 4.0f, FStyleColors::Primary, 1.0f));
diff --git a/Engine/Source/Editor/UnrealEd/Private/AssetThumbnail.cpp b/Engine/Source/Editor/UnrealEd/Private/AssetThumbnail.cpp
--- a/Engine/Source/Editor/UnrealEd/Private/AssetThumbnail.cpp
+++ b/Engine/Source/Editor/UnrealEd/Private/AssetThumbnail.cpp
@@ -40,6 +40,7 @@
 #include "ClassIconFinder.h"
 #include "IVREditorModule.h"
 #include "Framework/Application/SlateApplication.h"
+#include "Widgets/Images/SLayeredImage.h"
 
 namespace AssetThumbnailPool
 {
@@ -79,6 +80,10 @@ public:
 		, _AssetTypeColorOverride()
 		, _Padding(0)
 		, _GenericThumbnailSize(64)
+#if UE_CONTENTBROWSER_NEW_STYLE
+		, _ShowAssetChip(false)
+		, _AssetChipBorderImageOverride()
+#endif
 		, _ColorStripOrientation(EThumbnailColorStripOrientation::HorizontalBottomEdge)
 		{}
 
@@ -97,6 +102,10 @@ public:
 		SLATE_ARGUMENT(TOptional<FLinearColor>, AssetTypeColorOverride)
 		SLATE_ARGUMENT(FMargin, Padding)
 		SLATE_ATTRIBUTE(int32, GenericThumbnailSize)
+#if UE_CONTENTBROWSER_NEW_STYLE
+		SLATE_ARGUMENT(bool, ShowAssetChip)
+		SLATE_ARGUMENT(TAttribute<const FSlateBrush*>, AssetChipBorderImageOverride)
+#endif
 		SLATE_ARGUMENT(EThumbnailColorStripOrientation, ColorStripOrientation)
 
 	SLATE_END_ARGS()
@@ -118,7 +127,10 @@ public:
 		GenericThumbnailBorderPadding = 2.f;
 		GenericThumbnailSize = InArgs._GenericThumbnailSize;
 		ColorStripOrientation = InArgs._ColorStripOrientation;
-
+#if UE_CONTENTBROWSER_NEW_STYLE
+		bShowAssetChip = InArgs._ShowAssetChip;
+		AssetChipBorderImageOverride = InArgs._AssetChipBorderImageOverride;
+#endif
 		AssetThumbnail->OnAssetDataChanged().AddSP(this, &SAssetThumbnail::OnAssetDataChanged);
 
 		const FAssetData& AssetData = AssetThumbnail->GetAssetData();
@@ -262,6 +274,48 @@ public:
 				];
 		}
 
+#if UE_CONTENTBROWSER_NEW_STYLE
+		TSharedRef<SWidget> ContentWidget = OverlayWidget;
+
+		if (bShowAssetChip)
+		{
+			// The asset color chip
+			LayeredImage = SNew(SLayeredImage)
+				.Image(FAppStyle::GetBrush("ContentBrowser.AssetTileViewWhiteChipBorder"))
+				.ColorAndOpacity(FStyleColors::Secondary);
+			LayeredImage->AddLayer(FAppStyle::GetBrush("ContentBrowser.AssetTileViewWhiteChip"), AssetColor);
+
+			OverlayWidget->AddSlot()
+				.HAlign(HAlign_Right)
+				.VAlign(VAlign_Top)
+				[
+					LayeredImage.ToSharedRef()
+				];
+
+			TSharedRef<SImage> OverlayBorder = SNew(SImage)
+				.Image(FAppStyle::GetBrush("ContentBrowser.AssetTileItem.AssetBorder"))
+				.Visibility(EVisibility::HitTestInvisible);
+			
+			OverlayWidget->AddSlot()
+			[
+				OverlayBorder
+			];
+
+			if (AssetChipBorderImageOverride.IsSet())
+			{
+				TSharedRef<SBorder> ThumbnailWidgetBorder = SNew(SBorder)
+					.Padding(0)
+					.BorderImage(FAppStyle::GetBrush("ContentBrowser.AssetTileItem.AssetBorder"))
+					[
+						OverlayWidget
+					];
+
+				OverlayBorder->SetImage(AssetChipBorderImageOverride);
+				ThumbnailWidgetBorder->SetBorderImage(AssetChipBorderImageOverride);
+				ContentWidget = ThumbnailWidgetBorder;
+			}
+		}
+#else
 		// The asset color strip
 		OverlayWidget->AddSlot()
 		.HAlign(ColorStripOrientation == EThumbnailColorStripOrientation::HorizontalBottomEdge ? HAlign_Fill : HAlign_Right)
@@ -272,6 +326,7 @@ public:
 			.BorderBackgroundColor(AssetColor)
 			.Padding(this, &SAssetThumbnail::GetAssetColorStripPadding)
 		];
+#endif
 
 		if( InArgs._AllowAssetSpecificThumbnailOverlay && AssetTypeActions.IsValid() )
 		{
@@ -288,7 +343,11 @@ public:
 
 		ChildSlot
 		[
+#if UE_CONTENTBROWSER_NEW_STYLE
+			ContentWidget
+#else
 			OverlayWidget
+#endif
 		];
 
 		UpdateThumbnailVisibilities();
@@ -432,7 +491,15 @@ private:
 		}
 
 		//AssetBackgroundWidget->SetBorderBackgroundColor(AssetColor.CopyWithNewOpacity(0.3f));
+#if UE_CONTENTBROWSER_NEW_STYLE
+		if (LayeredImage.IsValid())
+		{
+			constexpr int32 ColoredChipLayerIndex = 1;
+			LayeredImage->SetLayerColor(ColoredChipLayerIndex, AssetColor);
+		}
+#else
 		AssetColorStripWidget->SetBorderBackgroundColor(AssetColor);
+#endif
 
 		UpdateThumbnailVisibilities();
 	}
@@ -770,6 +837,7 @@ private:
 	FCurveSequence ViewportFadeAnimation;
 	FCurveHandle ViewportFadeCurve;
 
+	TSharedPtr<SLayeredImage> LayeredImage;
 	FLinearColor AssetColor;
 	TOptional<FLinearColor> AssetTypeColorOverride;
 
@@ -783,6 +851,10 @@ private:
 	TAttribute< FLinearColor > HintColorAndOpacity;
 	TAttribute<int32> GenericThumbnailSize;
 	EThumbnailColorStripOrientation ColorStripOrientation;
+#if UE_CONTENTBROWSER_NEW_STYLE
+	bool bShowAssetChip;
+	TAttribute<const FSlateBrush*> AssetChipBorderImageOverride;
+#endif
 
 	bool bAllowHintText;
 	bool bAllowRealTimeOnHovered;
@@ -921,6 +993,10 @@ TSharedRef<SWidget> FAssetThumbnail::MakeThumbnailWidget( const FAssetThumbnailC
 		.AssetTypeColorOverride(InConfig.AssetTypeColorOverride)
 		.Padding(InConfig.Padding)
 		.GenericThumbnailSize(InConfig.GenericThumbnailSize)
+#if UE_CONTENTBROWSER_NEW_STYLE
+		.ShowAssetChip(InConfig.bShowAssetChip)
+		.AssetChipBorderImageOverride(InConfig.AssetChipBorderImageOverride)
+#endif
 		.ColorStripOrientation(InConfig.ColorStripOrientation);
 }
 
diff --git a/Engine/Source/Editor/UnrealEd/Private/DragAndDrop/AssetDragDropOp.cpp b/Engine/Source/Editor/UnrealEd/Private/DragAndDrop/AssetDragDropOp.cpp
--- a/Engine/Source/Editor/UnrealEd/Private/DragAndDrop/AssetDragDropOp.cpp
+++ b/Engine/Source/Editor/UnrealEd/Private/DragAndDrop/AssetDragDropOp.cpp
@@ -85,7 +85,13 @@ TSharedPtr<SWidget> FAssetDragDropOp::GetDefaultDecorator() const
 	}
 	else if (AssetThumbnail.IsValid())
 	{
+#if UE_CONTENTBROWSER_NEW_STYLE
+		FAssetThumbnailConfig AssetThumbnailConfig;
+		AssetThumbnailConfig.bShowAssetChip = true;
+		ThumbnailWidget = AssetThumbnail->MakeThumbnailWidget(AssetThumbnailConfig);
+#else
 		ThumbnailWidget = AssetThumbnail->MakeThumbnailWidget();
+#endif
 	}
 	else if (HasFolders())
 	{
diff --git a/Engine/Source/Editor/UnrealEd/Public/AssetThumbnail.h b/Engine/Source/Editor/UnrealEd/Public/AssetThumbnail.h
--- a/Engine/Source/Editor/UnrealEd/Public/AssetThumbnail.h
+++ b/Engine/Source/Editor/UnrealEd/Public/AssetThumbnail.h
@@ -64,6 +64,10 @@ struct FAssetThumbnailConfig
 		, HintColorAndOpacity( FLinearColor( 0.0f, 0.0f, 0.0f, 0.0f ) )
 		, AssetTypeColorOverride()
 		, Padding(0)
+#if UE_CONTENTBROWSER_NEW_STYLE
+		, bShowAssetChip(false)
+		, AssetChipBorderImageOverride()
+#endif
 	{
 	}
 
@@ -80,6 +84,10 @@ struct FAssetThumbnailConfig
 	FMargin Padding;
 	TAttribute<int32> GenericThumbnailSize = 64;
 	EThumbnailColorStripOrientation ColorStripOrientation = EThumbnailColorStripOrientation::HorizontalBottomEdge;
+#if UE_CONTENTBROWSER_NEW_STYLE
+	bool bShowAssetChip;
+	TAttribute<const FSlateBrush*> AssetChipBorderImageOverride;
+#endif
 };
 
 
