Content Browser: Allow dragging folders into favorites to add them.

diff --git a/Engine/Source/Editor/ContentBrowser/Private/SContentBrowser.cpp b/Engine/Source/Editor/ContentBrowser/Private/SContentBrowser.cpp
--- a/Engine/Source/Editor/ContentBrowser/Private/SContentBrowser.cpp
+++ b/Engine/Source/Editor/ContentBrowser/Private/SContentBrowser.cpp
@@ -1123,9 +1123,23 @@ bool SContentBrowser::IsCollectionViewDocked() const
 	return GetDefault<UContentBrowserSettings>()->GetDockCollections();
 }
 
+void SContentBrowser::AddFolderFavorite(const TArray<FString>& FolderPaths)
+{
+	for (const FString& FolderPath : FolderPaths)
+	{
+		const FContentBrowserItemPath ItemPath(FolderPath, EContentBrowserPathType::Virtual);
+		if (!ContentBrowserUtils::IsFavoriteFolder(ItemPath))
+		{
+			ContentBrowserUtils::AddFavoriteFolder(ItemPath);
+		}
+	}
+
+	SaveAndShowNewFolderFavorites(FolderPaths);
+}
+
 void SContentBrowser::ToggleFolderFavorite(const TArray<FString>& FolderPaths)
 {
-	bool bAddedFavorite = false;
+	TArray<FString> FolderPathsAdded;
 	for (const FString& FolderPath : FolderPaths)
 	{
 		const FContentBrowserItemPath ItemPath(FolderPath, EContentBrowserPathType::Virtual);
@@ -1136,14 +1150,19 @@ void SContentBrowser::ToggleFolderFavorite(const TArray<FString>& FolderPaths)
 		else
 		{
 			ContentBrowserUtils::AddFavoriteFolder(ItemPath);
-			bAddedFavorite = true;
+			FolderPathsAdded.Add(FolderPath);
 		}
 	}
 
+	SaveAndShowNewFolderFavorites(FolderPathsAdded);
+}
+
+void SContentBrowser::SaveAndShowNewFolderFavorites(const TArray<FString>& FolderPaths)
+{
 	FavoritePathViewPtr->SaveSettings(GEditorPerProjectIni, SettingsIniSection, InstanceName.ToString() + TEXT(".Favorites"));
-	
 	FavoritePathViewPtr->Populate();
-	if (bAddedFavorite)
+
+	if (!FolderPaths.IsEmpty())
 	{	
 		FavoritePathViewPtr->SetSelectedPaths(FolderPaths);
 		if (GetFavoriteFolderVisibility() == EVisibility::Collapsed)
@@ -1453,7 +1472,7 @@ void SContentBrowser::SetFavoritesExpanded(bool bExpanded)
 TSharedRef<SWidget> SContentBrowser::CreateFavoritesView(const FContentBrowserConfig* Config)
 {
 	// clang-format off
-	return
+	TSharedRef<SExpandableArea> FavoritesViewWidget = 
 		SAssignNew(FavoritesArea, SExpandableArea)
 		.BorderImage(FAppStyle::Get().GetBrush("Brushes.Header"))
 		.BodyBorderImage(FAppStyle::Get().GetBrush("Brushes.Recessed"))
@@ -1513,6 +1532,9 @@ TSharedRef<SWidget> SContentBrowser::CreateFavoritesView(const FContentBrowserCo
 			]
 		];
 	// clang-format on
+	
+	FavoritePathViewPtr->SetOnFolderFavoriteAdd(SFavoritePathView::FOnFolderFavoriteAdd::CreateSP(this, &SContentBrowser::AddFolderFavorite));
+	return FavoritesViewWidget; 
 }
 
 void SContentBrowser::SetPathViewExpanded(bool bExpanded)
diff --git a/Engine/Source/Editor/ContentBrowser/Private/SContentBrowser.h b/Engine/Source/Editor/ContentBrowser/Private/SContentBrowser.h
--- a/Engine/Source/Editor/ContentBrowser/Private/SContentBrowser.h
+++ b/Engine/Source/Editor/ContentBrowser/Private/SContentBrowser.h
@@ -568,9 +568,15 @@ private:
 	/** Whether or not the collections view is docked or exists in its own panel in the same area as the sources view */
 	bool IsCollectionViewDocked() const;
 
+	/** Set favorite status of an array of folders unless they're already favorited */
+	void AddFolderFavorite(const TArray<FString>& FolderPaths);
+
 	/** Toggles the favorite status of an array of folders*/
 	void ToggleFolderFavorite(const TArray<FString>& FolderPaths);
 
+	/** Save the favorite paths to settings and optionally select and show the paths passed in */
+	void SaveAndShowNewFolderFavorites(const TArray<FString>& FolderPaths);
+
 	/* Toggles the private show private content state of an array of folders*/
 	void TogglePrivateContentEdit(const TArray<FString>& FolderPaths);
 
diff --git a/Engine/Source/Editor/ContentBrowser/Private/SPathView.cpp b/Engine/Source/Editor/ContentBrowser/Private/SPathView.cpp
--- a/Engine/Source/Editor/ContentBrowser/Private/SPathView.cpp
+++ b/Engine/Source/Editor/ContentBrowser/Private/SPathView.cpp
@@ -10,6 +10,7 @@
 #include "Containers/ContainerAllocationPolicies.h"
 #include "Containers/StringView.h"
 #include "ContentBrowserConfig.h"
+#include "ContentBrowserDataDragDropOp.h"
 #include "ContentBrowserDataSource.h"
 #include "ContentBrowserDataSubsystem.h"
 #include "ContentBrowserItemData.h"
@@ -2793,6 +2794,75 @@ void SFavoritePathView::LoadSettings(const FString& IniFilename, const FString&
 	}
 }
 
+
+TSharedPtr<FContentBrowserDataDragDropOp> SFavoritePathView::GetContentBrowserDragDropOpFromEvent(const FDragDropEvent& DragDropEvent) const
+{
+	TSharedPtr<FDragDropOperation> Operation = DragDropEvent.GetOperation();
+	if (Operation.IsValid() && OnFolderFavoriteAdd.IsBound())
+	{
+		if (Operation->IsOfType<FContentBrowserDataDragDropOp>())
+		{
+			TSharedPtr<FContentBrowserDataDragDropOp> DragDropOp = StaticCastSharedPtr<FContentBrowserDataDragDropOp>(Operation);
+
+			// Only agree to the operation if the drag op only contains folders, since favorites cannot contain files. 
+			if (DragDropOp &&
+				!DragDropOp->GetDraggedFolders().IsEmpty() && 
+				DragDropOp->GetDraggedFiles().IsEmpty())
+			{
+				return DragDropOp;
+			}
+		}
+	}
+
+	return {};
+}
+
+
+void SFavoritePathView::OnDragEnter(const FGeometry& MyGeometry, const FDragDropEvent& DragDropEvent)
+{
+	// If we don't have the appropriate drop content, indicate to the user that nothing will happen.
+	if (!GetContentBrowserDragDropOpFromEvent(DragDropEvent))
+	{
+		DragDropEvent.GetOperation()->SetCursorOverride(EMouseCursor::SlashedCircle);
+	}
+}
+
+void SFavoritePathView::OnDragLeave(const FDragDropEvent& DragDropEvent)
+{
+	TSharedPtr<FDragDropOperation> Operation = DragDropEvent.GetOperation();
+	if (Operation.IsValid())
+	{
+		Operation->SetCursorOverride(TOptional<EMouseCursor::Type>());
+	}	
+}
+
+FReply SFavoritePathView::OnDrop(const FGeometry& MyGeometry, const FDragDropEvent& DragDropEvent)
+{
+	TSharedPtr<FContentBrowserDataDragDropOp> DragDropOp = GetContentBrowserDragDropOpFromEvent(DragDropEvent);
+	if (!DragDropOp.IsValid())
+	{
+		return FReply::Unhandled();
+	}
+
+	if (OnFolderFavoriteAdd.IsBound())
+	{
+		TArray<FString> FolderPaths;
+		for (const FContentBrowserItem& BrowserItem: DragDropOp->GetDraggedFolders())
+		{
+			FolderPaths.Add(BrowserItem.GetVirtualPath().ToString());
+		}
+		OnFolderFavoriteAdd.Execute(FolderPaths);
+	}
+	
+	return FReply::Handled();
+}
+
+void SFavoritePathView::SetOnFolderFavoriteAdd(const FOnFolderFavoriteAdd& InOnFolderFavoriteAdd)
+{
+	OnFolderFavoriteAdd = InOnFolderFavoriteAdd;
+}
+
+
 TSharedRef<ITableRow> SFavoritePathView::GenerateTreeRow(TSharedPtr<FTreeItem> TreeItem, const TSharedRef<STableViewBase>& OwnerTable)
 {
 	check(TreeItem.IsValid());
diff --git a/Engine/Source/Editor/ContentBrowser/Private/SPathView.h b/Engine/Source/Editor/ContentBrowser/Private/SPathView.h
--- a/Engine/Source/Editor/ContentBrowser/Private/SPathView.h
+++ b/Engine/Source/Editor/ContentBrowser/Private/SPathView.h
@@ -40,6 +40,7 @@
 #include "Widgets/Views/STableViewBase.h"
 #include "Widgets/Views/STreeView.h"
 
+class FContentBrowserDataDragDropOp;
 class FContentBrowserItemData;
 class FContentBrowserItemDataUpdate;
 class FContentBrowserPluginFilter;
@@ -569,7 +570,7 @@ class SFavoritePathView : public SPathView
 public:
 	SFavoritePathView();
 	virtual ~SFavoritePathView();
-
+	
 	/** Constructs this widget with InArgs */
 	virtual void Construct(const FArguments& InArgs) override;
 
@@ -581,9 +582,21 @@ public:
 	/** Loads any settings to config that should be persistent between editor sessions */
 	virtual void LoadSettings(const FString& IniFilename, const FString& IniSection, const FString& SettingsString) override;
 
+	/** Validate that the drop operation is valid for the favorites. Accept FAssetDragDropOp objects with folders only. */
+	virtual void OnDragEnter(const FGeometry& MyGeometry, const FDragDropEvent& DragDropEvent) override;
+
+	/** Restore any cursor overrides, in case once was set for invalid drop in OnDragEnter. */
+	virtual void OnDragLeave(const FDragDropEvent& DragDropEvent) override;
+	
+	/** Detect if a folder was dropped in the favorites view and add that folder (or folders) to favorites */
+	virtual FReply OnDrop(const FGeometry& MyGeometry, const FDragDropEvent& DragDropEvent) override;
+
 	/** Updates favorites based on an external change. */
 	void FixupFavoritesFromExternalChange(TArrayView<const AssetViewUtils::FMovedContentFolder> MovedFolders);
 
+	DECLARE_DELEGATE_OneParam(FOnFolderFavoriteAdd, const TArray<FString>& /*FoldersToAdd*/)
+	void SetOnFolderFavoriteAdd(const FOnFolderFavoriteAdd& InOnFolderFavoriteAdd);
+	
 private:
 	virtual TSharedRef<ITableRow> GenerateTreeRow(TSharedPtr<FTreeItem> TreeItem, const TSharedRef<STableViewBase>& OwnerTable) override;
 
@@ -594,8 +607,13 @@ private:
 
 	virtual void ConfigureTreeView(STreeView<TSharedPtr<FTreeItem>>::FArguments& InArgs) override;
 
+	/** Returns an FContentBrowserDataDragDropOp object but only if it qualifies as a proper droppable content browser drag drop op */
+	TSharedPtr<FContentBrowserDataDragDropOp> GetContentBrowserDragDropOpFromEvent(const FDragDropEvent& DragDropEvent) const; 
+
 private:
 	TArray<FString> RemovedByFolderMove;
 	FDelegateHandle OnFavoritesChangedHandle;
+	
+	FOnFolderFavoriteAdd OnFolderFavoriteAdd; 
 	bool bIsLoadingSettings = false;
 };
diff --git a/Engine/Source/Editor/ContentBrowserData/Public/ContentBrowserDataDragDropOp.h b/Engine/Source/Editor/ContentBrowserData/Public/ContentBrowserDataDragDropOp.h
--- a/Engine/Source/Editor/ContentBrowserData/Public/ContentBrowserDataDragDropOp.h
+++ b/Engine/Source/Editor/ContentBrowserData/Public/ContentBrowserDataDragDropOp.h
@@ -77,7 +77,7 @@ private:
 private:
 	FName FolderBrushName;
 	FName FolderShadowBrushName;
-	FLinearColor FolderColorOverride;
+	FLinearColor FolderColorOverride = FLinearColor::Gray;
 	TArray<FContentBrowserItem> DraggedItems;
 	TArray<FContentBrowserItem> DraggedFiles;
 	TArray<FContentBrowserItem> DraggedFolders;
