[BUGFIX] Imported Assets except one are selected

- PendingSyncItems wait until the Filtered Asset Items have the latest data

diff --git a/Engine/Source/Editor/ContentBrowser/Private/SAssetView.cpp b/Engine/Source/Editor/ContentBrowser/Private/SAssetView.cpp
--- a/Engine/Source/Editor/ContentBrowser/Private/SAssetView.cpp
+++ b/Engine/Source/Editor/ContentBrowser/Private/SAssetView.cpp
@@ -94,6 +94,9 @@ namespace UE::AssetView
 	/** Time delay between performing the last jump, and the jump term being reset */
 	constexpr double JumpDelaySeconds = 2.0;
 
+	/** Number of frames a deferred pending list will wait before clearing out the data.*/
+	constexpr int32 DeferredSyncTimeoutFramesCount = 30;
+
 	bool AllowAsync = true;
 	FAutoConsoleVariableRef CVarAllowAsync(
 		TEXT("AssetView.AllowAsync"),
@@ -1648,7 +1651,7 @@ void SAssetView::SyncToItems(TArrayView<const FContentBrowserItem> ItemsToSync,
 	{
 		PendingSyncItems.SelectedVirtualPaths.Add(Item.GetVirtualPath());
 	}
-
+	InitDeferredPendingSyncItems();	
 	bPendingFocusOnSync = bFocusOnSync;
 }
 
@@ -1659,7 +1662,7 @@ void SAssetView::SyncToVirtualPaths(TArrayView<const FName> VirtualPathsToSync,
 	{
 		PendingSyncItems.SelectedVirtualPaths.Add(VirtualPathToSync);
 	}
-
+	InitDeferredPendingSyncItems();
 	bPendingFocusOnSync = bFocusOnSync;
 }
 
@@ -1667,10 +1670,16 @@ void SAssetView::SyncToLegacy(TArrayView<const FAssetData> AssetDataList, TArray
 {
 	PendingSyncItems.Reset();
 	ContentBrowserUtils::ConvertLegacySelectionToVirtualPaths(AssetDataList, FolderList, /*UseFolderPaths*/false, PendingSyncItems.SelectedVirtualPaths);
-
+	InitDeferredPendingSyncItems();
 	bPendingFocusOnSync = bFocusOnSync;
 }
 
+void SAssetView::InitDeferredPendingSyncItems()
+{
+	DeferredPendingSyncItems.AddMissingVirtualPaths(PendingSyncItems);
+	DeferredSyncTimeoutFrames = UE::AssetView::DeferredSyncTimeoutFramesCount;
+}
+
 void SAssetView::SyncToSelection( const bool bFocusOnSync )
 {
 	PendingSyncItems.Reset();
@@ -1683,7 +1692,6 @@ void SAssetView::SyncToSelection( const bool bFocusOnSync )
 			PendingSyncItems.SelectedVirtualPaths.Add(Item->GetItem().GetVirtualPath());
 		}
 	}
-
 	bPendingFocusOnSync = bFocusOnSync;
 }
 
@@ -2054,8 +2062,11 @@ void SAssetView::Tick( const FGeometry& AllottedGeometry, const double InCurrent
 			const auto& Item = *ItemIt;
 			if(Item.IsValid())
 			{
-				if (PendingSyncItems.SelectedVirtualPaths.Contains(Item->GetItem().GetVirtualPath()))
+				const FName ItemVirtualPath = Item->GetItem().GetVirtualPath();
+				if (PendingSyncItems.SelectedVirtualPaths.Contains(ItemVirtualPath))
 				{
+					DeferredPendingSyncItems.SelectedVirtualPaths.Remove(ItemVirtualPath);
+
 					SetItemSelection(Item, true, ESelectInfo::OnNavigation);
 					
 					// Scroll the first item in the list that can be shown into view
@@ -2070,19 +2081,32 @@ void SAssetView::Tick( const FGeometry& AllottedGeometry, const double InCurrent
 	
 		bBulkSelecting = false;
 
-		if (bShouldNotifyNextAssetSync && !bUserSearching)
+		if (DeferredSyncTimeoutFrames > 0)
 		{
-			AssetSelectionChanged(TSharedPtr<FAssetViewItem>(), ESelectInfo::Direct);
+			DeferredSyncTimeoutFrames--;
+
+			if (DeferredSyncTimeoutFrames == 0)
+			{
+				DeferredPendingSyncItems.Reset();
+			}
 		}
 
-		// Default to always notifying
-		bShouldNotifyNextAssetSync = true;
+		if (DeferredPendingSyncItems.Num() == 0)
+		{
+			if (bShouldNotifyNextAssetSync && !bUserSearching)
+			{
+				AssetSelectionChanged(TSharedPtr<FAssetViewItem>(), ESelectInfo::Direct);
+			}
 
-		PendingSyncItems.Reset();
+			// Default to always notifying
+			bShouldNotifyNextAssetSync = true;
+			
+			PendingSyncItems.Reset();
 
-		if (bAllowFocusOnSync && bPendingFocusOnSync)
-		{
-			FocusList();
+			if (bAllowFocusOnSync && bPendingFocusOnSync)
+			{
+				FocusList();
+			}
 		}
 	}
 
diff --git a/Engine/Source/Editor/ContentBrowser/Public/HistoryManager.h b/Engine/Source/Editor/ContentBrowser/Public/HistoryManager.h
--- a/Engine/Source/Editor/ContentBrowser/Public/HistoryManager.h
+++ b/Engine/Source/Editor/ContentBrowser/Public/HistoryManager.h
@@ -33,6 +33,14 @@ struct FSelectionData
 	{
 		SelectedVirtualPaths.Empty();
 	}
+
+	void AddMissingVirtualPaths(const FSelectionData& SelectionData)
+	{
+		for (const FName& SelectedVirtualPath : SelectionData.SelectedVirtualPaths)
+		{
+			SelectedVirtualPaths.Add(SelectedVirtualPath);
+		}
+	}
 };
 
 /** The history data object, storing all important history data */
diff --git a/Engine/Source/Editor/ContentBrowser/Public/SAssetView.h b/Engine/Source/Editor/ContentBrowser/Public/SAssetView.h
--- a/Engine/Source/Editor/ContentBrowser/Public/SAssetView.h
+++ b/Engine/Source/Editor/ContentBrowser/Public/SAssetView.h
@@ -359,6 +359,9 @@ public:
 	/** Selects the specified assets and paths. */
 	void SyncToLegacy( TArrayView<const FAssetData> AssetDataList, TArrayView<const FString> FolderList, const bool bFocusOnSync = true );
 
+	/** Setup Deferred Pending Sync for recently added PendingSyncItems. */
+	void InitDeferredPendingSyncItems();
+
 	/** Sets the state of the asset view to the one described by the history data */
 	void ApplyHistoryData( const FHistoryData& History );
 
@@ -1002,6 +1005,12 @@ private:
 	/** The list of items to sync next frame */
 	FSelectionData PendingSyncItems;
 
+	/** The list of items used to ensure all the pending sync items are selected due to async nature of filtering*/
+	FSelectionData DeferredPendingSyncItems;
+
+	/** A Timeout counter to safeguard against infinite deferement of pending sync items*/
+	int32 DeferredSyncTimeoutFrames = 0;
+
 	/** Should we take focus when the PendingSyncAssets are processed? */
 	bool bPendingFocusOnSync;
 
