Added (dirty) external package information to asset view tooltip

diff --git a/Engine/Source/Editor/ContentBrowser/Private/AssetViewWidgets.cpp b/Engine/Source/Editor/ContentBrowser/Private/AssetViewWidgets.cpp
--- a/Engine/Source/Editor/ContentBrowser/Private/AssetViewWidgets.cpp
+++ b/Engine/Source/Editor/ContentBrowser/Private/AssetViewWidgets.cpp
@@ -3,6 +3,9 @@
 
 #include "AssetViewWidgets.h"
 
+#include "ActorFolder.h"
+#include "ActorFolderDesc.h"
+#include "AssetDefinitionRegistry.h"
 #include "AssetRegistry/AssetData.h"
 #include "AssetTagItemTypes.h"
 #include "AssetThumbnail.h"
@@ -53,6 +56,7 @@
 #include "SlateOptMacros.h"
 #include "SlotBase.h"
 #include "SourceControlHelpers.h"
+#include "AssetRegistry/AssetRegistryModule.h"
 #include "Styling/ISlateStyle.h"
 #include "Styling/SlateBrush.h"
 #include "Styling/SlateTypes.h"
@@ -412,6 +416,16 @@ void SAssetViewItem::Construct( const FArguments& InArgs )
 	// control state has been cached; for instance when the widget is killed via FWidgetGenerator::OnEndGenerationPass 
 	// or a view is refreshed due to user filtering/navigating):
 	HandleSourceControlStateChanged();
+
+	FAssetData AssetData;
+	AssetItem->GetItem().Legacy_TryGetAssetData(AssetData);
+	if (AssetData.IsValid())
+	{
+		if (const UAssetDefinition* AssetDefinition = UAssetDefinitionRegistry::Get()->GetAssetDefinitionForAsset(AssetData))
+		{
+			bShouldSaveExternalPackages = AssetDefinition->ShouldSaveExternalPackages();
+		}
+	}
 }
 
 void SAssetViewItem::Tick( const FGeometry& AllottedGeometry, const double InCurrentTime, const float InDeltaTime )
@@ -597,6 +611,83 @@ void SAssetViewItem::HandleSourceControlStateChanged()
 	}
 }
 
+void SAssetViewItem::CacheDirtyExternalPackageInfo()
+{
+	if(bShouldSaveExternalPackages)
+	{
+		CachedDirtyPackagesList.Empty();
+		
+		FAssetData AssetData;		
+		AssetItem->GetItem().Legacy_TryGetAssetData(AssetData);
+		if (AssetData.IsAssetLoaded())
+		{
+			if(const UObject* Asset = AssetData.GetAsset())
+			{
+				if (const UPackage* Package = Asset->GetPackage())
+				{
+					TArray<UPackage*> ExternalPackages = Package->GetExternalPackages();
+					IAssetRegistry& AssetRegistry = FModuleManager::GetModuleChecked<FAssetRegistryModule>(TEXT("AssetRegistry")).Get();
+
+					// Mirrored/copied from SSourceControlCommon.cpp
+					auto RetrieveAssetName = [](const FAssetData& InAssetData) -> FString
+					{
+						static const FName NAME_ActorLabel(TEXT("ActorLabel"));
+						if (InAssetData.FindTag(NAME_ActorLabel))
+						{
+							FString ResultAssetName;
+							InAssetData.GetTagValue(NAME_ActorLabel, ResultAssetName);
+							return ResultAssetName;
+						}
+
+						if (InAssetData.FindTag(FPrimaryAssetId::PrimaryAssetDisplayNameTag))
+						{
+							FString ResultAssetName;
+							InAssetData.GetTagValue(FPrimaryAssetId::PrimaryAssetDisplayNameTag, ResultAssetName);
+							return ResultAssetName;
+						}
+
+						if (InAssetData.AssetClassPath == UActorFolder::StaticClass()->GetClassPathName())
+						{
+							FString ActorFolderPath = UActorFolder::GetAssetRegistryInfoFromPackage(InAssetData.PackageName).GetDisplayName();
+							if (!ActorFolderPath.IsEmpty())
+							{
+								return ActorFolderPath;
+							}
+						}
+
+						return InAssetData.AssetName.ToString();
+					};
+					
+					for (const UPackage* ExternalPackage : ExternalPackages)
+					{
+						if (ExternalPackage->IsDirty())
+						{
+							TArray<FAssetData> DirtyAssetDataEntries;
+							AssetRegistry.GetAssetsByPackageName(*ExternalPackage->GetName(), DirtyAssetDataEntries);
+
+							if (CachedDirtyPackagesList.Len())
+							{								
+								CachedDirtyPackagesList.Append("\n");
+							}
+							
+							CachedDirtyPackagesList.Append(ExternalPackage->GetPathName());
+							
+							for (const FAssetData& DirtyAssetData : DirtyAssetDataEntries)
+							{
+								const FString AssetName = RetrieveAssetName(DirtyAssetData);
+								const FString AssetClass = DirtyAssetData.AssetClassPath.GetAssetName().ToString();
+						
+								CachedDirtyPackagesList.Append("\n\t");
+								CachedDirtyPackagesList.Append(FString::Printf(TEXT("%s (%s)"), *AssetName, *AssetClass)); 
+							}
+						}
+					}
+				}
+			}
+		}
+	}
+}
+
 const FSlateBrush* SAssetViewItem::GetDirtyImage() const
 {
 	return IsDirty() ? AssetDirtyBrush : NULL;
@@ -902,6 +993,22 @@ TSharedRef<SWidget> SAssetViewItem::CreateToolTipWidget() const
 							.ColorAndOpacity(FLinearColor(0.1f, 0.5f, 1.f, 1.f))
 						]
 
+						+ SVerticalBox::Slot()
+						.AutoHeight()
+						[
+							SNew(STextBlock)
+							.Visibility_Lambda([this]()-> EVisibility { return bShouldSaveExternalPackages && !GetExternalPackagesText().IsEmpty() ? EVisibility::Visible : EVisibility::Collapsed; })
+							.Text(LOCTEXT("DirtyExternalPackages", "Modified external packages:"))
+							.ColorAndOpacity(FStyleColors::Warning)
+						]
+						+ SVerticalBox::Slot()
+						.AutoHeight()
+						[
+							SNew(STextBlock)
+							.Visibility_Lambda([this]()-> EVisibility { return bShouldSaveExternalPackages && !GetExternalPackagesText().IsEmpty() ? EVisibility::Visible : EVisibility::Collapsed; })
+							.Text(this, &SAssetViewItem::GetExternalPackagesText)
+						]
+
 						+ SVerticalBox::Slot()
 						.AutoHeight()
 						[
@@ -1108,6 +1215,16 @@ FText SAssetViewItem::GetSourceControlText() const
 	return FText::GetEmpty();
 }
 
+FText SAssetViewItem::GetExternalPackagesText() const
+{
+	if (CachedDirtyPackagesList.Len())
+	{
+		return FText::FromString(CachedDirtyPackagesList);
+	}
+
+	return FText::GetEmpty();
+}
+
 FText SAssetViewItem::GetAssetUserDescription() const
 {
 	if (AssetItem && AssetItem->IsFile())
@@ -1167,6 +1284,8 @@ void SAssetViewItem::UpdateDirtyState()
 		bItemDirty = bNewIsDirty;
 		DirtyStateChanged();
 	}
+		
+	CacheDirtyExternalPackageInfo();
 }
 
 bool SAssetViewItem::IsDirty() const
diff --git a/Engine/Source/Editor/ContentBrowser/Public/AssetViewWidgets.h b/Engine/Source/Editor/ContentBrowser/Public/AssetViewWidgets.h
--- a/Engine/Source/Editor/ContentBrowser/Public/AssetViewWidgets.h
+++ b/Engine/Source/Editor/ContentBrowser/Public/AssetViewWidgets.h
@@ -229,6 +229,9 @@ protected:
 	/** Gets the text for the source control text block in the tooltip */
 	FText GetSourceControlText() const;
 
+	/** Gets the text detailing all the modified external packages (and their contained objects) */
+	FText GetExternalPackagesText() const;
+
 	/** Helper function for CreateToolTipWidget. Gets the user description for the asset, if it exists. */
 	FText GetAssetUserDescription() const;
 
@@ -259,6 +262,9 @@ protected:
 	/** Returns the width at which the name label will wrap the name */
 	virtual float GetNameTextWrapWidth() const { return 0.0f; }
 
+	/** Caches the ExternalPackagesText */
+	void CacheDirtyExternalPackageInfo();
+
 protected:
 	/** Data for a cached display tag for this item (used in the tooltip, and also as the display string in column views) */
 	struct FTagDisplayItem
@@ -346,6 +352,9 @@ protected:
 	/** Whether the item has a valid scc state brush (not empty) */
 	bool bHasCCStateBrush = false;
 
+	bool bShouldSaveExternalPackages = false;
+	FString CachedDirtyPackagesList;
+
 	/** Delegate handle for the HandleSourceControlStateChanged function callback */
 	FDelegateHandle SourceControlStateChangedDelegateHandle;
 };
