Removed the default hotkey for the save all in folder command of the CB as it is interfering with another engine command to save all the dirty assets.

diff --git a/Engine/Source/Editor/ContentBrowser/Private/ContentBrowserCommands.cpp b/Engine/Source/Editor/ContentBrowser/Private/ContentBrowserCommands.cpp
--- a/Engine/Source/Editor/ContentBrowser/Private/ContentBrowserCommands.cpp
+++ b/Engine/Source/Editor/ContentBrowser/Private/ContentBrowserCommands.cpp
@@ -18,7 +18,7 @@ void FContentBrowserCommands::RegisterCommands()
 	UI_COMMAND(PreviewAssets, "Preview Assets", "Loads the selected assets and previews them if possible", EUserInterfaceActionType::Button, FInputChord(EKeys::SpaceBar));
 	UI_COMMAND(CreateNewFolder, "Create New Folder", "Creates new folder in selected path", EUserInterfaceActionType::Button, FInputChord(EModifierKey::Control | EModifierKey::Shift, EKeys::N));
 	UI_COMMAND(SaveSelectedAsset, "Save Selected Item", "Save the selected item", EUserInterfaceActionType::Button, FInputChord(EModifierKey::Control, EKeys::S));
-	UI_COMMAND(SaveAllCurrentFolder, "Save All", "Save All in current folder", EUserInterfaceActionType::Button, FInputChord(EModifierKey::Control | EModifierKey::Shift, EKeys::S));
+	UI_COMMAND(SaveAllCurrentFolder, "Save All", "Save All in current folder", EUserInterfaceActionType::Button, FInputChord());
 	UI_COMMAND(ResaveAllCurrentFolder, "Resave All", "Resave all assets contained in the current folder", EUserInterfaceActionType::Button, FInputChord());
 	UI_COMMAND(EditPath, "Edit Path", "Edit the current content browser path", EUserInterfaceActionType::Button, FInputChord(EKeys::F4), FInputChord(EModifierKey::Control, EKeys::L));
 
