Setting to hide specific Collections from the ContentBrowser's collection view
* Setting is registered per project.
* There is a local user setting which will disable the effect.
* Collections still accessible via other methods (scripting, blueprints, etc)

diff --git a/Engine/Source/Editor/ContentBrowser/ContentBrowser.Build.cs b/Engine/Source/Editor/ContentBrowser/ContentBrowser.Build.cs
--- a/Engine/Source/Editor/ContentBrowser/ContentBrowser.Build.cs
+++ b/Engine/Source/Editor/ContentBrowser/ContentBrowser.Build.cs
@@ -22,6 +22,7 @@ public class ContentBrowser : ModuleRules
 				"AppFramework",
 				"Core",
 				"CoreUObject",
+				"DeveloperSettings",
 				"ApplicationCore",
                 "InputCore",
 				"EditorConfig",
diff --git a/Engine/Source/Editor/ContentBrowser/Private/ContentBrowserConfig.h b/Engine/Source/Editor/ContentBrowser/Private/ContentBrowserConfig.h
--- a/Engine/Source/Editor/ContentBrowser/Private/ContentBrowserConfig.h
+++ b/Engine/Source/Editor/ContentBrowser/Private/ContentBrowserConfig.h
@@ -4,6 +4,7 @@
 
 #include "AssetThumbnail.h"
 #include "EditorConfigBase.h"
+#include "Engine/DeveloperSettings.h"
 #include "UObject/ObjectPtr.h"
 #include "UObject/UObjectGlobals.h"
 
@@ -113,4 +114,14 @@ public:
 
 private:
 	static TObjectPtr<UContentBrowserConfig> Instance;
+};
+
+UCLASS(minimalapi, Config=Editor, DisplayName="Content Browser")
+class UContentBrowserCollectionProjectSettings : public UDeveloperSettings
+{
+	GENERATED_BODY()
+public:
+
+	UPROPERTY(config, EditAnywhere, Category=Collections, meta = (DisplayName="Exclude these collections from Content Browser"))
+	TArray<FName> ExcludedCollectionsFromView;
 };
\ No newline at end of file
diff --git a/Engine/Source/Editor/ContentBrowser/Private/SCollectionView.cpp b/Engine/Source/Editor/ContentBrowser/Private/SCollectionView.cpp
--- a/Engine/Source/Editor/ContentBrowser/Private/SCollectionView.cpp
+++ b/Engine/Source/Editor/ContentBrowser/Private/SCollectionView.cpp
@@ -8,6 +8,7 @@
 #include "CollectionManagerModule.h"
 #include "CollectionViewTypes.h"
 #include "CollectionViewUtils.h"
+#include "ContentBrowserConfig.h"
 #include "ContentBrowserDelegates.h"
 #include "ContentBrowserModule.h"
 #include "ContentBrowserPluginFilters.h"
@@ -51,6 +52,7 @@
 #include "SourcesSearch.h"
 #include "SourcesViewWidgets.h"
 #include "TelemetryRouter.h"
+#include "Settings/ContentBrowserSettings.h"
 #include "Styling/AppStyle.h"
 #include "Styling/ISlateStyle.h"
 #include "Styling/SlateColor.h"
@@ -177,6 +179,14 @@ void SCollectionView::Construct( const FArguments& InArgs )
 	ExternalSearchPtr = InArgs._ExternalSearch;
 	TitleContent = SNew(SHorizontalBox);
 
+	UContentBrowserCollectionProjectSettings* CollectionProjectSettings = GetMutableDefault<UContentBrowserCollectionProjectSettings>();
+	if (CollectionProjectSettings)
+	{
+		CollectionProjectSettings->OnSettingChanged().AddSPLambda(this, [this](UObject* CollectionProjectSettings_, struct FPropertyChangedEvent&)
+		{
+			UpdateFilteredCollectionItems();
+		});
+	}
 
 	PreventSelectionChangedDelegateCount = 0;
 
@@ -493,8 +503,20 @@ void SCollectionView::UpdateFilteredCollectionItems()
 	VisibleCollections.Reset();
 	VisibleRootCollectionItems.Reset();
 
-	auto AddVisibleCollection = [&](const TSharedPtr<FCollectionItem>& InCollectionItem)
+	const UContentBrowserCollectionProjectSettings* CollectionProjectSettings = GetDefault<UContentBrowserCollectionProjectSettings>();
+	const UContentBrowserSettings* ContentBrowserSettings = GetDefault<UContentBrowserSettings>();
+
+	auto AddVisibleCollection = [&](const FCollectionNameType& NameTypePair, const TSharedPtr<FCollectionItem>& InCollectionItem)
 	{
+		if (!ContentBrowserSettings->bDisplayExcludedCollections)
+		{
+			const int32 FoundIndex = CollectionProjectSettings->ExcludedCollectionsFromView.Find(NameTypePair.Name);
+			if (FoundIndex != INDEX_NONE)
+			{
+				return;
+			}
+		}
+		
 		VisibleCollections.Add(FCollectionNameType(InCollectionItem->CollectionName, InCollectionItem->CollectionType));
 		if (!InCollectionItem->ParentCollection.IsValid())
 		{
@@ -502,12 +524,12 @@ void SCollectionView::UpdateFilteredCollectionItems()
 		}
 	};
 
-	auto AddVisibleCollectionRecursive = [&](const TSharedPtr<FCollectionItem>& InCollectionItem)
+	auto AddVisibleCollectionRecursive = [&](const FCollectionNameType& NameTypePair, const TSharedPtr<FCollectionItem>& InCollectionItem)
 	{
 		TSharedPtr<FCollectionItem> CollectionItemToAdd = InCollectionItem;
 		do
 		{
-			AddVisibleCollection(CollectionItemToAdd);
+			AddVisibleCollection(NameTypePair, CollectionItemToAdd);
 			CollectionItemToAdd = CollectionItemToAdd->ParentCollection.Pin();
 		}
 		while(CollectionItemToAdd.IsValid());
@@ -519,7 +541,7 @@ void SCollectionView::UpdateFilteredCollectionItems()
 		// No filter, just mark everything as visible
 		for (const auto& AvailableCollectionInfo : AvailableCollections)
 		{
-			AddVisibleCollection(AvailableCollectionInfo.Value);
+			AddVisibleCollection(AvailableCollectionInfo.Key, AvailableCollectionInfo.Value);
 		}
 	}
 	else
@@ -532,7 +554,7 @@ void SCollectionView::UpdateFilteredCollectionItems()
 			const TSharedPtr<FCollectionItem>& CollectionItem = AvailableCollectionInfo.Value;
 			if (CollectionItemTextFilter->PassesFilter(*CollectionItem))
 			{
-				AddVisibleCollectionRecursive(CollectionItem);
+				AddVisibleCollectionRecursive(AvailableCollectionInfo.Key, CollectionItem);
 				CollectionsToExpandTo.Add(CollectionItem.ToSharedRef());
 			}
 		}
@@ -543,6 +565,8 @@ void SCollectionView::UpdateFilteredCollectionItems()
 			ExpandParentItems(CollectionItem);
 		}
 	}
+	
+	UContentBrowserSettings::OnSettingChanged().AddSP(this, &SCollectionView::HandleSettingChanged);
 
 	VisibleRootCollectionItems.Sort(FCollectionItem::FCompareFCollectionItemByName());
 	CollectionTreePtr->RequestTreeRefresh();
@@ -1396,6 +1420,14 @@ FReply SCollectionView::HandleDragDropOnCollectionItem(TSharedRef<FCollectionIte
 	return FReply::Unhandled();
 }
 
+void SCollectionView::HandleSettingChanged(FName PropertyName)
+{
+	if (PropertyName == GET_MEMBER_NAME_CHECKED(UContentBrowserSettings, bDisplayExcludedCollections))
+	{
+		UpdateFilteredCollectionItems();
+	}
+}
+
 void SCollectionView::ExpandParentItems(const TSharedRef<FCollectionItem>& InCollectionItem)
 {
 	for (TSharedPtr<FCollectionItem> CollectionItemToExpand = InCollectionItem->ParentCollection.Pin(); 
diff --git a/Engine/Source/Editor/ContentBrowser/Private/SCollectionView.h b/Engine/Source/Editor/ContentBrowser/Private/SCollectionView.h
--- a/Engine/Source/Editor/ContentBrowser/Private/SCollectionView.h
+++ b/Engine/Source/Editor/ContentBrowser/Private/SCollectionView.h
@@ -208,6 +208,8 @@ private:
 	/** Handle dropping something on one of our collections */
 	FReply HandleDragDropOnCollectionItem(TSharedRef<FCollectionItem> CollectionItem, const FGeometry& Geometry, const FDragDropEvent& DragDropEvent);
 
+	void HandleSettingChanged(FName PropertyName);
+
 	/** Recursively expand the parent items of this collection to ensure that it is visible */
 	void ExpandParentItems(const TSharedRef<FCollectionItem>& InCollectionItem);
 
diff --git a/Engine/Source/Editor/UnrealEd/Classes/Settings/ContentBrowserSettings.h b/Engine/Source/Editor/UnrealEd/Classes/Settings/ContentBrowserSettings.h
--- a/Engine/Source/Editor/UnrealEd/Classes/Settings/ContentBrowserSettings.h
+++ b/Engine/Source/Editor/UnrealEd/Classes/Settings/ContentBrowserSettings.h
@@ -65,6 +65,10 @@ public:
 	UPROPERTY(EditAnywhere, config, Category = ContentBrowser)
 	bool bEnableRealtimeMaterialInstanceThumbnails = true;
 
+	/** Displays collections in the Collections window that were excluded by project settings */
+	UPROPERTY(Config, EditAnywhere, Category=Collections, meta = (DisplayName="Display excluded collections in Content Browser"))
+	bool bDisplayExcludedCollections = false;
+
 public:
 
 	/** Sets whether we are allowed to display the engine folder or not, optional flag for setting override instead */
