Fixed a potential memory stomp (detected with -stompmalloc)

diff --git a/Engine/Source/Editor/ContentBrowser/Private/AssetTextFilter.cpp b/Engine/Source/Editor/ContentBrowser/Private/AssetTextFilter.cpp
--- a/Engine/Source/Editor/ContentBrowser/Private/AssetTextFilter.cpp
+++ b/Engine/Source/Editor/ContentBrowser/Private/AssetTextFilter.cpp
@@ -256,6 +256,7 @@ bool FCompiledAssetTextFilter::PassesFilter(FAssetFilterType InItem)
 
 	if (bIncludeAssetPath)
 	{
+		int32 TextBufferLenBefore = TextBuffer.Len();
 		// Get the full asset path, and also split it so we can compare each part in the filter
 		AssetPtr->GetVirtualPath().AppendString(TextBuffer);
         
@@ -267,17 +268,19 @@ bool FCompiledAssetTextFilter::PassesFilter(FAssetFilterType InItem)
 				&& LastDotIndex > LastSlashIndex && LastSlashIndex >= DisplayNameLen)
 			{
 				TextBuffer.LeftInline(LastDotIndex, EAllowShrinking::No);
+				checkf(TextBuffer.Len() >= TextBufferLenBefore, TEXT("If TextBuffer is smaller than it used to be, then DisplayNameLen is now invalid. That should never happen."));
 			}
 		}
-		AssetPathLen = TextBuffer.Len() - DisplayNameLen;
+		AssetPathLen = TextBuffer.Len() - TextBufferLenBefore;
 
 		if (bIncludeClassName && !AssetPtr->IsFolder())
 		{
+			TextBufferLenBefore = TextBuffer.Len();
 			// Get the full export text path as people sometimes search by copying this (requires class and asset path
 			// search to be enabled in order to match)
 			// TODO: this allocates a temporary FString inside the backends
 			AssetPtr->AppendItemReference(TextBuffer);
-			ExportTextPathLen = TextBuffer.Len() - AssetPathLen;
+			ExportTextPathLen = TextBuffer.Len() - TextBufferLenBefore;
 		}
 	}
 
