[ContentBrowser] AssetViewSortManager - added disk size & date sorting

#rnx

diff --git a/Engine/Source/Editor/ContentBrowser/Private/AssetViewSortManager.cpp b/Engine/Source/Editor/ContentBrowser/Private/AssetViewSortManager.cpp
--- a/Engine/Source/Editor/ContentBrowser/Private/AssetViewSortManager.cpp
+++ b/Engine/Source/Editor/ContentBrowser/Private/AssetViewSortManager.cpp
@@ -24,7 +24,7 @@ public:
 	/** Default constructor */
 	FCompareFAssetItemBase(bool bInAscending, const FName& InTag) : bAscending(bInAscending), Tag(InTag) {}
 
-	virtual ~FCompareFAssetItemBase() {}
+	virtual ~FCompareFAssetItemBase() = default;
 
 	/** Sort function */
 	FORCEINLINE bool operator()(const TSharedPtr<FAssetViewItem>& A, const TSharedPtr<FAssetViewItem>& B) const
@@ -120,7 +120,7 @@ struct FCompareFAssetItemByName : public FCompareFAssetItemBase
 {
 public:
 	FCompareFAssetItemByName(bool bInAscending, const FName& InTag) : FCompareFAssetItemBase(bInAscending, InTag) {}
-	virtual ~FCompareFAssetItemByName() {}
+	virtual ~FCompareFAssetItemByName() override = default;
 
 protected:
 	FORCEINLINE virtual bool Compare(const TSharedPtr<FAssetViewItem>& A, const TSharedPtr<FAssetViewItem>& B) const override
@@ -147,7 +147,7 @@ struct FCompareFAssetItemByClass : public FCompareFAssetItemBase
 {
 public:
 	FCompareFAssetItemByClass(bool bInAscending, const FName& InTag) : FCompareFAssetItemBase(bInAscending, InTag) {}
-	virtual ~FCompareFAssetItemByClass() {}
+	virtual ~FCompareFAssetItemByClass() override = default;
 
 protected:
 	FORCEINLINE virtual bool Compare(const TSharedPtr<FAssetViewItem>& A, const TSharedPtr<FAssetViewItem>& B) const override
@@ -173,11 +173,43 @@ protected:
 	}
 };
 
+struct FCompareFAssetItemByDiskSize : public FCompareFAssetItemBase
+{
+public:
+	FCompareFAssetItemByDiskSize(bool bInAscending, const FName& InTag) : FCompareFAssetItemBase(bInAscending, InTag) {}
+	virtual ~FCompareFAssetItemByDiskSize() override = default;
+
+protected:
+	FORCEINLINE virtual bool Compare(const TSharedPtr<FAssetViewItem>& A, const TSharedPtr<FAssetViewItem>& B) const override
+	{
+		auto GetDiskSize = [](const FContentBrowserItem& Item)
+		{
+			FContentBrowserItemDataAttributeValue DiskSizeValue = Item.GetItemAttribute(ContentBrowserItemAttributes::ItemDiskSize);
+			return DiskSizeValue.IsValid() ? DiskSizeValue.GetValue<int64>() : -1; // -1 is for items where size doesn't apply - ie. folders
+		};
+
+		const int64 ValueA = GetDiskSize(A->GetItem());
+		const int64 ValueB = GetDiskSize(B->GetItem());
+
+		if (ValueA < ValueB)
+		{
+			return !bAscending;
+		}
+
+		if (ValueB > ValueA)
+		{
+			return bAscending;
+		}
+
+		return FCompareFAssetItemBase::Compare(A, B);
+	}
+};
+
 struct FCompareFAssetItemByPath : public FCompareFAssetItemBase
 {
 public:
 	FCompareFAssetItemByPath(bool bInAscending, const FName& InTag) : FCompareFAssetItemBase(bInAscending, InTag) {}
-	virtual ~FCompareFAssetItemByPath() {}
+	virtual ~FCompareFAssetItemByPath() override = default;
 
 protected:
 	FORCEINLINE virtual bool Compare(const TSharedPtr<FAssetViewItem>& A, const TSharedPtr<FAssetViewItem>& B) const override
@@ -201,7 +233,7 @@ struct FCompareFAssetItemByTag : public FCompareFAssetItemBase
 {
 public:
 	FCompareFAssetItemByTag(bool bInAscending, const FName& InTag) : FCompareFAssetItemBase(bInAscending, InTag) {}
-	virtual ~FCompareFAssetItemByTag() {}
+	virtual ~FCompareFAssetItemByTag() override = default;
 
 protected:
 	FORCEINLINE virtual bool Compare(const TSharedPtr<FAssetViewItem>& A, const TSharedPtr<FAssetViewItem>& B) const override
@@ -238,7 +270,7 @@ struct FCompareFAssetItemByTagNumerical : public FCompareFAssetItemBase
 {
 public:
 	FCompareFAssetItemByTagNumerical(bool bInAscending, const FName& InTag) : FCompareFAssetItemBase(bInAscending, InTag) {}
-	virtual ~FCompareFAssetItemByTagNumerical() {}
+	virtual ~FCompareFAssetItemByTagNumerical() override = default;
 
 protected:
 	FORCEINLINE virtual bool Compare(const TSharedPtr<FAssetViewItem>& A, const TSharedPtr<FAssetViewItem>& B) const override
@@ -278,7 +310,7 @@ struct FCompareFAssetItemByTagDimensional : public FCompareFAssetItemBase
 {
 public:
 	FCompareFAssetItemByTagDimensional(bool bInAscending, const FName& InTag) : FCompareFAssetItemBase(bInAscending, InTag) {}
-	virtual ~FCompareFAssetItemByTagDimensional() {}
+	virtual ~FCompareFAssetItemByTagDimensional() override = default;
 
 protected:
 	FORCEINLINE virtual bool Compare(const TSharedPtr<FAssetViewItem>& A, const TSharedPtr<FAssetViewItem>& B) const override
@@ -332,7 +364,7 @@ const FName FAssetViewSortManager::NameColumnId = "Name";
 const FName FAssetViewSortManager::ClassColumnId = "Class";
 const FName FAssetViewSortManager::PathColumnId = "Path";
 const FName FAssetViewSortManager::RevisionControlColumnId = "RevisionControl";
-
+const FName FAssetViewSortManager::DiskSizeColumnId = "Size";
 
 FAssetViewSortManager::FAssetViewSortManager()
 {
@@ -401,6 +433,10 @@ void FAssetViewSortManager::SortList(TArray<TSharedPtr<FAssetViewItem>>& AssetIt
 		{
 			SortMethod.Add(MakeUnique<FCompareFAssetItemByPath>(bAscending, Tag));
 		}
+		else if (Tag == DiskSizeColumnId)
+		{
+			SortMethod.Add(MakeUnique<FCompareFAssetItemByDiskSize>(bAscending, Tag));
+		}
 		else
 		{
 			UObject::FAssetRegistryTag::ETagType TagType = UObject::FAssetRegistryTag::TT_Hidden;
@@ -509,6 +545,10 @@ void FAssetViewSortManager::ExportColumnsToCSV(TArray<TSharedPtr<FAssetViewItem>
 			{
 				ValueString = AssetItemData->GetVirtualPath().ToString();
 			}
+			else if (Column == DiskSizeColumnId)
+			{
+				AssetItem->GetTagValue(ContentBrowserItemAttributes::ItemDiskSize, ValueString);
+			}
 			else
 			{
 				AssetItem->GetTagValue(Column, ValueString);
@@ -544,6 +584,7 @@ void FAssetViewSortManager::SetSortColumnId(const EColumnSortPriority::Type InSo
 			}
 		}
 	}
+
 	if (bOrderChanged)
 	{
 		// If the order has changed, we need to remove any unneeded sorts by bumping the priority of the remaining valid ones
diff --git a/Engine/Source/Editor/ContentBrowser/Public/AssetViewSortManager.h b/Engine/Source/Editor/ContentBrowser/Public/AssetViewSortManager.h
--- a/Engine/Source/Editor/ContentBrowser/Public/AssetViewSortManager.h
+++ b/Engine/Source/Editor/ContentBrowser/Public/AssetViewSortManager.h
@@ -54,6 +54,7 @@ public:
 	static const FName NameColumnId;
 	static const FName ClassColumnId;
 	static const FName PathColumnId;
+	static const FName DiskSizeColumnId;
 
 	// The revision control column. NOTE: This column currently doesn't support sorting, but is wired through the Sort Manager so the feature can be added in the future
 	static const FName RevisionControlColumnId;
