Fixed renaming or deleting a collection switching the Content Browser's asset view to show all assets

This bug was introdiced in 2231311, as it would set the sources data to empty (which means "show all") whenever a collection was renamed or deleted

#rnx

diff --git a/Engine/Source/Editor/ContentBrowser/Private/SContentBrowser.cpp b/Engine/Source/Editor/ContentBrowser/Private/SContentBrowser.cpp
--- a/Engine/Source/Editor/ContentBrowser/Private/SContentBrowser.cpp
+++ b/Engine/Source/Editor/ContentBrowser/Private/SContentBrowser.cpp
@@ -4322,21 +4322,43 @@ void SContentBrowser::OnAssetViewRefreshRequested()
 
 void SContentBrowser::HandleCollectionRemoved(const FCollectionNameType& Collection)
 {
-	AssetViewPtr->SetSourcesData(FSourcesData());
-
-	auto RemoveHistoryDelegate = [&](const FHistoryData& HistoryData)
+	FSourcesData NewSourcesData = AssetViewPtr->GetSourcesData();
+	if (NewSourcesData.Collections.RemoveSingle(Collection) > 0)
 	{
-		return (HistoryData.SourcesData.Collections.Num() == 1 &&
-				HistoryData.SourcesData.VirtualPaths.Num() == 0 &&
-				HistoryData.SourcesData.Collections.Contains(Collection));
-	};
+		AssetViewPtr->SetSourcesData(NewSourcesData);
+	}
 
-	HistoryManager.RemoveHistoryData(RemoveHistoryDelegate);
+	HistoryManager.RemoveHistoryData([&Collection](const FHistoryData& HistoryData)
+	{
+		return HistoryData.SourcesData.Collections.Num() == 1
+			&& HistoryData.SourcesData.VirtualPaths.Num() == 0
+			&& HistoryData.SourcesData.Collections.Contains(Collection);
+	});
 }
 
 void SContentBrowser::HandleCollectionRenamed(const FCollectionNameType& OriginalCollection, const FCollectionNameType& NewCollection)
 {
-	return HandleCollectionRemoved(OriginalCollection);
+	FSourcesData NewSourcesData = AssetViewPtr->GetSourcesData();
+	if (NewSourcesData.Collections.RemoveSingle(OriginalCollection) > 0)
+	{
+		NewSourcesData.Collections.Add(NewCollection);
+		AssetViewPtr->SetSourcesData(NewSourcesData);
+	}
+
+	HistoryManager.RewriteHistoryData([&OriginalCollection, &NewCollection](FHistoryData& HistoryData)
+	{
+		if (HistoryData.SourcesData.Collections.RemoveSingle(OriginalCollection) > 0)
+		{
+			HistoryData.SourcesData.Collections.Add(NewCollection);
+		}
+	});
+
+	HistoryManager.RemoveHistoryData([&OriginalCollection](const FHistoryData& HistoryData)
+	{
+		return HistoryData.SourcesData.Collections.Num() == 1
+			&& HistoryData.SourcesData.VirtualPaths.Num() == 0
+			&& HistoryData.SourcesData.Collections.Contains(OriginalCollection);
+	});
 }
 
 void SContentBrowser::HandleCollectionUpdated(const FCollectionNameType& Collection)
@@ -4361,14 +4383,12 @@ void SContentBrowser::HandleCollectionUpdated(const FCollectionNameType& Collect
 
 void SContentBrowser::HandlePathRemoved(const FName Path)
 {
-	auto RemoveHistoryDelegate = [&](const FHistoryData& HistoryData)
+	HistoryManager.RemoveHistoryData([&Path](const FHistoryData& HistoryData)
 	{
-		return (HistoryData.SourcesData.VirtualPaths.Num() == 1 &&
-				HistoryData.SourcesData.Collections.Num() == 0 &&
-				HistoryData.SourcesData.VirtualPaths.Contains(Path));
-	};
-
-	HistoryManager.RemoveHistoryData(RemoveHistoryDelegate);
+		return HistoryData.SourcesData.VirtualPaths.Num() == 1
+			&& HistoryData.SourcesData.Collections.Num() == 0
+			&& HistoryData.SourcesData.VirtualPaths.Contains(Path);
+	});
 }
 
 void SContentBrowser::HandleItemDataUpdated(TArrayView<const FContentBrowserItemDataUpdate> InUpdatedItems)
diff --git a/Engine/Source/Editor/ContentBrowser/Public/HistoryManager.h b/Engine/Source/Editor/ContentBrowser/Public/HistoryManager.h
--- a/Engine/Source/Editor/ContentBrowser/Public/HistoryManager.h
+++ b/Engine/Source/Editor/ContentBrowser/Public/HistoryManager.h
@@ -108,6 +108,16 @@ public:
 	 */
 	void GetAvailableHistoryMenuItems(bool bGetPrior, FMenuBuilder& MenuBuilder);
 
+	/** Rewrite all history data as determined by the passed in predicate */
+	template <class PREDICATE_CLASS>
+	void RewriteHistoryData(const PREDICATE_CLASS& Predicate)
+	{
+		for (FHistoryData& HistoryDataEntry : HistoryData)
+		{
+			Predicate(HistoryDataEntry);
+		}
+	}
+
 	/** Removes all history data as determined by the passed in predicate */
 	template <class PREDICATE_CLASS>
 	void RemoveHistoryData(const PREDICATE_CLASS& Predicate)
