[ContentBrowser] Tooltip

diff --git a/Engine/Source/Editor/AssetDefinition/Private/AssetStatusAssetDataInfoProvider.cpp b/Engine/Source/Editor/AssetDefinition/Private/AssetStatusAssetDataInfoProvider.cpp
--- a/Engine/Source/Editor/AssetDefinition/Private/AssetStatusAssetDataInfoProvider.cpp
+++ b/Engine/Source/Editor/AssetDefinition/Private/AssetStatusAssetDataInfoProvider.cpp
@@ -1,6 +1,8 @@
 ﻿// Copyright Epic Games, Inc. All Rights Reserved.
 
 #include "AssetStatusAssetDataInfoProvider.h"
+
+#if UE_CONTENTBROWSER_NEW_STYLE
 #include "AssetRegistry/AssetData.h"
 #include "Misc/PackageName.h"
 #include "UObject/UObjectGlobals.h"
@@ -17,3 +19,4 @@ FString FAssetStatusAssetDataInfoProvider::TryGetFilename() const
 	FPackageName::TryConvertLongPackageNameToFilename(AssetData.PackageName.ToString(), OutFileName, *PackageExtension);
 	return OutFileName;
 }
+#endif
diff --git a/Engine/Source/Editor/AssetDefinition/Public/AssetDefinition.h b/Engine/Source/Editor/AssetDefinition/Public/AssetDefinition.h
--- a/Engine/Source/Editor/AssetDefinition/Public/AssetDefinition.h
+++ b/Engine/Source/Editor/AssetDefinition/Public/AssetDefinition.h
@@ -2,7 +2,6 @@
 
 #pragma once
 
-#include "AssetDefinitionStatusBase.h"
 #include "UObject/Object.h"
 #include "UObject/SoftObjectPtr.h"
 #include "AssetRegistry/AssetData.h"
@@ -15,7 +14,10 @@
 
 class IToolkitHost;
 class UThumbnailInfo;
+struct FAssetDisplayInfo;
 struct FSlateBrush;
+class IAssetSystemInfoProvider;
+class IAssetStatusInfoProvider;
 class ISourceControlRevision;
 class SWidget;
 
@@ -600,7 +602,7 @@ public:
 
 #if UE_CONTENTBROWSER_NEW_STYLE
 	// Implementation will change in the AssetDefinitionDefault in a later update
-	virtual void GetAssetStatusInfo(const TSharedPtr<IAssetStatusInfoProvider>& InAssetStatusInfoProvider, FAssetStatusInfo& OutStatusInfo) const
+	virtual void GetAssetStatusInfo(const TSharedPtr<IAssetStatusInfoProvider>& InAssetStatusInfoProvider, TArray<FAssetDisplayInfo>& OutStatusInfo) const
 	{
 		return;
 	}
diff --git a/Engine/Source/Editor/AssetDefinition/Public/AssetDefinitionAssetInfo.h b/Engine/Source/Editor/AssetDefinition/Public/AssetDefinitionAssetInfo.h
new file mode 100644
--- /dev/null
+++ b/Engine/Source/Editor/AssetDefinition/Public/AssetDefinitionAssetInfo.h
@@ -0,0 +1,82 @@
+// Copyright Epic Games, Inc. All Rights Reserved.
+
+#pragma once
+
+#if UE_CONTENTBROWSER_NEW_STYLE
+#include "Internationalization/Text.h"
+#include "Layout/Visibility.h"
+#include "Misc/Attribute.h"
+
+struct FSlateBrush;
+
+enum class EStatusSeverity
+{
+	Info,
+	Warning,
+	Error
+};
+ENUM_CLASS_FLAGS(EStatusSeverity)
+
+struct FAssetStatusPriority
+{
+	FAssetStatusPriority()
+		: Severity(EStatusSeverity::Info)
+		, SeverityPriority(0)
+	{}
+
+	FAssetStatusPriority(const TAttribute<EStatusSeverity>& InSeverity)
+		: Severity(InSeverity)
+		, SeverityPriority(0)
+	{}
+
+	FAssetStatusPriority(const TAttribute<EStatusSeverity>& InSeverity, int32 InSeverityPriority)
+		: Severity(InSeverity)
+		, SeverityPriority(InSeverityPriority)
+	{}
+
+public:
+	bool operator==(const FAssetStatusPriority& InOtherStatusPriority) const
+	{
+		if (!Severity.IsSet() || !InOtherStatusPriority.Severity.IsSet())
+		{
+			return false;
+		}
+
+		return Severity.Get() == InOtherStatusPriority.Severity.Get() && SeverityPriority == InOtherStatusPriority.SeverityPriority;
+	}
+
+	bool operator<(const FAssetStatusPriority& InOtherStatusPriority) const
+	{
+		if (!Severity.IsSet())
+		{
+			return true;
+		}
+
+		if (!InOtherStatusPriority.Severity.IsSet())
+		{
+			return false;
+		}
+
+		if (Severity.Get() == InOtherStatusPriority.Severity.Get())
+		{
+			return SeverityPriority < InOtherStatusPriority.SeverityPriority;
+		}
+
+		return Severity.Get() < InOtherStatusPriority.Severity.Get();
+	}
+
+public:
+	TAttribute<EStatusSeverity> Severity;
+	int32 SeverityPriority;
+};
+
+struct FAssetDisplayInfo
+{
+public:
+	TAttribute<const FSlateBrush*> StatusIcon;
+	TAttribute<FText> StatusTitle;
+	TAttribute<FText> StatusDescription;
+	TAttribute<EVisibility> IsVisible;
+	TAttribute<FAssetStatusPriority> Priority;
+};
+#endif
diff --git a/Engine/Source/Editor/AssetDefinition/Public/AssetDefinitionStatusBase.h b/Engine/Source/Editor/AssetDefinition/Public/AssetDefinitionStatusBase.h
deleted file mode 100644
diff --git a/Engine/Source/Editor/AssetDefinition/Public/AssetStatusAssetDataInfoProvider.h b/Engine/Source/Editor/AssetDefinition/Public/AssetStatusAssetDataInfoProvider.h
--- a/Engine/Source/Editor/AssetDefinition/Public/AssetStatusAssetDataInfoProvider.h
+++ b/Engine/Source/Editor/AssetDefinition/Public/AssetStatusAssetDataInfoProvider.h
@@ -2,7 +2,8 @@
 
 #pragma once
 
-#include "AssetDefinitionStatusBase.h"
+#if UE_CONTENTBROWSER_NEW_STYLE
+#include "IAssetStatusInfoProvider.h"
 #include "AssetRegistry/AssetData.h"
 
 class ASSETDEFINITION_API FAssetStatusAssetDataInfoProvider : public IAssetStatusInfoProvider
@@ -19,3 +20,4 @@ public:
 private:
 	FAssetData AssetData;
 };
+#endif
diff --git a/Engine/Source/Editor/AssetDefinition/Public/IAssetStatusInfoProvider.h b/Engine/Source/Editor/AssetDefinition/Public/IAssetStatusInfoProvider.h
new file mode 100644
--- /dev/null
+++ b/Engine/Source/Editor/AssetDefinition/Public/IAssetStatusInfoProvider.h
@@ -0,0 +1,21 @@
+// Copyright Epic Games, Inc. All Rights Reserved.
+
+#pragma once
+
+#if UE_CONTENTBROWSER_NEW_STYLE
+#include "Containers/UnrealString.h"
+
+class UPackage;
+
+class IAssetStatusInfoProvider
+{
+public:
+	virtual ~IAssetStatusInfoProvider() = default;
+
+	/** Try to find the Package without loading it, if the package is not loaded it will return nullptr */
+	virtual UPackage* FindPackage() const = 0;
+
+	/** Try to get the filename, return empty otherwise */
+	virtual FString TryGetFilename() const = 0;
+};
+#endif
diff --git a/Engine/Source/Editor/AssetDefinition/Public/IAssetSystemInfoProvider.h b/Engine/Source/Editor/AssetDefinition/Public/IAssetSystemInfoProvider.h
new file mode 100644
--- /dev/null
+++ b/Engine/Source/Editor/AssetDefinition/Public/IAssetSystemInfoProvider.h
@@ -0,0 +1,18 @@
+// Copyright Epic Games, Inc. All Rights Reserved.
+
+#pragma once
+
+#if UE_CONTENTBROWSER_NEW_STYLE
+#include "Containers/Array.h"
+
+struct FAssetDisplayInfo;
+
+class IAssetSystemInfoProvider
+{
+public:
+	virtual ~IAssetSystemInfoProvider() = default;
+
+	/** Populate the OutAssetDisplayInfo with the information of the Asset like the Path/Size etc... */
+	virtual void PopulateAssetInfo(TArray<FAssetDisplayInfo>& OutAssetDisplayInfo) const = 0;
+};
+#endif
diff --git a/Engine/Source/Editor/ContentBrowser/Private/AssetSystemContentBrowserInfoProvider.cpp b/Engine/Source/Editor/ContentBrowser/Private/AssetSystemContentBrowserInfoProvider.cpp
new file mode 100644
--- /dev/null
+++ b/Engine/Source/Editor/ContentBrowser/Private/AssetSystemContentBrowserInfoProvider.cpp
@@ -0,0 +1,577 @@
+// Copyright Epic Games, Inc. All Rights Reserved.
+
+#include "AssetSystemContentBrowserInfoProvider.h"
+
+#if UE_CONTENTBROWSER_NEW_STYLE
+#include "ActorFolder.h"
+#include "ActorFolderDesc.h"
+#include "AssetDefinition.h"
+#include "AssetDefinitionAssetInfo.h"
+#include "AssetDefinitionRegistry.h"
+#include "AssetRegistry/AssetRegistryModule.h"
+#include "AssetViewTypes.h"
+#include "AutoReimport/AssetSourceFilenameCache.h"
+#include "CollectionManagerModule.h"
+#include "ContentBrowserDataSource.h"
+#include "ContentBrowserUtils.h"
+#include "ICollectionManager.h"
+#include "Misc/EngineBuildSettings.h"
+
+#define LOCTEXT_NAMESPACE "AssetSystemContentBrowserInfoProvider"
+
+FAssetSystemContentBrowserInfoProvider::FAssetSystemContentBrowserInfoProvider(const TSharedPtr<FAssetViewItem>& InAssetItem)
+	: AssetItem(InAssetItem)
+{
+	if (AssetItem.IsValid())
+	{
+		OnItemDataChangedCacheDisplayTagsDelegateHandle = AssetItem->OnItemDataChanged().AddRaw(this, &FAssetSystemContentBrowserInfoProvider::CacheDisplayTags);
+		OnItemDataChangedCacheDirtyExternalPackageDelegateHandle = AssetItem->OnItemDataChanged().AddRaw(this, &FAssetSystemContentBrowserInfoProvider::CacheDirtyExternalPackageInfo);
+	}
+	FAssetData AssetData;
+	AssetItem->GetItem().Legacy_TryGetAssetData(AssetData);
+	if (AssetData.IsValid())
+	{
+		if (const UAssetDefinition* AssetDefinition = UAssetDefinitionRegistry::Get()->GetAssetDefinitionForAsset(AssetData))
+		{
+			bShouldSaveExternalPackages = AssetDefinition->ShouldSaveExternalPackages();
+		}
+	}
+	CacheDisplayTags();
+}
+
+FAssetSystemContentBrowserInfoProvider::~FAssetSystemContentBrowserInfoProvider()
+{
+	if (AssetItem.IsValid())
+	{
+		AssetItem->OnItemDataChanged().Remove(OnItemDataChangedCacheDisplayTagsDelegateHandle);
+		AssetItem->OnItemDataChanged().Remove(OnItemDataChangedCacheDirtyExternalPackageDelegateHandle);
+		OnItemDataChangedCacheDisplayTagsDelegateHandle.Reset();
+		OnItemDataChangedCacheDirtyExternalPackageDelegateHandle.Reset();
+	}
+}
+
+void FAssetSystemContentBrowserInfoProvider::PopulateAssetInfo(TArray<FAssetDisplayInfo>& OutAssetDisplayInfo) const
+{
+	if (AssetItem->IsFile())
+	{
+		// The tooltip contains the name, class, path, asset registry tags and source control status
+		FText PublicStateText;
+		const FSlateBrush* PublicStateIcon = nullptr;
+
+		// Create a box to hold every line of info in the body of the tooltip
+		TSharedRef<SVerticalBox> InfoBox = SNew(SVerticalBox);
+
+		FAssetData ItemAssetData;
+		AssetItem->GetItem().Legacy_TryGetAssetData(ItemAssetData);
+
+		// TODO: Always use the virtual path?
+		FAssetDisplayInfo PathInfo = FAssetDisplayInfo();
+		PathInfo.StatusTitle = LOCTEXT("TileViewTooltipPath", "Path");
+		if (ItemAssetData.IsValid())
+		{
+			PathInfo.StatusDescription = FText::FromName(ItemAssetData.PackagePath);
+		}
+		else
+		{
+			PathInfo.StatusDescription = FText::FromName(AssetItem->GetItem().GetVirtualPath());
+		}
+		OutAssetDisplayInfo.Add(PathInfo);
+
+		if (ItemAssetData.IsValid() && ItemAssetData.PackageName != NAME_None)
+		{
+			const FString PackagePathWithinRoot = ContentBrowserUtils::GetPackagePathWithinRoot(ItemAssetData.PackageName.ToString());
+			int32 PackageNameLength = PackagePathWithinRoot.Len();
+
+			int32 MaxAssetPathLen = ContentBrowserUtils::GetMaxAssetPathLen();
+
+			// Asset Path Length Info
+			FAssetDisplayInfo AssetPathLengthInfo = FAssetDisplayInfo();
+			AssetPathLengthInfo.StatusTitle = LOCTEXT("TileViewTooltipAssetPathLengthKey", "Asset Filepath Length");
+			AssetPathLengthInfo.StatusDescription = FText::Format(LOCTEXT("TileViewTooltipAssetPathLengthValue", "{0} / {1}"), FText::AsNumber(PackageNameLength), FText::AsNumber(MaxAssetPathLen));
+			OutAssetDisplayInfo.Add(AssetPathLengthInfo);
+
+			int32 PackageNameLengthForCooking = ContentBrowserUtils::GetPackageLengthForCooking(ItemAssetData.PackageName.ToString(), FEngineBuildSettings::IsInternalBuild());
+
+			// Cook Path Length Info
+			int32 MaxCookPathLen = ContentBrowserUtils::GetMaxCookPathLen();
+			FAssetDisplayInfo CookPathLenInfo = FAssetDisplayInfo();
+			CookPathLenInfo.StatusTitle = LOCTEXT("TileViewTooltipPathLengthForCookingKey", "Cooking Filepath Length");
+			CookPathLenInfo.StatusDescription = FText::Format(LOCTEXT("TileViewTooltipPathLengthForCookingValue", "{0} / {1}"), FText::AsNumber(PackageNameLengthForCooking), FText::AsNumber(MaxCookPathLen));
+			OutAssetDisplayInfo.Add(CookPathLenInfo);
+
+			if (ItemAssetData.PackageFlags & PKG_NotExternallyReferenceable)
+			{
+				PublicStateText = LOCTEXT("PrivateAssetState", "Private");
+			}
+			else
+			{
+				PublicStateText = LOCTEXT("PublicAssetState", "Public");
+			}
+		}
+
+		if (!AssetItem->GetItem().CanEdit())
+		{
+			if(AssetItem->GetItem().CanView())
+			{
+				PublicStateText = LOCTEXT("ViewReadOnlyAssetState", "View / Read Only");
+				PublicStateIcon = FAppStyle::GetBrush("AssetEditor.ReadOnlyOpenable");
+
+			}
+			else
+			{
+				PublicStateText = LOCTEXT("ReadOnlyAssetState", "Read Only");
+				PublicStateIcon = FAppStyle::GetBrush("Icons.Lock");
+			}
+		}
+
+		if(!AssetItem->GetItem().IsSupported())
+		{
+			PublicStateText = LOCTEXT("UnsupportedAssetState", "Unsupported");
+		}
+
+		// Add tags
+		for (const FTagContentBrowserDisplayItem& DisplayTagItem : CachedDisplayTags)
+		{
+			FAssetDisplayInfo DisplayTagInfo = FAssetDisplayInfo();
+			DisplayTagInfo.StatusTitle = DisplayTagItem.DisplayKey;
+			DisplayTagInfo.StatusDescription = DisplayTagItem.DisplayValue;
+			OutAssetDisplayInfo.Add(DisplayTagInfo);
+		}
+
+		// Add asset source files
+		if (ItemAssetData.IsValid())
+		{
+			TOptional<FAssetImportInfo> ImportInfo = FAssetSourceFilenameCache::ExtractAssetImportInfo(ItemAssetData);
+			if (ImportInfo.IsSet())
+			{
+				for (const FAssetImportInfo::FSourceFile& File : ImportInfo->SourceFiles)
+				{
+					FText SourceLabel = LOCTEXT("TileViewTooltipSourceFile", "Source File");
+					if (File.DisplayLabelName.Len() > 0)
+					{
+						SourceLabel = FText::FromString(FText(LOCTEXT("TileViewTooltipSourceFile", "Source File")).ToString() + TEXT(" (") + File.DisplayLabelName + TEXT(")"));
+					}
+					FAssetDisplayInfo SourceFileInfo = FAssetDisplayInfo();
+					SourceFileInfo.StatusTitle = SourceLabel;
+					SourceFileInfo.StatusDescription = FText::FromString(File.RelativeFilename);
+					OutAssetDisplayInfo.Add(SourceFileInfo);
+				}
+			}
+		}
+
+		static const IConsoleVariable* EnablePublicAssetFeatureCVar = IConsoleManager::Get().FindConsoleVariable(TEXT("AssetTools.EnablePublicAssetFeature"));
+		const bool bIsPublicAssetUIEnabled = EnablePublicAssetFeatureCVar && EnablePublicAssetFeatureCVar->GetBool();
+
+		// Restriction Info
+		FAssetDisplayInfo RestrictionInfo = FAssetDisplayInfo();
+		RestrictionInfo.IsVisible = PublicStateIcon && bIsPublicAssetUIEnabled && !PublicStateText.IsEmpty() ? EVisibility::Visible : EVisibility::Collapsed;
+		RestrictionInfo.StatusIcon = PublicStateIcon;
+		RestrictionInfo.StatusTitle = LOCTEXT("Restriction", "Restriction");
+		RestrictionInfo.StatusDescription = PublicStateText;
+		OutAssetDisplayInfo.Add(RestrictionInfo);
+
+		// Unsupported Info
+		FAssetDisplayInfo UnsupportedInfo = FAssetDisplayInfo();
+		UnsupportedInfo.IsVisible = AssetItem->GetItem().IsSupported() ? EVisibility::Collapsed : EVisibility::Visible;;
+		UnsupportedInfo.StatusTitle = LOCTEXT("UnsupportedAssetTitleText", "Item is not supported");
+		UnsupportedInfo.StatusDescription = LOCTEXT("UnsupportedAssetDescriptionText", "This type of asset is not allowed in this project. Delete unsupported assets to avoid errors.");
+		OutAssetDisplayInfo.Add(UnsupportedInfo);
+
+		// External Package Info
+		FAssetDisplayInfo ExternalPackageInfo = FAssetDisplayInfo();
+		ExternalPackageInfo.IsVisible = bShouldSaveExternalPackages && !GetExternalPackagesText().IsEmpty() ? EVisibility::Visible : EVisibility::Collapsed;;
+		ExternalPackageInfo.StatusTitle = LOCTEXT("DirtyExternalPackages", "Modified external packages");
+		ExternalPackageInfo.StatusDescription = GetExternalPackagesText();
+		OutAssetDisplayInfo.Add(ExternalPackageInfo);
+
+		// User Description
+		FAssetDisplayInfo UserDescriptionInfo = FAssetDisplayInfo();
+		UserDescriptionInfo.IsVisible = GetAssetUserDescription().IsEmpty() ? EVisibility::Collapsed : EVisibility::Visible;;
+		UserDescriptionInfo.StatusTitle = LOCTEXT("UserDescriptionTitle", "User Description");
+		UserDescriptionInfo.StatusDescription = GetAssetUserDescription();
+		OutAssetDisplayInfo.Add(UserDescriptionInfo);
+
+		// Collection Pips
+		if (ItemAssetData.IsValid())
+		{
+			ICollectionManager& CollectionManager = FCollectionManagerModule::GetModule().Get();
+
+			TArray<FCollectionNameType> CollectionsContainingObject;
+			CollectionManager.GetCollectionsContainingObject(ItemAssetData.ToSoftObjectPath(), CollectionsContainingObject);
+
+			if (CollectionsContainingObject.Num() > 0)
+			{
+				for (const FCollectionNameType& CollectionContainingObject : CollectionsContainingObject)
+				{
+					FCollectionStatusInfo CollectionStatusInfo;
+					if (CollectionManager.GetCollectionStatusInfo(CollectionContainingObject.Name, CollectionContainingObject.Type, CollectionStatusInfo))
+					{
+						FAssetDisplayInfo CollectionInfo = FAssetDisplayInfo();
+						CollectionInfo.StatusTitle = FText::FromName(CollectionContainingObject.Name);
+						CollectionInfo.StatusDescription = FText::AsNumber(CollectionStatusInfo.NumObjects);
+						OutAssetDisplayInfo.Add(CollectionInfo);
+					}
+				}
+			}
+		}
+	}
+}
+
+void FAssetSystemContentBrowserInfoProvider::CacheDisplayTags()
+{
+	CachedDisplayTags.Reset();
+
+	const FContentBrowserItemDataAttributeValues AssetItemAttributes = AssetItem->GetItem().GetItemAttributes(/*bIncludeMetaData*/true);
+	
+	FAssetData ItemAssetData;
+	AssetItem->GetItem().Legacy_TryGetAssetData(ItemAssetData);
+
+	// Add all visible attributes
+	for (const auto& AssetItemAttributePair : AssetItemAttributes)
+	{
+		const FName AttributeName = AssetItemAttributePair.Key;
+		const FContentBrowserItemDataAttributeValue& AttributeValue = AssetItemAttributePair.Value;
+		const FContentBrowserItemDataAttributeMetaData& AttributeMetaData = AttributeValue.GetMetaData();
+
+		if (AttributeMetaData.AttributeType == UObject::FAssetRegistryTag::TT_Hidden)
+		{
+			continue;
+		}
+	
+		// Build the display value for this attribute
+		FText DisplayValue;
+		if (AttributeValue.GetValueType() == EContentBrowserItemDataAttributeValueType::Text)
+		{
+			DisplayValue = AttributeValue.GetValueText();
+		}
+		else
+		{
+			const FString AttributeValueStr = AttributeValue.GetValue<FString>();
+
+			auto ReformatNumberStringForDisplay = [](const FString& InNumberString) -> FText
+			{
+				// Respect the number of decimal places in the source string when converting for display
+				int32 NumDecimalPlaces = 0;
+				{
+					int32 DotIndex = INDEX_NONE;
+					if (InNumberString.FindChar(TEXT('.'), DotIndex))
+					{
+						NumDecimalPlaces = InNumberString.Len() - DotIndex - 1;
+					}
+				}
+	
+				if (NumDecimalPlaces > 0)
+				{
+					// Convert the number as a double
+					double Num = 0.0;
+					LexFromString(Num, *InNumberString);
+	
+					const FNumberFormattingOptions NumFormatOpts = FNumberFormattingOptions()
+						.SetMinimumFractionalDigits(NumDecimalPlaces)
+						.SetMaximumFractionalDigits(NumDecimalPlaces);
+	
+					return FText::AsNumber(Num, &NumFormatOpts);
+				}
+
+				const bool bIsSigned = InNumberString.Len() > 0 && (InNumberString[0] == TEXT('-') || InNumberString[0] == TEXT('+'));
+				if (bIsSigned)
+				{
+					// Convert the number as a signed int
+					int64 Num = 0;
+					LexFromString(Num, *InNumberString);
+					return FText::AsNumber(Num);
+				}
+
+				// Convert the number as an unsigned int
+				uint64 Num = 0;
+				LexFromString(Num, *InNumberString);
+				return FText::AsNumber(Num);
+			};
+	
+			bool bHasSetDisplayValue = false;
+	
+			// Numerical tags need to format the specified number based on the display flags
+			if (!bHasSetDisplayValue && AttributeMetaData.AttributeType == UObject::FAssetRegistryTag::TT_Numerical && AttributeValueStr.IsNumeric())
+			{
+				bHasSetDisplayValue = true;
+	
+				const bool bAsMemory = !!(AttributeMetaData.DisplayFlags & UObject::FAssetRegistryTag::TD_Memory);
+	
+				if (bAsMemory)
+				{
+					// Memory should be a 64-bit unsigned number of bytes
+					uint64 NumBytes = 0;
+					LexFromString(NumBytes, *AttributeValueStr);
+	
+					DisplayValue = FText::AsMemory(NumBytes);
+				}
+				else
+				{
+					DisplayValue = ReformatNumberStringForDisplay(AttributeValueStr);
+				}
+			}
+	
+			// Dimensional tags need to be split into their component numbers, with each component number re-formatted
+			if (!bHasSetDisplayValue && AttributeMetaData.AttributeType == UObject::FAssetRegistryTag::TT_Dimensional)
+			{
+				// Formats:
+				//   123         (1D)
+				//   123x234     (2D)
+				//   123x234*345 (2D array)
+				//   123x234x345 (3D)
+				int32 FirstXPos;
+				if (AttributeValueStr.FindChar(TEXT('x'), FirstXPos))
+				{
+					FString FirstPart = AttributeValueStr.Left(FirstXPos);
+					FString Remainder = AttributeValueStr.Mid(FirstXPos + 1);
+					int32 RemainderSeparatorPos;
+
+					if (Remainder.FindChar(TEXT('*'), RemainderSeparatorPos))
+					{
+						// AxB*C form (2D array)
+						FString SecondPart = Remainder.Left(RemainderSeparatorPos);
+						FString ThirdPart = Remainder.Mid(RemainderSeparatorPos + 1);
+
+						bHasSetDisplayValue = true;
+						DisplayValue = FText::Format(LOCTEXT("DisplayTag2xArrayFmt", "{0} \u00D7 {1} ({2} elements)"),
+							ReformatNumberStringForDisplay(FirstPart),
+							ReformatNumberStringForDisplay(SecondPart),
+							ReformatNumberStringForDisplay(ThirdPart));
+					}
+					else if (Remainder.FindChar(TEXT('x'), RemainderSeparatorPos))
+					{
+						// AxBxC form (3D)
+						FString SecondPart = Remainder.Left(RemainderSeparatorPos);
+						FString ThirdPart = Remainder.Mid(RemainderSeparatorPos + 1);
+
+						bHasSetDisplayValue = true;
+						DisplayValue = FText::Format(LOCTEXT("DisplayTag3xFmt", "{0} \u00D7 {1} \u00D7 {2}"),
+							ReformatNumberStringForDisplay(FirstPart),
+							ReformatNumberStringForDisplay(SecondPart),
+							ReformatNumberStringForDisplay(ThirdPart));
+					}
+					else
+					{
+						// 2D form by default
+						bHasSetDisplayValue = true;
+						DisplayValue = FText::Format(LOCTEXT("DisplayTag2xFmt", "{0} \u00D7 {1}"),
+							ReformatNumberStringForDisplay(FirstPart),
+							ReformatNumberStringForDisplay(Remainder));
+					}
+				}
+				else
+				{
+					// No separators, assume 1D
+					bHasSetDisplayValue = true;
+					DisplayValue = ReformatNumberStringForDisplay(AttributeValueStr);
+				}
+			}
+	
+			// Chronological tags need to format the specified timestamp based on the display flags
+			if (!bHasSetDisplayValue && AttributeMetaData.AttributeType == UObject::FAssetRegistryTag::TT_Chronological)
+			{
+				bHasSetDisplayValue = true;
+	
+				FDateTime Timestamp;
+				if (FDateTime::Parse(AttributeValueStr, Timestamp))
+				{
+					const bool bDisplayDate = !!(AttributeMetaData.DisplayFlags & UObject::FAssetRegistryTag::TD_Date);
+					const bool bDisplayTime = !!(AttributeMetaData.DisplayFlags & UObject::FAssetRegistryTag::TD_Time);
+					const FString TimeZone = (AttributeMetaData.DisplayFlags & UObject::FAssetRegistryTag::TD_InvariantTz) ? FText::GetInvariantTimeZone() : FString();
+	
+					if (bDisplayDate && bDisplayTime)
+					{
+						DisplayValue = FText::AsDateTime(Timestamp, EDateTimeStyle::Short, EDateTimeStyle::Short, TimeZone);
+					}
+					else if (bDisplayDate)
+					{
+						DisplayValue = FText::AsDate(Timestamp, EDateTimeStyle::Short, TimeZone);
+					}
+					else if (bDisplayTime)
+					{
+						DisplayValue = FText::AsTime(Timestamp, EDateTimeStyle::Short, TimeZone);
+					}
+				}
+			}
+	
+			// The tag value might be localized text, so we need to parse it for display
+			if (!bHasSetDisplayValue && FTextStringHelper::IsComplexText(*AttributeValueStr))
+			{
+				bHasSetDisplayValue = FTextStringHelper::ReadFromBuffer(*AttributeValueStr, DisplayValue) != nullptr;
+			}
+	
+			// Do our best to build something valid from the string value
+			if (!bHasSetDisplayValue)
+			{
+				bHasSetDisplayValue = true;
+	
+				// Since all we have at this point is a string, we can't be very smart here.
+				// We need to strip some noise off class paths in some cases, but can't load the asset to inspect its UPROPERTYs manually due to performance concerns.
+				FString ValueString = FPackageName::ExportTextPathToObjectPath(AttributeValueStr);
+	
+				const TCHAR StringToRemove[] = TEXT("/Script/");
+				if (ValueString.StartsWith(StringToRemove))
+				{
+					// Remove the class path for native classes, and also remove Engine. for engine classes
+					const int32 SizeOfPrefix = UE_ARRAY_COUNT(StringToRemove) - 1;
+					ValueString.MidInline(SizeOfPrefix, ValueString.Len() - SizeOfPrefix, EAllowShrinking::No);
+					ValueString.ReplaceInline(TEXT("Engine."), TEXT(""));
+				}
+	
+				if (ItemAssetData.IsValid())
+				{
+					if (const UClass* AssetClass = ItemAssetData.GetClass())
+					{
+						if (const FProperty* TagField = FindFProperty<FProperty>(AssetClass, AttributeName))
+						{
+							const FProperty* TagProp = nullptr;
+							const UEnum* TagEnum = nullptr;
+							if (const FByteProperty* ByteProp = CastField<FByteProperty>(TagField))
+							{
+								TagProp = ByteProp;
+								TagEnum = ByteProp->Enum;
+							}
+							else if (const FEnumProperty* EnumProp = CastField<FEnumProperty>(TagField))
+							{
+								TagProp = EnumProp;
+								TagEnum = EnumProp->GetEnum();
+							}
+
+							// Strip off enum prefixes if they exist
+							if (TagProp)
+							{
+								if (TagEnum)
+								{
+									const FString EnumPrefix = TagEnum->GenerateEnumPrefix();
+									if (EnumPrefix.Len() && ValueString.StartsWith(EnumPrefix))
+									{
+										ValueString.RightChopInline(EnumPrefix.Len() + 1, EAllowShrinking::No);	// +1 to skip over the underscore
+									}
+								}
+
+								ValueString = FName::NameToDisplayString(ValueString, false);
+							}
+						}
+					}
+				}
+	
+				DisplayValue = FText::AsCultureInvariant(MoveTemp(ValueString));
+			}
+			
+			// Add suffix to the value, if one is defined for this tag
+			if (!AttributeMetaData.Suffix.IsEmpty())
+			{
+				DisplayValue = FText::Format(LOCTEXT("DisplayTagSuffixFmt", "{0} {1}"), DisplayValue, AttributeMetaData.Suffix);
+			}
+		}
+	
+		if (!DisplayValue.IsEmpty())
+		{
+			CachedDisplayTags.Add(FTagContentBrowserDisplayItem(AttributeName, AttributeMetaData.DisplayName, DisplayValue, AttributeMetaData.bIsImportant));
+		}
+	}
+}
+
+void FAssetSystemContentBrowserInfoProvider::CacheDirtyExternalPackageInfo()
+{
+	if (bShouldSaveExternalPackages)
+	{
+		CachedDirtyExternalPackagesList.Empty();
+		
+		FAssetData AssetData;		
+		AssetItem->GetItem().Legacy_TryGetAssetData(AssetData);
+		if (AssetData.IsAssetLoaded())
+		{
+			if(const UObject* Asset = AssetData.GetAsset())
+			{
+				if (const UPackage* Package = Asset->GetPackage())
+				{
+					TArray<UPackage*> ExternalPackages = Package->GetExternalPackages();
+					IAssetRegistry& AssetRegistry = FModuleManager::GetModuleChecked<FAssetRegistryModule>(TEXT("AssetRegistry")).Get();
+
+					// Mirrored/copied from SSourceControlCommon.cpp
+					auto RetrieveAssetName = [](const FAssetData& InAssetData) -> FString
+					{
+						static const FName NAME_ActorLabel(TEXT("ActorLabel"));
+						if (InAssetData.FindTag(NAME_ActorLabel))
+						{
+							FString ResultAssetName;
+							InAssetData.GetTagValue(NAME_ActorLabel, ResultAssetName);
+							return ResultAssetName;
+						}
+
+						if (InAssetData.FindTag(FPrimaryAssetId::PrimaryAssetDisplayNameTag))
+						{
+							FString ResultAssetName;
+							InAssetData.GetTagValue(FPrimaryAssetId::PrimaryAssetDisplayNameTag, ResultAssetName);
+							return ResultAssetName;
+						}
+
+						if (InAssetData.AssetClassPath == UActorFolder::StaticClass()->GetClassPathName())
+						{
+							FString ActorFolderPath = UActorFolder::GetAssetRegistryInfoFromPackage(InAssetData.PackageName).GetDisplayName();
+							if (!ActorFolderPath.IsEmpty())
+							{
+								return ActorFolderPath;
+							}
+						}
+
+						return InAssetData.AssetName.ToString();
+					};
+					
+					for (const UPackage* ExternalPackage : ExternalPackages)
+					{
+						if (ExternalPackage->IsDirty())
+						{
+							TArray<FAssetData> DirtyAssetDataEntries;
+							AssetRegistry.GetAssetsByPackageName(*ExternalPackage->GetName(), DirtyAssetDataEntries);
+
+							if (CachedDirtyExternalPackagesList.Len())
+							{								
+								CachedDirtyExternalPackagesList.Append("\n");
+							}
+							
+							CachedDirtyExternalPackagesList.Append(ExternalPackage->GetPathName());
+							
+							for (const FAssetData& DirtyAssetData : DirtyAssetDataEntries)
+							{
+								const FString AssetName = RetrieveAssetName(DirtyAssetData);
+								const FString AssetClass = DirtyAssetData.AssetClassPath.GetAssetName().ToString();
+						
+								CachedDirtyExternalPackagesList.Append("\n\t");
+								CachedDirtyExternalPackagesList.Append(FString::Printf(TEXT("%s (%s)"), *AssetName, *AssetClass)); 
+							}
+						}
+					}
+				}
+			}
+		}
+	}
+}
+
+FText FAssetSystemContentBrowserInfoProvider::GetExternalPackagesText() const
+{
+	if (CachedDirtyExternalPackagesList.Len())
+	{
+		return FText::FromString(CachedDirtyExternalPackagesList);
+	}
+
+	return FText::GetEmpty();
+}
+
+FText FAssetSystemContentBrowserInfoProvider::GetAssetUserDescription() const
+{
+	if (AssetItem && AssetItem->IsFile())
+	{
+		FContentBrowserItemDataAttributeValue DescriptionAttributeValue = AssetItem->GetItem().GetItemAttribute(ContentBrowserItemAttributes::ItemDescription);
+		if (DescriptionAttributeValue.IsValid())
+		{
+			return DescriptionAttributeValue.GetValue<FText>();
+		}
+	}
+
+	return FText::GetEmpty();
+}
+
+#undef LOCTEXT_NAMESPACE
+#endif
diff --git a/Engine/Source/Editor/ContentBrowser/Private/AssetSystemContentBrowserInfoProvider.h b/Engine/Source/Editor/ContentBrowser/Private/AssetSystemContentBrowserInfoProvider.h
new file mode 100644
--- /dev/null
+++ b/Engine/Source/Editor/ContentBrowser/Private/AssetSystemContentBrowserInfoProvider.h
@@ -0,0 +1,74 @@
+// Copyright Epic Games, Inc. All Rights Reserved.
+
+#pragma once
+
+#if UE_CONTENTBROWSER_NEW_STYLE
+#include "Containers/Array.h"
+#include "Containers/UnrealString.h"
+#include "Delegates/IDelegateInstance.h"
+#include "IAssetSystemInfoProvider.h"
+#include "Internationalization/Text.h"
+#include "UObject/NameTypes.h"
+#include "Templates/SharedPointer.h"
+
+class FAssetViewItem;
+
+class FAssetSystemContentBrowserInfoProvider : public IAssetSystemInfoProvider
+{
+public:
+	FAssetSystemContentBrowserInfoProvider() = delete;
+	FAssetSystemContentBrowserInfoProvider(const TSharedPtr<FAssetViewItem>& InAssetItem);
+	virtual ~FAssetSystemContentBrowserInfoProvider() override;
+
+	virtual void PopulateAssetInfo(TArray<FAssetDisplayInfo>& OutAssetDisplayInfo) const override;
+
+private:
+	/** Cache the Display Tags for this asset, called everytime the ItemAssetData change*/
+	void CacheDisplayTags();
+
+	/** Cache the External Package Info if any, called everytime the ItemAssetData change */
+	void CacheDirtyExternalPackageInfo();
+
+	/** Get the CachedDirtyExternalPackage Text */
+	FText GetExternalPackagesText() const;
+
+	/** Get the Item Description for this item */
+	FText GetAssetUserDescription() const;
+
+private:
+	/** Data for a cached display tag for this item (used in the tooltip, and also as the display string in column views) */
+	struct FTagContentBrowserDisplayItem
+	{
+		FTagContentBrowserDisplayItem(FName InTagKey, FText InDisplayKey, FText InDisplayValue, const bool InImportant)
+			: TagKey(InTagKey)
+			, DisplayKey(MoveTemp(InDisplayKey))
+			, DisplayValue(MoveTemp(InDisplayValue))
+			, bImportant(InImportant)
+		{
+		}
+
+		FName TagKey;
+		FText DisplayKey;
+		FText DisplayValue;
+		bool bImportant;
+	};
+
+	/** The cached display tags for this item */
+	TArray<FTagContentBrowserDisplayItem> CachedDisplayTags;
+
+	/** Whether it should save external package */
+	bool bShouldSaveExternalPackages = false;
+
+	/** The cached external package to save */
+	FString CachedDirtyExternalPackagesList;
+
+	/** Delegate handle of the CacheDisplayTags */
+	FDelegateHandle OnItemDataChangedCacheDisplayTagsDelegateHandle;
+
+	/** Delegate handle of the CacheDirtyPackage */
+	FDelegateHandle OnItemDataChangedCacheDirtyExternalPackageDelegateHandle;
+
+	/** AssetViewItem of this provider */
+	TSharedPtr<FAssetViewItem> AssetItem;
+};
+#endif
diff --git a/Engine/Source/Editor/ContentBrowser/Private/AssetViewWidgets.cpp b/Engine/Source/Editor/ContentBrowser/Private/AssetViewWidgets.cpp
--- a/Engine/Source/Editor/ContentBrowser/Private/AssetViewWidgets.cpp
+++ b/Engine/Source/Editor/ContentBrowser/Private/AssetViewWidgets.cpp
@@ -6,6 +6,7 @@
 #include "ActorFolder.h"
 #include "ActorFolderDesc.h"
 #include "AssetDefinitionRegistry.h"
+#include "AssetSystemContentBrowserInfoProvider.h"
 #include "AssetRegistry/AssetData.h"
 #include "AssetTagItemTypes.h"
 #include "AssetThumbnail.h"
@@ -77,6 +78,7 @@
 #include "Widgets/Images/SLayeredImage.h"
 #include "Widgets/Layout/SBorder.h"
 #include "Widgets/Layout/SBox.h"
+#include "Widgets/Layout/SSeparator.h"
 #include "Widgets/Layout/SWrapBox.h"
 #include "Widgets/SBoxPanel.h"
 #include "Widgets/SNullWidget.h"
@@ -296,8 +298,13 @@ public:
 
 		SToolTip::Construct(
 			SToolTip::FArguments()
+#if UE_CONTENTBROWSER_NEW_STYLE
+			.TextMargin(FMargin(12.f, 8.f, 12.f, 2.f))
+			.BorderImage(FAppStyle::GetBrush("ContentBrowser.Tooltip.Border"))
+#else
 			.TextMargin(1.0f)
 			.BorderImage(UE::ContentBrowser::Private::FContentBrowserStyle::Get().GetBrush("ContentBrowser.TileViewTooltip.ToolTipBorder"))
+#endif
 			);
 	}
 
@@ -1130,7 +1137,58 @@ TSharedRef<SWidget> SAssetViewItem::CreateToolTipWidget() const
 					}
 				}
 			}
+#if UE_CONTENTBROWSER_NEW_STYLE
+			return SNew(SVerticalBox)
+
+				+ SVerticalBox::Slot()
+				.AutoHeight()
+				[
+					SNew(SVerticalBox)
+
+					+ SVerticalBox::Slot()
+					.AutoHeight()
+					[
+						SNew(SHorizontalBox)
 
+						+ SHorizontalBox::Slot()
+						.AutoWidth()
+						.VAlign(VAlign_Center)
+						[
+							SNew(STextBlock)
+							.Text(FolderName)
+							.ColorAndOpacity(FStyleColors::White)
+							.Font(FAppStyle::GetFontStyle("ContentBrowser.Tooltip.EntryFont"))
+						]
+
+						+ SHorizontalBox::Slot()
+						.AutoWidth()
+						.VAlign(VAlign_Center)
+						[
+							SNew(STextBlock)
+							.Font(FAppStyle::GetFontStyle("ContentBrowser.Tooltip.EntryFont"))
+							.ColorAndOpacity(FStyleColors::White)
+							.Text(LOCTEXT("FolderNameBracketed", " (Folder)"))
+						]
+					]
+				]
+
+				+ SVerticalBox::Slot()
+				.Padding(FMargin(0.f,6.f))
+				.AutoHeight()
+				[
+					SNew(SSeparator)
+					.Orientation(Orient_Horizontal)
+					.Thickness(1.f)
+					.ColorAndOpacity(COLOR("#484848FF"))
+					.SeparatorImage(FAppStyle::Get().GetBrush("WhiteBrush"))
+				]
+
+				+ SVerticalBox::Slot()
+				.AutoHeight()
+				[
+					InfoBox
+				];
+#else
 			return SNew(SBorder)
 				.Padding(6)
 				.BorderImage( UE::ContentBrowser::Private::FContentBrowserStyle::Get().GetBrush("ContentBrowser.TileViewTooltip.NonContentBorder") )
@@ -1184,6 +1242,7 @@ TSharedRef<SWidget> SAssetViewItem::CreateToolTipWidget() const
 						]
 					]
 				];
+#endif
 		}
 	}
 	else
@@ -1239,7 +1298,34 @@ FText SAssetViewItem::GetAssetUserDescription() const
 
 	return FText::GetEmpty();
 }
+#if UE_CONTENTBROWSER_NEW_STYLE
+void SAssetViewItem::AddToToolTipInfoBox(const TSharedRef<SVerticalBox>& InfoBox, const FText& Key, const FText& Value, bool bImportant) const
+{
+	InfoBox->AddSlot()
+	.Padding(0.f, 0.f, 0.f, 6.f)
+	.AutoHeight()
+	[
+		SNew(SHorizontalBox)
+		+SHorizontalBox::Slot()
+		.AutoWidth()
+		.Padding(0, 0, 4, 0)
+		[
+			SNew(STextBlock)
+			.Font(FAppStyle::GetFontStyle("ContentBrowser.Tooltip.EntryFont"))
+			.Text(FText::Format(NSLOCTEXT("AssetThumbnailToolTip", "AssetViewTooltipFormat", "{0}:"), Key))
+		]
 
+		+SHorizontalBox::Slot()
+		.AutoWidth()
+		[
+			SNew(STextBlock)
+			.Font(FAppStyle::GetFontStyle("ContentBrowser.Tooltip.EntryFont"))
+			.ColorAndOpacity(FStyleColors::White)
+			.Text(Value)
+		]
+	];
+}
+#else
 void SAssetViewItem::AddToToolTipInfoBox(const TSharedRef<SVerticalBox>& InfoBox, const FText& Key, const FText& Value, bool bImportant) const
 {
 	FWidgetStyle ImportantStyle;
@@ -1269,6 +1355,7 @@ void SAssetViewItem::AddToToolTipInfoBox(const TSharedRef<SVerticalBox>& InfoBox
 		]
 	];
 }
+#endif
 
 void SAssetViewItem::UpdateDirtyState()
 {
@@ -1953,6 +2040,7 @@ void SAssetTileItem::Construct( const FArguments& InArgs )
 		ThumbnailConfig.Padding= FMargin(2.0f);
 		ThumbnailConfig.GenericThumbnailSize = MakeAttributeSP(this, &SAssetTileItem::GetGenericThumbnailSize);
 #if UE_CONTENTBROWSER_NEW_STYLE
+		ThumbnailConfig.AssetSystemInfoProvider = MakeShared<FAssetSystemContentBrowserInfoProvider>(AssetItem);
 		ThumbnailConfig.bAllowAssetStatusThumbnailOverlay = true;
 		ThumbnailConfig.bShowAssetChip = true;
 		ThumbnailConfig.AssetChipBorderImageOverride = TAttribute<const FSlateBrush*>::CreateSP(this, &SAssetTileItem::GetAssetAreaOverlayBackgroundImage);
@@ -1978,6 +2066,13 @@ void SAssetTileItem::Construct( const FArguments& InArgs )
 		}
 
 		Thumbnail = AssetThumbnail->MakeThumbnailWidget(ThumbnailConfig);
+#if UE_CONTENTBROWSER_NEW_STYLE
+		// Use the same tooltip as the Thumbnail
+		if (const TSharedPtr<IToolTip>& ThumbnailTooltip = Thumbnail->GetToolTip())
+		{
+			SetToolTip(ThumbnailTooltip);
+		}
+#endif
 	}
 	else
 	{
diff --git a/Engine/Source/Editor/Documentation/Private/SDocumentationToolTip.cpp b/Engine/Source/Editor/Documentation/Private/SDocumentationToolTip.cpp
--- a/Engine/Source/Editor/Documentation/Private/SDocumentationToolTip.cpp
+++ b/Engine/Source/Editor/Documentation/Private/SDocumentationToolTip.cpp
@@ -41,6 +41,7 @@ void SDocumentationToolTip::Construct( const FArguments& InArgs )
 	DocumentationMargin = InArgs._DocumentationMargin;
 	IsDisplayingDocumentationLink = false;
 	Shortcut = InArgs._Shortcut;
+	OverrideFullTooltipContent = InArgs._OverrideExtendedToolTipContent;
 
 	ExcerptName = InArgs._ExcerptName;
 	IsShowingFullTip = false;
@@ -54,6 +55,12 @@ void SDocumentationToolTip::Construct( const FArguments& InArgs )
 
 	SAssignNew(DocumentationControlBox, SHorizontalBox);
 	SAssignNew(FullTipContent, SBox);
+	if (OverrideFullTooltipContent.IsValid())
+	{
+		FullTipContent->SetContent(OverrideFullTooltipContent.ToSharedRef());
+		FullTipContent->SetVisibility(TAttribute<EVisibility>::CreateSP(this, &SDocumentationToolTip::GetOverriddenFullToolTipVisibility));
+	}
+
 	ConstructSimpleTipContent();
 
 	ChildSlot
@@ -110,7 +117,7 @@ void SDocumentationToolTip::Construct( const FArguments& InArgs )
 	TransitionLength = 0.2f;
 	LastDesiredSize = SimpleTipContent->GetDesiredSize();
 	TransitionStartSize = SimpleTipContent->GetDesiredSize();
-	bFullTipContentIsReady = false;
+	bFullTipContentIsReady = OverrideFullTooltipContent.IsValid();
 }
 
 void SDocumentationToolTip::ConstructSimpleTipContent()
@@ -215,7 +222,7 @@ void SDocumentationToolTip::ConstructSimpleTipContent()
 
 void SDocumentationToolTip::AddDocumentation(TSharedPtr< SVerticalBox > VerticalBox)
 {
-	if ( !DocumentationLink.IsEmpty() )
+	if ( !DocumentationLink.IsEmpty() && !OverrideFullTooltipContent.IsValid())
 	{
 		if ( !DocumentationPage.IsValid() )
 		{
@@ -538,7 +545,7 @@ void SDocumentationToolTip::Tick( const FGeometry& AllottedGeometry, const doubl
 	LastDesiredSize = WidgetContent->GetDesiredSize();
 
 	const FModifierKeysState ModifierKeys = FSlateApplication::Get().GetModifierKeys();
-	const bool NeedsUpdate = IsDisplayingDocumentationLink != GetDefault<UEditorPerProjectUserSettings>()->bDisplayDocumentationLink;
+	const bool NeedsUpdate = !OverrideFullTooltipContent.IsValid() && IsDisplayingDocumentationLink != GetDefault<UEditorPerProjectUserSettings>()->bDisplayDocumentationLink;
 	if (TransitionStartTime > 0)
 	{
 		TransitionPercentage = (InCurrentTime - TransitionStartTime) / TransitionLength;
@@ -557,25 +564,31 @@ void SDocumentationToolTip::Tick( const FGeometry& AllottedGeometry, const doubl
 
 	if ( !IsShowingFullTip && ModifierKeys.IsAltDown() && ModifierKeys.IsControlDown() )
 	{
-		if ( !bFullTipContentIsReady && DocumentationPage.IsValid() && DocumentationPage->HasExcerpt(ExcerptName))
+		if (!OverrideFullTooltipContent.IsValid())
 		{
-			ConstructFullTipContent();
-		}
-		else if ( GetDefault<UEditorPerProjectUserSettings>()->bDisplayDocumentationLink )
-		{
-			ReloadDocumentation();
+			if ( !bFullTipContentIsReady && DocumentationPage.IsValid() && DocumentationPage->HasExcerpt(ExcerptName))
+			{
+				ConstructFullTipContent();
+			}
+			else if ( GetDefault<UEditorPerProjectUserSettings>()->bDisplayDocumentationLink )
+			{
+				ReloadDocumentation();
+			}
 		}
 
 		if ( bFullTipContentIsReady)
 		{
-			// Analytics event
-			if (FEngineAnalytics::IsAvailable())
+			if (!OverrideFullTooltipContent.IsValid())
 			{
-				TArray<FAnalyticsEventAttribute> Params;
-				Params.Add(FAnalyticsEventAttribute(TEXT("Page"), DocumentationLink));
-				Params.Add(FAnalyticsEventAttribute(TEXT("Excerpt"), ExcerptName));
+				// Analytics event
+				if (FEngineAnalytics::IsAvailable())
+				{
+					TArray<FAnalyticsEventAttribute> Params;
+					Params.Add(FAnalyticsEventAttribute(TEXT("Page"), DocumentationLink));
+					Params.Add(FAnalyticsEventAttribute(TEXT("Excerpt"), ExcerptName));
 
-				FEngineAnalytics::GetProvider().RecordEvent(TEXT("Editor.Usage.Documentation.FullTooltipShown"), Params);
+					FEngineAnalytics::GetProvider().RecordEvent(TEXT("Editor.Usage.Documentation.FullTooltipShown"), Params);
+				}
 			}
 			bIsInTransition = true;
 			IsShowingFullTip = true;
@@ -601,7 +614,7 @@ void SDocumentationToolTip::Tick( const FGeometry& AllottedGeometry, const doubl
 bool SDocumentationToolTip::IsInteractive() const
 {
 	const FModifierKeysState ModifierKeys = FSlateApplication::Get().GetModifierKeys();
-	return ( DocumentationPage.IsValid() && ModifierKeys.IsAltDown() && ModifierKeys.IsControlDown() );
+	return ((OverrideFullTooltipContent.IsValid() || DocumentationPage.IsValid()) && ModifierKeys.IsAltDown() && ModifierKeys.IsControlDown() );
 }
 
 FVector2D SDocumentationToolTip::ComputeDesiredSize(float LayoutScaleMultiplier) const
@@ -609,6 +622,11 @@ FVector2D SDocumentationToolTip::ComputeDesiredSize(float LayoutScaleMultiplier)
 	return LastDesiredSize;
 }
 
+EVisibility SDocumentationToolTip::GetOverriddenFullToolTipVisibility() const
+{
+	return IsShowingFullTip ? EVisibility::Visible : EVisibility::Collapsed;
+}
+
 EVisibility SDocumentationToolTip::GetFullTipVisibility() const
 {
 	if (IsShowingFullTip)
@@ -647,7 +665,7 @@ EVisibility SDocumentationToolTip::GetShortcutVisibility() const
 
 const FSlateBrush* SDocumentationToolTip::GetSimpleTipBorderStyle() const
 {
-	if (IsShowingFullTip)
+	if (IsShowingFullTip && !OverrideContent.IsValid())
 	{
 		return FDocumentationStyleSet::Get().GetBrush("ToolTip.Header");
 	}
diff --git a/Engine/Source/Editor/Documentation/Public/SDocumentationToolTip.h b/Engine/Source/Editor/Documentation/Public/SDocumentationToolTip.h
--- a/Engine/Source/Editor/Documentation/Public/SDocumentationToolTip.h
+++ b/Engine/Source/Editor/Documentation/Public/SDocumentationToolTip.h
@@ -38,6 +38,7 @@ public:
 		, _ColorAndOpacity( FLinearColor::Black )
 		, _AddDocumentation( true )
 		, _DocumentationMargin(0)
+		, _OverrideExtendedToolTipContent(nullptr)
 		, _Content()
 		{}
 
@@ -73,6 +74,9 @@ public:
 		/**  */
 		SLATE_ARGUMENT( FString, ExcerptName )
 
+		/** Extended ToolTip to use when Ctrl + Alt is pressed instead of loading it from the Documentation */
+		SLATE_ARGUMENT(TSharedPtr<SWidget>, OverrideExtendedToolTipContent)
+
 		/** Arbitrary content to be displayed in the tool tip; overrides any text that may be set. */
 		SLATE_DEFAULT_SLOT( FArguments, Content )
 
@@ -106,6 +110,8 @@ public:
 
 private:
 
+	EVisibility GetOverriddenFullToolTipVisibility() const;
+
 	void ConstructSimpleTipContent();
 
 	void ConstructFullTipContent();
@@ -126,6 +132,7 @@ private:
 	TAttribute< FText > TextContent;
 	TAttribute< FText > Shortcut;
 	TSharedPtr< SWidget > OverrideContent;
+	TSharedPtr< SWidget > OverrideFullTooltipContent;
 	FTextBlockStyle StyleInfo;
 	FTextBlockStyle SubduedStyleInfo;
 	FTextBlockStyle HyperlinkTextStyleInfo;
diff --git a/Engine/Source/Editor/EditorStyle/Private/StarshipStyle.cpp b/Engine/Source/Editor/EditorStyle/Private/StarshipStyle.cpp
--- a/Engine/Source/Editor/EditorStyle/Private/StarshipStyle.cpp
+++ b/Engine/Source/Editor/EditorStyle/Private/StarshipStyle.cpp
@@ -6342,6 +6342,9 @@ void FStarshipEditorStyle::FStyle::SetupContentBrowserStyle()
 		Set( "ContentBrowser.AssetTileViewNameFontVerySmall", DEFAULT_FONT( "VeryLight", 7 ) );
 		Set( "ContentBrowser.AssetTileViewNameFontDirty", FStyleFonts::Get().SmallBold);
 
+		Set("ContentBrowser.Tooltip.Border", new FSlateRoundedBoxBrush(FStyleColors::Secondary, 0.f, COLOR("#484848FF"), 2.f));
+		Set("ContentBrowser.ToolTip.CommandBorder", new FSlateRoundedBoxBrush(FStyleColors::Transparent, 4.f, FStyleColors::White, 1.f));
+
 		Set("ContentBrowser.AssetListView.ColumnListTableRow", FTableRowStyle()
 			.SetEvenRowBackgroundBrush(FSlateColorBrush(FStyleColors::Recessed))
 			.SetEvenRowBackgroundHoveredBrush(FSlateColorBrush(FStyleColors::SelectHover))
diff --git a/Engine/Source/Editor/UnrealEd/Private/AssetDefinitionDefault.cpp b/Engine/Source/Editor/UnrealEd/Private/AssetDefinitionDefault.cpp
--- a/Engine/Source/Editor/UnrealEd/Private/AssetDefinitionDefault.cpp
+++ b/Engine/Source/Editor/UnrealEd/Private/AssetDefinitionDefault.cpp
@@ -2,7 +2,9 @@
 
 #include "AssetDefinitionDefault.h"
 
+#include "AssetDefinitionAssetInfo.h"
 #include "AssetToolsModule.h"
+#include "IAssetStatusInfoProvider.h"
 #include "EditorFramework/AssetImportData.h"
 #include "EditorFramework/ThumbnailInfo.h"
 #include "ISourceControlModule.h"
@@ -12,6 +14,7 @@
 
 #define LOCTEXT_NAMESPACE "AssetDefinitionDefault"
 
+#if UE_CONTENTBROWSER_NEW_STYLE
 namespace UE::AssetDefinitionDefault::Status
 {
 	EVisibility GetDirtyStatusVisibility(const TSharedPtr<IAssetStatusInfoProvider> InAssetStatusInfoProvider)
@@ -98,6 +101,7 @@ namespace UE::AssetDefinitionDefault::Status
 		return SourceControlDescription;
 	}
 }
+#endif
 
 EAssetCommandResult UAssetDefinitionDefault::OpenAssets(const FAssetOpenArgs& OpenArgs) const
 {
@@ -130,21 +134,21 @@ EAssetCommandResult UAssetDefinitionDefault::PerformAssetDiff(const FAssetDiffAr
 }
 
 #if UE_CONTENTBROWSER_NEW_STYLE
-void UAssetDefinitionDefault::GetAssetStatusInfo(const TSharedPtr<IAssetStatusInfoProvider>& InAssetStatusInfoProvider, FAssetStatusInfo& OutStatusInfo) const
+void UAssetDefinitionDefault::GetAssetStatusInfo(const TSharedPtr<IAssetStatusInfoProvider>& InAssetStatusInfoProvider, TArray<FAssetDisplayInfo>& OutStatusInfo) const
 {
-	FAssetStatus DirtyStatus;
+	FAssetDisplayInfo DirtyStatus;
 	DirtyStatus.StatusIcon = FAppStyle::GetBrush("ContentBrowser.ContentDirty");
 	DirtyStatus.Priority = FAssetStatusPriority(EStatusSeverity::Info, 1);
 	DirtyStatus.StatusDescription = LOCTEXT("DirtyAssetTooltip", "Asset has unsaved changes");
 	DirtyStatus.IsVisible = TAttribute<EVisibility>::Create(TAttribute<EVisibility>::FGetter::CreateStatic(&UE::AssetDefinitionDefault::Status::GetDirtyStatusVisibility, InAssetStatusInfoProvider));
-	OutStatusInfo.AssetStatus.Add(DirtyStatus);
+	OutStatusInfo.Add(DirtyStatus);
 
-	FAssetStatus SCCStatus;
+	FAssetDisplayInfo SCCStatus;
 	SCCStatus.Priority = FAssetStatusPriority(EStatusSeverity::Info, 0);
 	SCCStatus.StatusIcon = TAttribute<const FSlateBrush*>::Create(TAttribute<const FSlateBrush*>::FGetter::CreateStatic(&UE::AssetDefinitionDefault::Status::GetSourceControlStatusBrush, InAssetStatusInfoProvider));
 	SCCStatus.IsVisible = TAttribute<EVisibility>::Create(TAttribute<EVisibility>::FGetter::CreateStatic(&UE::AssetDefinitionDefault::Status::GetSourceControlStatusVisibility, InAssetStatusInfoProvider));
 	SCCStatus.StatusDescription = TAttribute<FText>::Create(TAttribute<FText>::FGetter::CreateStatic(&UE::AssetDefinitionDefault::Status::GetSourceControlStatusDescription, InAssetStatusInfoProvider));
-	OutStatusInfo.AssetStatus.Add(SCCStatus);
+	OutStatusInfo.Add(SCCStatus);
 }
 #endif
 
diff --git a/Engine/Source/Editor/UnrealEd/Private/AssetThumbnail.cpp b/Engine/Source/Editor/UnrealEd/Private/AssetThumbnail.cpp
--- a/Engine/Source/Editor/UnrealEd/Private/AssetThumbnail.cpp
+++ b/Engine/Source/Editor/UnrealEd/Private/AssetThumbnail.cpp
@@ -1,6 +1,7 @@
 // Copyright Epic Games, Inc. All Rights Reserved.
 
 #include "AssetThumbnail.h"
+#include "AssetThumbnailToolTip.h"
 #include "AssetDefinitionRegistry.h"
 #include "AssetStatusAssetDataInfoProvider.h"
 #include "Engine/Blueprint.h"
@@ -36,14 +37,21 @@
 #include "ShaderCompiler.h"
 #include "AssetCompilingManager.h"
 #include "AssetDefinitionRegistry.h"
+#include "AssetDefinitionAssetInfo.h"
+#include "AssetDefinitionDefault.h"
 #include "IAssetTools.h"
 #include "AssetTypeActions_Base.h"
 #include "AssetToolsModule.h"
 #include "Styling/SlateIconFinder.h"
 #include "ClassIconFinder.h"
+#include "IAssetSystemInfoProvider.h"
 #include "IVREditorModule.h"
+#include "SDocumentationToolTip.h"
+#include "Fonts/FontMeasure.h"
 #include "Framework/Application/SlateApplication.h"
 #include "Widgets/Images/SLayeredImage.h"
+#include "Widgets/Layout/SSeparator.h"
+#include "Widgets/Text/SRichTextBlock.h"
 
 namespace AssetThumbnailPool
 {
@@ -107,7 +115,9 @@ public:
 		SLATE_ARGUMENT(FMargin, Padding)
 		SLATE_ATTRIBUTE(int32, GenericThumbnailSize)
 #if UE_CONTENTBROWSER_NEW_STYLE
+		SLATE_ARGUMENT(TSharedPtr<IAssetSystemInfoProvider>, AssetSystemInfoProvider)
 		SLATE_ARGUMENT(bool, AllowAssetStatusThumbnailOverlay)
+		SLATE_ATTRIBUTE(TSharedPtr<IToolTip>, TooltipOverride)
 		SLATE_ARGUMENT(bool, ShowAssetChip)
 		SLATE_ARGUMENT(TAttribute<const FSlateBrush*>, AssetChipBorderImageOverride)
 #endif
@@ -133,6 +143,7 @@ public:
 		GenericThumbnailSize = InArgs._GenericThumbnailSize;
 		ColorStripOrientation = InArgs._ColorStripOrientation;
 #if UE_CONTENTBROWSER_NEW_STYLE
+		AssetSystemInfoProvider = InArgs._AssetSystemInfoProvider;
 		bShowAssetChip = InArgs._ShowAssetChip;
 		AssetChipBorderImageOverride = InArgs._AssetChipBorderImageOverride;
 #endif
@@ -161,6 +172,18 @@ public:
 
 		TSharedRef<SOverlay> OverlayWidget = SNew(SOverlay);
 
+#if UE_CONTENTBROWSER_NEW_STYLE
+		if (InArgs._TooltipOverride.IsSet())
+		{
+			SetToolTip(InArgs._TooltipOverride);
+		}
+		else
+		{
+			// Set our tooltip - this will refresh each time it's opened to make sure it's up-to-date
+			SetToolTip(SNew(SAssetThumbnailToolTip).AssetThumbnail(SharedThis(this)));
+		}
+#endif
+
 		UpdateThumbnailClass(AssetTypeActions.Get());
 
 		ClassThumbnailBrushOverride = InArgs._ClassThumbnailBrushOverride;
@@ -336,12 +359,17 @@ public:
 #if UE_CONTENTBROWSER_NEW_STYLE
 		const UAssetDefinition* AssetDefinition = UAssetDefinitionRegistry::Get()->GetAssetDefinitionForClass(Class);
 
+		if (!AssetDefinition)
+		{
+			AssetDefinition = GetDefault<UAssetDefinitionDefault>();
+		}
+
 		if (InArgs._AllowAssetStatusThumbnailOverlay && AssetDefinition)
 		{
 			const TSharedPtr<FAssetStatusAssetDataInfoProvider> AssetDataInfoProvider = MakeShared<FAssetStatusAssetDataInfoProvider>(AssetData);
 			AssetDefinition->GetAssetStatusInfo(AssetDataInfoProvider, OverlayInfo);
 			// Sort the list based on Status Priority
-			auto SortStatus = [] (const FAssetStatus& InFirstStatus, const FAssetStatus& InSecondStatus)
+			auto SortStatus = [] (const FAssetDisplayInfo& InFirstStatus, const FAssetDisplayInfo& InSecondStatus)
 			{
 				if (!InFirstStatus.Priority.IsSet())
 				{
@@ -355,12 +383,12 @@ public:
 
 				return InSecondStatus.Priority.Get() < InFirstStatus.Priority.Get();
 			};
-			OverlayInfo.AssetStatus.Sort(SortStatus);
+			OverlayInfo.Sort(SortStatus);
 
 			const TSharedRef<SHorizontalBox> HorizontalBox = SNew(SHorizontalBox);
-			for (int32 StatusIndex = 0; StatusIndex < OverlayInfo.AssetStatus.Num(); StatusIndex++)
+			for (int32 StatusIndex = 0; StatusIndex < OverlayInfo.Num(); StatusIndex++)
 			{
-				Statuses.Add(CreateStatusWidget(StatusIndex, OverlayInfo.AssetStatus[StatusIndex]));
+				Statuses.Add(CreateStatusWidget(StatusIndex, OverlayInfo[StatusIndex]));
 				HorizontalBox->AddSlot()
 				[
 					Statuses[StatusIndex].ToSharedRef()
@@ -553,11 +581,269 @@ public:
 		}
 	}
 
+#if UE_CONTENTBROWSER_NEW_STYLE
+	TSharedRef<SWidget> GetDefaultTooltip() const
+	{
+		const FAssetData& AssetData = AssetThumbnail->GetAssetData();
+
+		const UClass* Class = FindObjectSafe<UClass>(AssetData.AssetClassPath);
+
+		TArray<FAssetDisplayInfo> OutSystemInfo;
+		TSharedRef<SVerticalBox> OverallTooltipVBox = SNew(SVerticalBox);
+		{
+			const FSlateBrush* ClassIcon = FAppStyle::GetDefaultBrush();
+			TOptional<FLinearColor> Color;
+			TArray<FAssetDisplayInfo> OutStatusInfo;
+			if (const UAssetDefinition* AssetDefinition = UAssetDefinitionRegistry::Get()->GetAssetDefinitionForClass(AssetData.GetClass()))
+			{
+				TSharedRef<FAssetStatusAssetDataInfoProvider> Provider = MakeShared<FAssetStatusAssetDataInfoProvider>(AssetData);
+				AssetDefinition->GetAssetStatusInfo(Provider, OutStatusInfo);
+				ClassIcon = AssetDefinition->GetIconBrush(AssetData, AssetData.AssetClassPath.GetAssetName());
+				Color = AssetDefinition->GetAssetColor();
+			}
+			// Get the default asset definition to retrieve some default information on it (ex: SourceControl)
+			else if (const UAssetDefinition* AssetDefinitionDefaultUObject = GetDefault<UAssetDefinitionDefault>())
+			{
+				TSharedRef<FAssetStatusAssetDataInfoProvider> Provider = MakeShared<FAssetStatusAssetDataInfoProvider>(AssetData);
+				AssetDefinitionDefaultUObject->GetAssetStatusInfo(Provider, OutStatusInfo);
+			}
+
+			if (AssetSystemInfoProvider.IsValid())
+			{
+				AssetSystemInfoProvider->PopulateAssetInfo(OutSystemInfo);
+			}
+
+			if (ClassIcon == nullptr || ClassIcon == FAppStyle::GetDefaultBrush())
+			{
+				ClassIcon = FSlateIconFinder::FindIconForClass(AssetData.GetClass()).GetIcon();
+			}
+
+			FText ClassNameText = NSLOCTEXT("AssetThumbnail", "ClassNameText", "Not Found");
+			if (Class != NULL)
+			{
+				ClassNameText = Class->GetDisplayNameText();
+			}
+			else if (!AssetData.AssetClassPath.IsNull())
+			{
+				ClassNameText = FText::FromString(AssetData.AssetClassPath.ToString());
+			}
+
+			const FText NameText = FText::FromString(AssetData.AssetName.ToString());
+			// Name/Type Slot
+			OverallTooltipVBox->AddSlot()
+			.AutoHeight()
+			[
+				SNew(SVerticalBox)
+
+				+ SVerticalBox::Slot()
+				.Padding(0.f, 0.f, 0.f, 6.f)
+				.AutoHeight()
+				[
+					SNew(STextBlock)
+					.Text(NameText)
+					.ColorAndOpacity(FStyleColors::White)
+					.Font(FAppStyle::GetFontStyle("ContentBrowser.Tooltip.AssetNameFont"))
+				]
+
+				+ SVerticalBox::Slot()
+				.AutoHeight()
+				[
+					SNew(SHorizontalBox)
+
+					+ SHorizontalBox::Slot()
+					.AutoWidth()
+					.Padding(0.f, 0.f, 4.f, 0.f)
+					[
+						SNew(SBox)
+						.WidthOverride(16.f)
+						.HeightOverride(16.f)
+						[
+							SNew(SImage)
+							.Image(ClassIcon)
+							.ColorAndOpacity_Lambda([Color] () { return Color.IsSet() ? Color.GetValue() : FStyleColors::White;})
+						]
+					]
+
+					+ SHorizontalBox::Slot()
+					.AutoWidth()
+					.VAlign(VAlign_Center)
+					[
+						SNew(STextBlock)
+						.Font(FAppStyle::GetFontStyle("ContentBrowser.Tooltip.EntryFont"))
+						.Text(ClassNameText)
+					]
+				]
+			];
+
+			// Separator
+			OverallTooltipVBox->AddSlot()
+			.Padding(0.f,6.f)
+			.AutoHeight()
+			[
+				SNew(SSeparator)
+				.Orientation(Orient_Horizontal)
+				.Thickness(1.f)
+				.ColorAndOpacity(COLOR("#484848FF"))
+				.SeparatorImage(FAppStyle::Get().GetBrush("WhiteBrush"))
+			];
+
+			TSharedRef<SVerticalBox> StatusVerticalBox = SNew(SVerticalBox);
+
+			for (const FAssetDisplayInfo& AssetStatusInfo : OutStatusInfo)
+			{
+				StatusVerticalBox->AddSlot()
+					.Padding(0.f, 0.f, 0.f, 6.f)
+					.AutoHeight()
+					[
+						SNew(SHorizontalBox)
+						.Visibility(AssetStatusInfo.IsVisible)
+
+						+ SHorizontalBox::Slot()
+						.AutoWidth()
+						.Padding(0.f, 0.f, 4.f, 0.f)
+						[
+						   SNew(SBox)
+						   .WidthOverride(16.f)
+						   .HeightOverride(16.f)
+						   [
+							   SNew(SImage)
+							   .Image(AssetStatusInfo.StatusIcon)
+						   ]
+						]
+
+						+ SHorizontalBox::Slot()
+						.AutoWidth()
+						.VAlign(VAlign_Center)
+						[
+							SNew(STextBlock)
+							.Font(FAppStyle::GetFontStyle("ContentBrowser.Tooltip.EntryFont"))
+							.Text(AssetStatusInfo.StatusDescription)
+						]
+					];
+			}
+
+			// Status
+			OverallTooltipVBox->AddSlot()
+			.AutoHeight()
+			[
+				StatusVerticalBox
+			];
+
+			// More info
+			if (!OutSystemInfo.IsEmpty())
+			{
+				OverallTooltipVBox->AddSlot()
+				.Padding(0.f, 2.f, 0.f, 0.f)
+				.AutoHeight()
+				.VAlign(VAlign_Center)
+				[
+					SNew(SRichTextBlock)
+					.Visibility_Lambda([] ()
+					{
+						const FModifierKeysState ModifierKeys = FSlateApplication::Get().GetModifierKeys();
+						return ModifierKeys.IsAltDown() && ModifierKeys.IsControlDown() ? EVisibility::Collapsed : EVisibility::Visible;
+					})
+					// .Font(FAppStyle::GetFontStyle("ContentBrowser.Tooltip.MoreInfo"))
+					.Justification(ETextJustify::Center)
+					.Text(NSLOCTEXT("AssetThumbnail", "MoreInfoTooltip", "Hold <WrappedCommand/> for more"))
+					+ SRichTextBlock::WidgetDecorator( TEXT("WrappedCommand"), this, &SAssetThumbnail::OnCreateWidgetDecoratorWidget )
+				];
+			}
+		}
+
+		TSharedRef<SVerticalBox> ExtendedToolTipVerticalBox = SNew(SVerticalBox);
+		TSharedRef<SBox> ExtendedToolTip = SNew(SBox)
+			.Padding(FMargin(9.f, -9.f, 10.f, 6.f))
+			[
+				ExtendedToolTipVerticalBox
+			];
+
+		for (const FAssetDisplayInfo& SystemInfo : OutSystemInfo)
+		{
+			if (!SystemInfo.StatusTitle.IsSet() || !SystemInfo.StatusDescription.IsSet())
+			{
+				continue;
+			}
+
+			if ((SystemInfo.IsVisible.IsSet() && SystemInfo.IsVisible.Get().IsVisible()) || !SystemInfo.IsVisible.IsSet())
+			{
+				AddToExtendedToolTipInfoBox(ExtendedToolTipVerticalBox, SystemInfo.StatusIcon, SystemInfo.StatusTitle.Get(), SystemInfo.StatusDescription.Get());
+			}
+		}
+
+		return SNew(SDocumentationToolTip)
+				.OverrideExtendedToolTipContent(ExtendedToolTip)
+				[
+					OverallTooltipVBox
+				];
+	}
+#endif
+
 private:
 #if UE_CONTENTBROWSER_NEW_STYLE
+	FSlateWidgetRun::FWidgetRunInfo OnCreateWidgetDecoratorWidget(const FTextRunInfo& InRunInfo, const ISlateStyle* InStyle) const
+	{
+		TSharedRef<SWidget> CtrlAltWidget = SNew(SBorder)
+			.VAlign(VAlign_Center)
+			.HAlign(HAlign_Center)
+			.BorderImage(FAppStyle::GetBrush("ContentBrowser.ToolTip.CommandBorder"))
+			[
+				SNew(STextBlock)
+				.Font(FAppStyle::GetFontStyle("ContentBrowser.Tooltip.MoreInfo"))
+				.Text(NSLOCTEXT("AssetThumbnail", "CtrlAltLabel", " Ctrl + Alt "))
+			];
+		const TSharedRef<FSlateFontMeasure> FontMeasure = FSlateApplication::Get().GetRenderer()->GetFontMeasureService();
+		const int16 Baseline = FontMeasure->GetBaseline(FAppStyle::GetFontStyle("ContentBrowser.Tooltip.MoreInfo"));
+
+		return FSlateWidgetRun::FWidgetRunInfo(CtrlAltWidget, Baseline - 2);
+
+	}
+
+	void AddToExtendedToolTipInfoBox(const TSharedRef<SVerticalBox>& InfoBox, const TAttribute<const FSlateBrush*>& Icon, const FText& Key, const FText& Value) const
+	{
+		InfoBox->AddSlot()
+		.Padding(0.f, 0.f, 0.f, 6.f)
+		.AutoHeight()
+		[
+			SNew(SHorizontalBox)
+
+			+SHorizontalBox::Slot()
+			.AutoWidth()
+			.Padding(0, 0, 4, 0)
+			[
+				SNew(SBox)
+				.Visibility(Icon.IsSet()? EVisibility::Visible : EVisibility::Collapsed)
+				.WidthOverride(16.f)
+				.HeightOverride(16.f)
+				[
+					SNew(SImage)
+					.Image(Icon)
+				]
+			]
+
+			+SHorizontalBox::Slot()
+			.AutoWidth()
+			.Padding(0, 0, 4, 0)
+			[
+				SNew(STextBlock)
+				.Font(FAppStyle::GetFontStyle("ContentBrowser.Tooltip.EntryFont"))
+				.Text(FText::Format(NSLOCTEXT("AssetThumbnailToolTip", "AssetViewTooltipFormat", "{0}:"), Key))
+			]
+
+			+SHorizontalBox::Slot()
+			.AutoWidth()
+			[
+				SNew(STextBlock)
+				.Font(FAppStyle::GetFontStyle("ContentBrowser.Tooltip.EntryFont"))
+				.ColorAndOpacity(FStyleColors::White)
+				.Text(Value)
+			]
+		];
+	}
+
 	EVisibility GetStatusBorderVisibility() const
 	{
-		for (const FAssetStatus& AssetStatus : OverlayInfo.AssetStatus)
+		for (const FAssetDisplayInfo& AssetStatus : OverlayInfo)
 		{
 			if (AssetStatus.IsVisible.IsSet() && AssetStatus.IsVisible.Get().IsVisible())
 			{
@@ -567,7 +853,7 @@ private:
 		return EVisibility::Collapsed;
 	}
 
-	TSharedRef<SWidget> CreateStatusWidget(int32 StatusIndex, const FAssetStatus& InStatusInfo) const
+	TSharedRef<SWidget> CreateStatusWidget(int32 StatusIndex, const FAssetDisplayInfo& InStatusInfo) const
 	{
 		return SNew(SBox)
 			.WidthOverride(StatusSize)
@@ -582,9 +868,9 @@ private:
 
 	EVisibility GetStatusVisibility(int32 StatusIndex) const
 	{
-		if (OverlayInfo.AssetStatus.IsValidIndex(StatusIndex))
+		if (OverlayInfo.IsValidIndex(StatusIndex))
 		{
-			const FAssetStatus& AssetStatus = OverlayInfo.AssetStatus[StatusIndex];
+			const FAssetDisplayInfo& AssetStatus = OverlayInfo[StatusIndex];
 			if (AssetStatus.IsVisible.IsSet() && AssetStatus.IsVisible.Get() == EVisibility::Visible)
 			{
 				return GetStatusVisibilityBasedOnGeometry(StatusIndex);
@@ -615,7 +901,7 @@ private:
 		const int32 MaxShownStatus = FMath::FloorToInt32(ThumbnailWidth / StatusSize);
 
 		int32 ShownStatus = 0;
-		for (const FAssetStatus& Status : OverlayInfo.AssetStatus)
+		for (const FAssetDisplayInfo& Status : OverlayInfo)
 		{
 			if (Status.IsVisible.IsSet() && Status.IsVisible.Get().IsVisible())
 			{
@@ -632,7 +918,7 @@ private:
 		const int32 MaxShownStatus = FMath::FloorToInt32(ThumbnailWidth / StatusSize);
 
 		int32 ShownStatus = 0;
-		for (const FAssetStatus& Status : OverlayInfo.AssetStatus)
+		for (const FAssetDisplayInfo& Status : OverlayInfo)
 		{
 			if (Status.IsVisible.IsSet() && Status.IsVisible.Get().IsVisible())
 			{
@@ -1082,8 +1368,9 @@ private:
 	TSharedPtr<SWidget> StatusOverflowWidget;
 	const float StatusSize = 16.f;
 	const float StatusPadding = 4.f;
-	FAssetStatusInfo OverlayInfo;
+	TArray<FAssetDisplayInfo> OverlayInfo;
 	TArray<TSharedPtr<SWidget>> Statuses;
+	TSharedPtr<IAssetSystemInfoProvider> AssetSystemInfoProvider;
 	bool bShowAssetChip;
 	TAttribute<const FSlateBrush*> AssetChipBorderImageOverride;
 #endif
@@ -1107,8 +1394,37 @@ private:
 	bool bIsClassType;
 };
 
+void SAssetThumbnailToolTip::Construct(const FArguments& InArgs)
+{
+	AssetThumbnail = InArgs._AssetThumbnail;
 
+	SToolTip::Construct(
+		SToolTip::FArguments()
+		.TextMargin(FMargin(1.f, -3.f))
+		.BorderImage(FAppStyle::GetBrush("ContentBrowser.Tooltip.Border"))
+		);
+}
 
+bool SAssetThumbnailToolTip::IsEmpty() const
+{
+	return !AssetThumbnail.IsValid();
+}
+
+void SAssetThumbnailToolTip::OnOpening()
+{
+#if UE_CONTENTBROWSER_NEW_STYLE
+	TSharedPtr<SAssetThumbnail> AssetViewItemPin = AssetThumbnail.Pin();
+	if (AssetViewItemPin.IsValid())
+	{
+		SetContentWidget(AssetViewItemPin->GetDefaultTooltip());
+	}
+#endif
+}
+
+void SAssetThumbnailToolTip::OnClosed()
+{
+	ResetContentWidget();
+}
 
 FAssetThumbnail::FAssetThumbnail( UObject* InAsset, uint32 InWidth, uint32 InHeight, const TSharedPtr<class FAssetThumbnailPool>& InThumbnailPool )
 	: ThumbnailPool(InThumbnailPool)
@@ -1226,6 +1542,8 @@ TSharedRef<SWidget> FAssetThumbnail::MakeThumbnailWidget( const FAssetThumbnailC
 		.Padding(InConfig.Padding)
 		.GenericThumbnailSize(InConfig.GenericThumbnailSize)
 #if UE_CONTENTBROWSER_NEW_STYLE
+		.AssetSystemInfoProvider(InConfig.AssetSystemInfoProvider)
+		.TooltipOverride(InConfig.TooltipOverride)
 		.AllowAssetStatusThumbnailOverlay(InConfig.bAllowAssetStatusThumbnailOverlay)
 		.ShowAssetChip(InConfig.bShowAssetChip)
 		.AssetChipBorderImageOverride(InConfig.AssetChipBorderImageOverride)
diff --git a/Engine/Source/Editor/UnrealEd/Private/AssetThumbnailToolTip.h b/Engine/Source/Editor/UnrealEd/Private/AssetThumbnailToolTip.h
new file mode 100644
--- /dev/null
+++ b/Engine/Source/Editor/UnrealEd/Private/AssetThumbnailToolTip.h
@@ -0,0 +1,31 @@
+﻿// Copyright Epic Games, Inc. All Rights Reserved.
+
+#pragma once
+#include "Widgets/SToolTip.h"
+
+class SAssetThumbnail;
+
+// AssetThumbnail ToolTip, Implementation is inside AssetThumbnail.cpp
+class SAssetThumbnailToolTip : public SToolTip
+{
+public:
+	SLATE_BEGIN_ARGS(SAssetThumbnailToolTip)
+		: _AssetThumbnail()
+	{ }
+
+	SLATE_ARGUMENT(TSharedPtr<SAssetThumbnail>, AssetThumbnail)
+
+	SLATE_END_ARGS()
+
+	void Construct(const FArguments& InArgs);
+
+	// IToolTip interface
+	virtual bool IsEmpty() const override;
+
+	virtual void OnOpening() override;
+
+	virtual void OnClosed() override;
+
+private:
+	TWeakPtr<SAssetThumbnail> AssetThumbnail;
+};
diff --git a/Engine/Source/Editor/UnrealEd/Public/AssetDefinitionDefault.h b/Engine/Source/Editor/UnrealEd/Public/AssetDefinitionDefault.h
--- a/Engine/Source/Editor/UnrealEd/Public/AssetDefinitionDefault.h
+++ b/Engine/Source/Editor/UnrealEd/Public/AssetDefinitionDefault.h
@@ -17,7 +17,7 @@ public:
 	UNREALED_API virtual EAssetCommandResult OpenAssets(const FAssetOpenArgs& OpenArgs) const override;
 	UNREALED_API virtual EAssetCommandResult PerformAssetDiff(const FAssetDiffArgs& DiffArgs) const override;
 #if UE_CONTENTBROWSER_NEW_STYLE
-	UNREALED_API virtual void GetAssetStatusInfo(const TSharedPtr<IAssetStatusInfoProvider>& InAssetStatusInfoProvider, FAssetStatusInfo& OutStatusInfo) const override;
+	UNREALED_API virtual void GetAssetStatusInfo(const TSharedPtr<IAssetStatusInfoProvider>& InAssetStatusInfoProvider, TArray<FAssetDisplayInfo>& OutStatusInfo) const override;
 #endif
 	// UAssetDefinition End
 };
diff --git a/Engine/Source/Editor/UnrealEd/Public/AssetThumbnail.h b/Engine/Source/Editor/UnrealEd/Public/AssetThumbnail.h
--- a/Engine/Source/Editor/UnrealEd/Public/AssetThumbnail.h
+++ b/Engine/Source/Editor/UnrealEd/Public/AssetThumbnail.h
@@ -15,6 +15,8 @@ class FAssetThumbnailPool;
 class FSlateShaderResource;
 class FSlateTexture2DRHIRef;
 class FSlateTextureRenderTarget2DResource;
+class IAssetSystemInfoProvider;
+class IToolTip;
 struct FPropertyChangedEvent;
 
 namespace EThumbnailLabel
@@ -65,6 +67,7 @@ struct FAssetThumbnailConfig
 		, AssetTypeColorOverride()
 		, Padding(0)
 #if UE_CONTENTBROWSER_NEW_STYLE
+		, AssetSystemInfoProvider(nullptr)
 		, bAllowAssetStatusThumbnailOverlay( false )
 		, bShowAssetChip(false)
 		, AssetChipBorderImageOverride()
@@ -86,6 +89,8 @@ struct FAssetThumbnailConfig
 	TAttribute<int32> GenericThumbnailSize = 64;
 	EThumbnailColorStripOrientation ColorStripOrientation = EThumbnailColorStripOrientation::HorizontalBottomEdge;
 #if UE_CONTENTBROWSER_NEW_STYLE
+	TSharedPtr<IAssetSystemInfoProvider> AssetSystemInfoProvider;
+	TAttribute<TSharedPtr<IToolTip>> TooltipOverride;
 	bool bAllowAssetStatusThumbnailOverlay;
 	bool bShowAssetChip;
 	TAttribute<const FSlateBrush*> AssetChipBorderImageOverride;
