[ContentBrowser] Issue 584874: Warnings in Build Engine Localization

diff --git a/Engine/Source/Editor/ContentBrowser/Private/AssetViewWidgets.cpp b/Engine/Source/Editor/ContentBrowser/Private/AssetViewWidgets.cpp
--- a/Engine/Source/Editor/ContentBrowser/Private/AssetViewWidgets.cpp
+++ b/Engine/Source/Editor/ContentBrowser/Private/AssetViewWidgets.cpp
@@ -1167,7 +1167,7 @@ TSharedRef<SWidget> SAssetViewItem::CreateToolTipWidget() const
 							SNew(STextBlock)
 							.Font(FAppStyle::GetFontStyle("ContentBrowser.Tooltip.EntryFont"))
 							.ColorAndOpacity(FStyleColors::White)
-							.Text(LOCTEXT("FolderNameBracketed", " (Folder)"))
+							.Text(LOCTEXT("FolderNameBracketedLabel", " (Folder)"))
 						]
 					]
 				]
