Convert the show other developers filter to a hide other developers instead.

diff --git a/Engine/Plugins/Experimental/EditorDataStorageFeatures/Source/TedsOutliner/Public/TedsOutlinerFilter.h b/Engine/Plugins/Experimental/EditorDataStorageFeatures/Source/TedsOutliner/Public/TedsOutlinerFilter.h
--- a/Engine/Plugins/Experimental/EditorDataStorageFeatures/Source/TedsOutliner/Public/TedsOutlinerFilter.h
+++ b/Engine/Plugins/Experimental/EditorDataStorageFeatures/Source/TedsOutliner/Public/TedsOutlinerFilter.h
@@ -34,7 +34,6 @@ public:
 	virtual FName GetIconName() const override;
 
 	/** If true, the filter will be active in the FilterBar when it is inactive in the UI (i.e the filter pill is grayed out)
-	 * @See: FFrontendFilter_ShowOtherDevelopers in Content Browser
 	 */
 	virtual bool IsInverseFilter() const override;
 
diff --git a/Engine/Source/Editor/ContentBrowser/Private/Filters.cpp b/Engine/Source/Editor/ContentBrowser/Private/Filters.cpp
--- a/Engine/Source/Editor/ContentBrowser/Private/Filters.cpp
+++ b/Engine/Source/Editor/ContentBrowser/Private/Filters.cpp
@@ -60,19 +60,19 @@ void FFilter_ShowRedirectors::LoadSettings(const FString& IniFilename, const FSt
  * FFilter_OtherDevelopers 
  */
 
-FFilter_ShowOtherDevelopers::FFilter_ShowOtherDevelopers(TSharedPtr<FFrontendFilterCategory> InCategory, FName InFilterBarIdentifier)
+FFilter_HideOtherDevelopers::FFilter_HideOtherDevelopers(TSharedPtr<FFrontendFilterCategory> InCategory, FName InFilterBarIdentifier)
 	: FFrontendFilter(InCategory)
 	, FilterBarIdentifier(InFilterBarIdentifier)
 	, PathPermissionList(MakeShared<FPathPermissionList>())
 {
 	UContentBrowserDataSubsystem* ContentBrowserData = IContentBrowserDataModule::Get().GetSubsystem();
-	ItemDataUpdatedHandle = ContentBrowserData->OnItemDataUpdated().AddRaw(this, &FFilter_ShowOtherDevelopers::HandleItemDataUpdated);
-	ItemDataRefreshedHandle = ContentBrowserData->OnItemDataRefreshed().AddRaw(this, &FFilter_ShowOtherDevelopers::HandleItemDataRefreshed);
+	ItemDataUpdatedHandle = ContentBrowserData->OnItemDataUpdated().AddRaw(this, &FFilter_HideOtherDevelopers::HandleItemDataUpdated);
+	ItemDataRefreshedHandle = ContentBrowserData->OnItemDataRefreshed().AddRaw(this, &FFilter_HideOtherDevelopers::HandleItemDataRefreshed);
 
 	BuildFilter();
 }
 
-FFilter_ShowOtherDevelopers::~FFilter_ShowOtherDevelopers()
+FFilter_HideOtherDevelopers::~FFilter_HideOtherDevelopers()
 {
 	if (IContentBrowserDataModule* ContentBrowserModule = IContentBrowserDataModule::GetPtr())
 	{
@@ -84,12 +84,12 @@ FFilter_ShowOtherDevelopers::~FFilter_ShowOtherDevelopers()
 	}
 }
 
-TSharedRef<const FPathPermissionList> FFilter_ShowOtherDevelopers::GetPathPermissionList() 
+TSharedRef<const FPathPermissionList> FFilter_HideOtherDevelopers::GetPathPermissionList() 
 {
 	return PathPermissionList;
 }
 
-void FFilter_ShowOtherDevelopers::BuildFilter()
+void FFilter_HideOtherDevelopers::BuildFilter()
 {
 	static const FName NAME_OtherDevelopers{"OtherDevelopers"};
 	// Update list of other developer folders and put into permission list 
@@ -107,7 +107,7 @@ void FFilter_ShowOtherDevelopers::BuildFilter()
 	if (OtherDeveloperFolders.Num() != PreviousPaths.Num() || OtherDeveloperFolders.Difference(PreviousPaths).Num() != 0
 		|| PreviousPaths.Difference(OtherDeveloperFolders).Num() != 0)
 	{
-		UE_LOG(LogContentBrowser, Verbose, TEXT("[%s] FFilterShowOtherDevelopers rebuilt exclusion list: %s"), *FilterBarIdentifier.ToString(), 
+		UE_LOG(LogContentBrowser, Verbose, TEXT("[%s] FFilterHideOtherDevelopers rebuilt exclusion list: %s"), *FilterBarIdentifier.ToString(), 
 			*FString::JoinBy(OtherDeveloperFolders, TEXT(","), UE_PROJECTION_MEMBER(FName, ToString)));
 		
 		// Recreate the permission list so that the content browser can do pointer comparison to tell that the list has changed rather than tbinding the delegate
@@ -120,16 +120,16 @@ void FFilter_ShowOtherDevelopers::BuildFilter()
 	}
 	else
 	{
-		UE_LOG(LogContentBrowser, Verbose, TEXT("[%s] FFilterShowOtherDevelopers keeping previous exclusion list"), *FilterBarIdentifier.ToString());
+		UE_LOG(LogContentBrowser, Verbose, TEXT("[%s] FFilterHideOtherDevelopers keeping previous exclusion list"), *FilterBarIdentifier.ToString());
 	}
 }
 
-void FFilter_ShowOtherDevelopers::HandleItemDataRefreshed()
+void FFilter_HideOtherDevelopers::HandleItemDataRefreshed()
 {
 	BuildFilter();
 }
 
-void FFilter_ShowOtherDevelopers::HandleItemDataUpdated(TArrayView<const FContentBrowserItemDataUpdate> InUpdatedItems)
+void FFilter_HideOtherDevelopers::HandleItemDataUpdated(TArrayView<const FContentBrowserItemDataUpdate> InUpdatedItems)
 {
 	const bool bNeedsRebuild = Algo::AnyOf(InUpdatedItems, [](const FContentBrowserItemDataUpdate& Update){
 		FName InternalPath = Update.GetItemData().GetInternalPath();
@@ -149,43 +149,43 @@ void FFilter_ShowOtherDevelopers::HandleItemDataUpdated(TArrayView<const FConten
 }
 
 /** Returns the human readable name for this filter */
-FText FFilter_ShowOtherDevelopers::GetDisplayName() const
+FText FFilter_HideOtherDevelopers::GetDisplayName() const
 {
-	return LOCTEXT("FrontendFilter_ShowOtherDevelopers", "Other Developers"); 
+	return LOCTEXT("FrontendFilter_HideOtherDevelopers", "Hide Other Developers"); 
 }
 
 /** Returns the tooltip for this filter, shown in the filters menu */
-FText FFilter_ShowOtherDevelopers::GetToolTipText() const
+FText FFilter_HideOtherDevelopers::GetToolTipText() const
 { 
-	return LOCTEXT("FrontendFilter_ShowOtherDevelopersTooltip", "Allow display of assets in developer folders that aren't yours."); 
+	return LOCTEXT("FrontendFilter_HideOtherDevelopersTooltip", "Hide the display of assets in developer folders that aren't yours."); 
 }
 
 /** Returns the name of the icon to use in menu entries */
-FName FFilter_ShowOtherDevelopers::GetIconName() const
+FName FFilter_HideOtherDevelopers::GetIconName() const
 {
 	return NAME_None;
 }
 
 /** Notification that the filter became active or inactive */
-void FFilter_ShowOtherDevelopers::ActiveStateChanged(bool bActive)
+void FFilter_HideOtherDevelopers::ActiveStateChanged(bool bActive)
 {
 	if (bActive)
 	{
-		UE_LOG(LogContentBrowser, Verbose, TEXT("[%s] FFilterShowOtherDevelopers active, hiding content from other developers"), *FilterBarIdentifier.ToString());
+		UE_LOG(LogContentBrowser, Verbose, TEXT("[%s] FFilterHideOtherDevelopers active, hiding content from other developers"), *FilterBarIdentifier.ToString());
 	}
 	else
 	{
-		UE_LOG(LogContentBrowser, Verbose, TEXT("[%s] FFilterShowOtherDevelopers inactive, showing content from all developers"), *FilterBarIdentifier.ToString());
+		UE_LOG(LogContentBrowser, Verbose, TEXT("[%s] FFilterHideOtherDevelopers inactive, showing content from all developers"), *FilterBarIdentifier.ToString());
 	}
 }
 
 /** Can be overriden for custom FilterBar subclasses to save settings, currently not implemented in any gneeric Filter Bar */
-void FFilter_ShowOtherDevelopers::SaveSettings(const FString& IniFilename, const FString& IniSection, const FString& SettingsString) const
+void FFilter_HideOtherDevelopers::SaveSettings(const FString& IniFilename, const FString& IniSection, const FString& SettingsString) const
 {
 }
 
 /** Can be overriden for custom FilterBar subclasses to load settings, currently not implemented in any gneeric Filter Bar */
-void FFilter_ShowOtherDevelopers::LoadSettings(const FString& IniFilename, const FString& IniSection, const FString& SettingsString)
+void FFilter_HideOtherDevelopers::LoadSettings(const FString& IniFilename, const FString& IniSection, const FString& SettingsString)
 {
 }
 
diff --git a/Engine/Source/Editor/ContentBrowser/Private/Filters.h b/Engine/Source/Editor/ContentBrowser/Private/Filters.h
--- a/Engine/Source/Editor/ContentBrowser/Private/Filters.h
+++ b/Engine/Source/Editor/ContentBrowser/Private/Filters.h
@@ -35,7 +35,6 @@ public:
 	virtual FName GetIconName() const override;
 
 	/** If true, the filter will be active in the FilterBar when it is inactive in the UI (i.e the filter pill is grayed out)
-	 * @See: FFrontendFilter_ShowOtherDevelopers in Content Browser
 	 */
 	virtual bool IsInverseFilter() const 
 	{
@@ -99,12 +98,12 @@ private:
 };
 
 // Non-frontend filter which modifies content browser backend query to exclude folders belonging to other developers 
-class FFilter_ShowOtherDevelopers : public FFrontendFilter
+class FFilter_HideOtherDevelopers : public FFrontendFilter
 {
 public: 
-	FFilter_ShowOtherDevelopers(TSharedPtr<FFrontendFilterCategory> InCategory, FName FilterBarIdentifier);
-	FFilter_ShowOtherDevelopers& operator=(const FFilter_ShowOtherDevelopers&) = delete;
-	virtual ~FFilter_ShowOtherDevelopers();
+	FFilter_HideOtherDevelopers(TSharedPtr<FFrontendFilterCategory> InCategory, FName FilterBarIdentifier);
+	FFilter_HideOtherDevelopers& operator=(const FFilter_HideOtherDevelopers&) = delete;
+	virtual ~FFilter_HideOtherDevelopers();
 
 	/** Get the list of folders to be denied when this filter is active (visually disabled, becuase it's an inverse filter) */
 	TSharedRef<const FPathPermissionList> GetPathPermissionList();
@@ -114,7 +113,7 @@ public:
 	virtual bool PassesFilter(FAssetFilterType InItem) const override { return true; }
 
 	/** Returns the system name for this filter */
-	virtual FString GetName() const override { return TEXT("ShowOtherDevelopersBackend"); }
+	virtual FString GetName() const override { return TEXT("HideOtherDevelopersBackend"); }
 
 	/** Returns the human readable name for this filter */
 	virtual FText GetDisplayName() const override;
@@ -125,9 +124,6 @@ public:
 	/** Returns the name of the icon to use in menu entries */
 	virtual FName GetIconName() const override;
 
-	/** If true, the filter will be active in the FilterBar when it is inactive in the UI (i.e the filter pill is grayed out) */
-	virtual bool IsInverseFilter() const override { return true; }
-
 	/** Notification that the filter became active or inactive */
 	virtual void ActiveStateChanged(bool bActive) override;
 
diff --git a/Engine/Source/Editor/ContentBrowser/Private/SAssetPicker.cpp b/Engine/Source/Editor/ContentBrowser/Private/SAssetPicker.cpp
--- a/Engine/Source/Editor/ContentBrowser/Private/SAssetPicker.cpp
+++ b/Engine/Source/Editor/ContentBrowser/Private/SAssetPicker.cpp
@@ -273,12 +273,12 @@ void SAssetPicker::Construct( const FArguments& InArgs )
 		];
 
 		// Use the 'other developer' filter from the filter list widget. 
-		OtherDevelopersFilter = StaticCastSharedPtr<FFilter_ShowOtherDevelopers>(FilterListPtr->GetFrontendFilter(TEXT("ShowOtherDevelopersBackend")));
+		OtherDevelopersFilter = StaticCastSharedPtr<FFilter_HideOtherDevelopers>(FilterListPtr->GetFrontendFilter(TEXT("HideOtherDevelopersBackend")));
 	}
 	else
 	{
 		// Filter UI is off, but the 'other developer' filter is a built-in feature.
-		OtherDevelopersFilter = MakeShared<FFilter_ShowOtherDevelopers>(nullptr, FName(SaveSettingsName));
+		OtherDevelopersFilter = MakeShared<FFilter_HideOtherDevelopers>(nullptr, FName(SaveSettingsName));
 		OtherDevelopersFilter->SetActiveInCollection(OtherDevelopersFilter.ToSharedRef(), false, *FrontendFilters);
 	}
 
@@ -663,7 +663,7 @@ bool SAssetPicker::IsShowingOtherDevelopersContent() const
 {
 	if (FilterListPtr.IsValid())
 	{
-		return FilterListPtr->IsFrontendFilterActive(OtherDevelopersFilter);
+		return !FilterListPtr->IsFrontendFilterActive(OtherDevelopersFilter);
 	}
 	else
 	{
@@ -671,11 +671,10 @@ bool SAssetPicker::IsShowingOtherDevelopersContent() const
 		{
 			if (FrontendFilters->GetFilterAtIndex(i) == OtherDevelopersFilter)
 			{
-				// Inverse filter, so its presence in the list means "hide other developers content"
-				return false;
+				return true;
 			}
 		}
-		return true;
+		return false;
 	}
 }
 
@@ -703,11 +702,11 @@ void SAssetPicker::HandleShowOtherDevelopersCheckStateChanged( ECheckBoxState In
 	{
 		if (InCheckboxState == ECheckBoxState::Checked)
 		{
-			OtherDevelopersFilter->SetActiveInCollection(OtherDevelopersFilter.ToSharedRef(), /* do show other developers content*/ true, *FrontendFilters);
+			OtherDevelopersFilter->SetActiveInCollection(OtherDevelopersFilter.ToSharedRef(), /* do show other developers content*/ false, *FrontendFilters);
 		}
 		else
 		{
-			OtherDevelopersFilter->SetActiveInCollection(OtherDevelopersFilter.ToSharedRef(), /* do not show other developers content*/ false, *FrontendFilters);
+			OtherDevelopersFilter->SetActiveInCollection(OtherDevelopersFilter.ToSharedRef(), /* do not show other developers content*/ true, *FrontendFilters);
 		}
 	}
 }
diff --git a/Engine/Source/Editor/ContentBrowser/Private/SAssetPicker.h b/Engine/Source/Editor/ContentBrowser/Private/SAssetPicker.h
--- a/Engine/Source/Editor/ContentBrowser/Private/SAssetPicker.h
+++ b/Engine/Source/Editor/ContentBrowser/Private/SAssetPicker.h
@@ -19,7 +19,7 @@
 #include "Widgets/DeclarativeSyntaxSupport.h"
 #include "Widgets/SCompoundWidget.h"
 
-class FFilter_ShowOtherDevelopers;
+class FFilter_HideOtherDevelopers;
 class FUICommandList;
 class SAssetSearchBox;
 class SAssetView;
@@ -190,7 +190,7 @@ private:
 	/** Filters needed for filtering the assets */
 	TSharedPtr< FAssetFilterCollectionType > FilterCollection;
 	TSharedPtr<class FAssetTextFilter> TextFilter;
-	TSharedPtr<FFilter_ShowOtherDevelopers> OtherDevelopersFilter;
+	TSharedPtr<FFilter_HideOtherDevelopers> OtherDevelopersFilter;
 
 	EAssetTypeCategories::Type DefaultFilterMenuExpansion;
 
diff --git a/Engine/Source/Editor/ContentBrowser/Private/SFilterList.cpp b/Engine/Source/Editor/ContentBrowser/Private/SFilterList.cpp
--- a/Engine/Source/Editor/ContentBrowser/Private/SFilterList.cpp
+++ b/Engine/Source/Editor/ContentBrowser/Private/SFilterList.cpp
@@ -71,7 +71,7 @@ void SFilterList::Construct( const FArguments& InArgs )
 
 	TSharedPtr<FFrontendFilterCategory> DefaultCategory = MakeShareable( new FFrontendFilterCategory(LOCTEXT("FrontendFiltersCategory", "Other Filters"), LOCTEXT("FrontendFiltersCategoryTooltip", "Filter assets by all filters in this category.")) );
 	
-	TSharedRef<FFilter_ShowOtherDevelopers> OtherDevelopersFilter = MakeShared<FFilter_ShowOtherDevelopers>(DefaultCategory, InArgs._FilterBarIdentifier);
+	TSharedRef<FFilter_HideOtherDevelopers> OtherDevelopersFilter = MakeShared<FFilter_HideOtherDevelopers>(DefaultCategory, InArgs._FilterBarIdentifier);
 	// This filter affecst the backend query so we must perform a full refresh when it changes
 	OtherDevelopersFilter->OnChanged().Add(this->OnFilterChanged);
 
@@ -189,8 +189,8 @@ TSharedRef<SWidget> SFilterList::ExternalMakeAddFilterMenu()
 
 FARFilter SFilterList::GetCombinedBackendFilter(TArray<TSharedRef<const FPathPermissionList>>& OutPermissionLists) const
 {
-	TSharedPtr<FFilter_ShowOtherDevelopers> OtherDevelopersFilter = StaticCastSharedPtr<FFilter_ShowOtherDevelopers>(GetFrontendFilter("ShowOtherDevelopersBackend"));
-	if (!OtherDevelopersFilter->IsActive()) // Filter when inactive, because it's an inverse filter 
+	TSharedPtr<FFilter_HideOtherDevelopers> OtherDevelopersFilter = StaticCastSharedPtr<FFilter_HideOtherDevelopers>(GetFrontendFilter("HideOtherDevelopersBackend"));
+	if (OtherDevelopersFilter->IsActive())
 	{
 		OutPermissionLists.Add(OtherDevelopersFilter->GetPathPermissionList());
 	}
@@ -269,17 +269,8 @@ void SFilterList::DisableFiltersThatHideItems(TArrayView<const FContentBrowserIt
 			}
 		}
 
-		auto AddAndActivateInverseFilter = [this, &ExecuteOnFilterChanged](const TSharedRef<FFilterBase<FAssetFilterType>>& InFilter) 
-		{
-			int32 ExistingIndex = Filters.IndexOfByPredicate([InFilter](TSharedPtr<SFilter> Filter) { return Filter->GetFrontendFilter() == InFilter; });
-			TSharedRef<SFilter> FilterWidget = ExistingIndex == INDEX_NONE ? AddFilterToBar(InFilter) : Filters[ExistingIndex];
-			FilterWidget->SetEnabled(true, false);
-			SetFrontendFilterActive(InFilter, true);
-			ExecuteOnFilterChanged = true;
-		};
-		
-		TSharedPtr<FFilter_ShowOtherDevelopers> OtherDevelopersFilter = StaticCastSharedPtr<FFilter_ShowOtherDevelopers>(GetFrontendFilter("ShowOtherDevelopersBackend"));
-		// Special case: if item is hidden because of "other developers" filter, disable it
+		TSharedPtr<FFilter_HideOtherDevelopers> OtherDevelopersFilter = StaticCastSharedPtr<FFilter_HideOtherDevelopers>(GetFrontendFilter("HideOtherDevelopersBackend"));
+		// Special case: if item is hidden because of "hide other developers" filter, disable it
 		if (OtherDevelopersFilter.IsValid() && OtherDevelopersFilter->IsActive())
 		{
 			TSharedRef<const FPathPermissionList> PermissionList = OtherDevelopersFilter->GetPathPermissionList();	
@@ -287,12 +278,26 @@ void SFilterList::DisableFiltersThatHideItems(TArrayView<const FContentBrowserIt
 			{
 				if (PermissionList->PassesStartsWithFilter(WriteToString<256>(Item.GetInternalPath())))
 				{
-					AddAndActivateInverseFilter(OtherDevelopersFilter.ToSharedRef());
+					TSharedRef<FFilter_HideOtherDevelopers>	OtherDevelopersFilterAsRef = OtherDevelopersFilter.ToSharedRef();
+					int32 ExistingIndex = Filters.IndexOfByPredicate([OtherDevelopersFilterAsRef](TSharedPtr<SFilter> Filter) { return Filter->GetFrontendFilter() == OtherDevelopersFilterAsRef; });
+					TSharedRef<SFilter> FilterWidget = ExistingIndex == INDEX_NONE ? AddFilterToBar(OtherDevelopersFilterAsRef) : Filters[ExistingIndex];
+					FilterWidget->SetEnabled(false, false);
+					SetFrontendFilterActive(OtherDevelopersFilterAsRef, false);
+					ExecuteOnFilterChanged = true;
 					break;
 				}
 			}
 		}
 
+			auto AddAndActivateInverseFilter = [this, &ExecuteOnFilterChanged](const TSharedRef<FFilterBase<FAssetFilterType>>& InFilter) 
+		{
+			int32 ExistingIndex = Filters.IndexOfByPredicate([InFilter](TSharedPtr<SFilter> Filter) { return Filter->GetFrontendFilter() == InFilter; });
+			TSharedRef<SFilter> FilterWidget = ExistingIndex == INDEX_NONE ? AddFilterToBar(InFilter) : Filters[ExistingIndex];
+			FilterWidget->SetEnabled(true, false);
+			SetFrontendFilterActive(InFilter, true);
+			ExecuteOnFilterChanged = true;
+		};
+
 		// Special case: if the object is a redirector then enable the 'show redirectors' filter - this will also prevent
 		// folders that contain only redirectors from being hidden with the "hide empty folders" setting
 		FString RedirectorClassPath = UObjectRedirector::StaticClass()->GetPathName();
diff --git a/Engine/Source/Editor/ContentBrowser/Private/SFilterList.h b/Engine/Source/Editor/ContentBrowser/Private/SFilterList.h
--- a/Engine/Source/Editor/ContentBrowser/Private/SFilterList.h
+++ b/Engine/Source/Editor/ContentBrowser/Private/SFilterList.h
@@ -23,7 +23,6 @@
 #include "UObject/NameTypes.h"
 #include "Widgets/DeclarativeSyntaxSupport.h"
 
-class FFilter_ShowOtherDevelopers;
 class FFrontendFilter;
 class SWidget;
 class UClass;
diff --git a/Engine/Source/Editor/ContentBrowser/Public/FrontendFilters.h b/Engine/Source/Editor/ContentBrowser/Public/FrontendFilters.h
--- a/Engine/Source/Editor/ContentBrowser/Public/FrontendFilters.h
+++ b/Engine/Source/Editor/ContentBrowser/Public/FrontendFilters.h
@@ -234,7 +234,7 @@ public:
 };
 
 /** An inverse filter that allows display of content in developer folders that are not the current user's */
-class UE_DEPRECATED(5.5, "This frontend filter has been deprecated and replaced with backend filtering. see FFilter_ShowOtherDevelopers.")
+class UE_DEPRECATED(5.5, "This frontend filter has been deprecated and replaced with backend filtering. see FFilter_HideOtherDevelopers.")
 CONTENTBROWSER_API FFrontendFilter_ShowOtherDevelopers : public FFrontendFilter
 {
 public:
