[ContentBrowser] Fix for UI regression

#rnx

diff --git a/Engine/Source/Editor/ContentBrowser/Private/SContentBrowser.cpp b/Engine/Source/Editor/ContentBrowser/Private/SContentBrowser.cpp
--- a/Engine/Source/Editor/ContentBrowser/Private/SContentBrowser.cpp
+++ b/Engine/Source/Editor/ContentBrowser/Private/SContentBrowser.cpp
@@ -1264,11 +1264,7 @@ void SContentBrowser::OnFilterBarLayoutChanging(EFilterBarLayout NewLayout)
 			]
 
 			+ SHorizontalBox::Slot()
-#if UE_CONTENTBROWSER_NEW_STYLE
 			.FillWidth(1.0)
-#else
-			.AutoWidth()
-#endif
 			.VAlign(VAlign_Top)
 			.Padding(0, 4, 0, 0)
 			[
@@ -1376,7 +1372,7 @@ void SContentBrowser::OnFilterBarLayoutChanging(EFilterBarLayout NewLayout)
 					]
 
 					+ SHorizontalBox::Slot()
-					.FillWidth(1.0f)
+					.AutoWidth()
 					.MaxWidth(SearchBoxMaxWidth)
 					[
 						SNew(SBox)
