[ContentBrowser] ViewOptions menu relocation/refactor

#rnx

diff --git a/Engine/Source/Editor/ContentBrowser/Private/Menus/AssetViewMenus.cpp b/Engine/Source/Editor/ContentBrowser/Private/Menus/AssetViewMenus.cpp
new file mode 100644
--- /dev/null
+++ b/Engine/Source/Editor/ContentBrowser/Private/Menus/AssetViewMenus.cpp
@@ -0,0 +1,385 @@
+// Copyright Epic Games, Inc. All Rights Reserved.
+
+#include "ContentBrowserMenuContexts.h"
+#include "ContentBrowserMenuUtils.h"
+#include "Misc/DelayedAutoRegister.h"
+#include "SAssetView.h"
+#include "SContentBrowser.h"
+#include "SFilterList.h"
+#include "ToolMenu.h"
+#include "ToolMenus.h"
+#include "UObject/NameTypes.h"
+
+#define LOCTEXT_NAMESPACE "ContentBrowser"
+
+FDelayedAutoRegisterHelper SAssetView::AssetViewOptionsMenuRegistration(
+	EDelayedRegisterRunPhase::EndOfEngineInit,
+	[]
+	{
+		static const FName AssetViewOptionsMenuName("ContentBrowser.AssetViewOptions");
+		if (!ensureAlways(!UToolMenus::Get()->IsMenuRegistered(AssetViewOptionsMenuName)))
+		{
+			return;
+		}
+
+		FToolMenuOwnerScoped OwnerScoped(UE_MODULE_NAME);
+
+		UToolMenu* Menu = UToolMenus::Get()->RegisterMenu(AssetViewOptionsMenuName);
+		Menu->bCloseSelfOnly = true;
+
+		Menu->AddDynamicSection(NAME_None, FNewToolMenuDelegate::CreateLambda([](UToolMenu* InMenu)
+		{
+			FName OwningContentBrowserName = NAME_None;
+			FFiltersAdditionalParams FilterParams;
+			if (UContentBrowserAssetViewContextMenuContext* Context = InMenu->FindContext<UContentBrowserAssetViewContextMenuContext>())
+			{
+				if (Context->AssetView.IsValid())
+				{
+					const TSharedRef<SAssetView> AssetView = Context->AssetView.Pin().ToSharedRef();
+
+					// AssetViewType Section
+					{
+						FToolMenuSection& Section = InMenu->AddSection("AssetViewType", LOCTEXT("ViewTypeHeading", "View Type"));
+						Section.AddMenuEntry(
+							"TileView",
+							LOCTEXT("TileViewOption", "Tiles"),
+							LOCTEXT("TileViewOptionToolTip", "View assets as tiles in a grid."),
+							FSlateIcon(),
+							FUIAction(
+								FExecuteAction::CreateSP(AssetView, &SAssetView::SetCurrentViewTypeFromMenu, EAssetViewType::Tile),
+								FCanExecuteAction(),
+								FIsActionChecked::CreateSP(AssetView, &SAssetView::IsCurrentViewType, EAssetViewType::Tile)
+							),
+							EUserInterfaceActionType::RadioButton
+						);
+
+						Section.AddMenuEntry(
+							"ListView",
+							LOCTEXT("ListViewOption", "List"),
+							LOCTEXT("ListViewOptionToolTip", "View assets in a list with thumbnails."),
+							FSlateIcon(),
+							FUIAction(
+								FExecuteAction::CreateSP(AssetView, &SAssetView::SetCurrentViewTypeFromMenu, EAssetViewType::List),
+								FCanExecuteAction(),
+								FIsActionChecked::CreateSP(AssetView, &SAssetView::IsCurrentViewType, EAssetViewType::List)
+							),
+							EUserInterfaceActionType::RadioButton
+						);
+
+						Section.AddMenuEntry(
+							"ColumnView",
+							LOCTEXT("ColumnViewOption", "Columns"),
+							LOCTEXT("ColumnViewOptionToolTip", "View assets in a list with columns of details."),
+							FSlateIcon(),
+							FUIAction(
+								FExecuteAction::CreateSP(AssetView, &SAssetView::SetCurrentViewTypeFromMenu, EAssetViewType::Column),
+								FCanExecuteAction(),
+								FIsActionChecked::CreateSP(AssetView, &SAssetView::IsCurrentViewType, EAssetViewType::Column)
+							),
+							EUserInterfaceActionType::RadioButton
+							);
+
+						const FText CustomViewLabel = AssetView->ViewExtender ? AssetView->ViewExtender->GetViewDisplayName() : LOCTEXT("CustomViewOption", "Custom");
+						const FText CustomViewTooltip = AssetView->ViewExtender ? AssetView->ViewExtender->GetViewTooltipText() : LOCTEXT("CustomViewOptionToolTip", "A user specified custom view.");
+
+						Section.AddMenuEntry(
+							"CustomView",
+							CustomViewLabel,
+							CustomViewTooltip,
+							FSlateIcon(),
+							FUIAction(
+								FExecuteAction::CreateSP(AssetView, &SAssetView::SetCurrentViewTypeFromMenu, EAssetViewType::Custom),
+								FCanExecuteAction(),
+								FIsActionChecked::CreateSP(AssetView, &SAssetView::IsCurrentViewType, EAssetViewType::Custom),
+								FIsActionButtonVisible::CreateSP(AssetView, &SAssetView::IsCustomViewSet)
+							),
+							EUserInterfaceActionType::RadioButton
+						);
+					}
+
+					// Filter Bar Section
+					if (TSharedPtr<SFilterList> FilterBarPinned = AssetView->FilterBar.Pin())
+					{
+						FToolMenuSection& Section = InMenu->AddSection("FilterBar", LOCTEXT("FilterBarHeading", "Filter Display"));
+
+						Section.AddMenuEntry(
+							"VerticalLayout",
+							LOCTEXT("FilterListVerticalLayout", "Vertical"),
+							LOCTEXT("FilterListVerticalLayoutToolTip", "Swap to a vertical layout for the filter bar"),
+							FSlateIcon(),
+							FUIAction(
+								FExecuteAction::CreateLambda([FilterBarPinned]()
+								{
+									if (FilterBarPinned->GetFilterLayout() != EFilterBarLayout::Vertical)
+									{
+										FilterBarPinned->SetFilterLayout(EFilterBarLayout::Vertical);
+									}
+								}),
+								FCanExecuteAction(),
+								FIsActionChecked::CreateLambda([FilterBarPinned]()
+								{
+									return FilterBarPinned->GetFilterLayout() == EFilterBarLayout::Vertical;
+								})),
+							EUserInterfaceActionType::RadioButton
+						);
+
+						Section.AddMenuEntry(
+							"HorizontalLayout",
+							LOCTEXT("FilterListHorizontalLayout", "Horizontal"),
+							LOCTEXT("FilterListHorizontalLayoutToolTip", "Swap to a Horizontal layout for the filter bar"),
+							FSlateIcon(),
+							FUIAction(
+								FExecuteAction::CreateLambda([FilterBarPinned]()
+								{
+									if (FilterBarPinned->GetFilterLayout() != EFilterBarLayout::Horizontal)
+									{
+										FilterBarPinned->SetFilterLayout(EFilterBarLayout::Horizontal);
+									}
+								}),
+								FCanExecuteAction(),
+								FIsActionChecked::CreateLambda([FilterBarPinned]()
+								{
+									return FilterBarPinned->GetFilterLayout() == EFilterBarLayout::Horizontal;
+								})),
+							EUserInterfaceActionType::RadioButton
+						);
+					}
+
+					// View Section
+					{
+						FToolMenuSection& Section = InMenu->AddSection("View", LOCTEXT("ViewHeading", "View"));
+
+						Section.AddMenuEntry(
+							"ShowFolders",
+							LOCTEXT("ShowFoldersOption", "Show Folders"),
+							LOCTEXT("ShowFoldersOptionToolTip", "Show folders in the view as well as assets?"),
+							FSlateIcon(),
+							FUIAction(
+								FExecuteAction::CreateSP(AssetView, &SAssetView::ToggleShowFolders),
+								FCanExecuteAction::CreateSP(AssetView, &SAssetView::IsToggleShowFoldersAllowed),
+								FIsActionChecked::CreateSP(AssetView, &SAssetView::IsShowingFolders)
+							),
+							EUserInterfaceActionType::ToggleButton
+						);
+
+						Section.AddMenuEntry(
+							"ShowEmptyFolders",
+							LOCTEXT("ShowEmptyFoldersOption", "Show Empty Folders"),
+							LOCTEXT("ShowEmptyFoldersOptionToolTip", "Show empty folders in the view as well as assets?"),
+							FSlateIcon(),
+							FUIAction(
+								FExecuteAction::CreateSP(AssetView, &SAssetView::ToggleShowEmptyFolders),
+								FCanExecuteAction::CreateSP(AssetView, &SAssetView::IsToggleShowEmptyFoldersAllowed),
+								FIsActionChecked::CreateSP(AssetView, &SAssetView::IsShowingEmptyFolders)
+							),
+							EUserInterfaceActionType::ToggleButton
+						);
+
+						Section.AddMenuEntry(
+							"ShowFavorite",
+							LOCTEXT("ShowFavoriteOptions", "Show Favorites"),
+							LOCTEXT("ShowFavoriteOptionToolTip", "Show the favorite folders in the view?"),
+							FSlateIcon(),
+							FUIAction(
+								FExecuteAction::CreateSP(AssetView, &SAssetView::ToggleShowFavorites),
+								FCanExecuteAction::CreateSP(AssetView, &SAssetView::IsToggleShowFavoritesAllowed),
+								FIsActionChecked::CreateSP(AssetView, &SAssetView::IsShowingFavorites)
+							),
+							EUserInterfaceActionType::ToggleButton
+						);
+
+						Section.AddMenuEntry(
+							"FilterRecursively",
+							LOCTEXT("FilterRecursivelyOption", "Filter Recursively"),
+							LOCTEXT("FilterRecursivelyOptionToolTip", "Should filters apply recursively in the view?"),
+							FSlateIcon(),
+							FUIAction(
+								FExecuteAction::CreateSP(AssetView, &SAssetView::ToggleFilteringRecursively),
+								FCanExecuteAction::CreateSP(AssetView, &SAssetView::IsToggleFilteringRecursivelyAllowed),
+								FIsActionChecked::CreateSP(AssetView, &SAssetView::IsFilteringRecursively)
+							),
+							EUserInterfaceActionType::ToggleButton
+						);
+
+						Section.AddMenuEntry(
+							"ShowAllFolder",
+							LOCTEXT("ShowAllFolderOption", "Show All Folder"),
+							LOCTEXT("ShowAllFolderOptionToolTip", "Show the all folder in the view?"),
+							FSlateIcon(),
+							FUIAction(
+								FExecuteAction::CreateSP(AssetView, &SAssetView::ToggleShowAllFolder),
+								FCanExecuteAction(),
+								FIsActionChecked::CreateSP(AssetView, &SAssetView::IsShowingAllFolder)
+							),
+							EUserInterfaceActionType::ToggleButton
+						);
+
+						Section.AddMenuEntry(
+							"OrganizeFolders",
+							LOCTEXT("OrganizeFoldersOption", "Organize Folders"),
+							LOCTEXT("OrganizeFoldersOptionToolTip", "Organize folders in the view?"),
+							FSlateIcon(),
+							FUIAction(
+								FExecuteAction::CreateSP(AssetView, &SAssetView::ToggleOrganizeFolders),
+								FCanExecuteAction(),
+								FIsActionChecked::CreateSP(AssetView, &SAssetView::IsOrganizingFolders)
+							),
+							EUserInterfaceActionType::ToggleButton
+						);
+
+						if (AssetView->bShowPathViewFilters)
+						{
+							Section.AddSubMenu(
+								"PathViewFilters",
+								LOCTEXT("PathViewFilters", "Path View Filters"),
+								LOCTEXT("PathViewFilters_ToolTip", "Path View Filters"),
+								FNewToolMenuDelegate());
+						}
+					}
+
+					// Search Section
+					{
+						FToolMenuSection& Section = InMenu->AddSection("Search", LOCTEXT("SearchHeading", "Search"));
+						Section.AddMenuEntry(
+							"IncludeClassName",
+							LOCTEXT("IncludeClassNameOption", "Search Asset Class Names"),
+							LOCTEXT("IncludeClassesNameOptionTooltip", "Include asset type names in search criteria?  (e.g. Blueprint, Texture, Sound)"),
+							FSlateIcon(),
+							FUIAction(
+								FExecuteAction::CreateSP(AssetView, &SAssetView::ToggleIncludeClassNames),
+								FCanExecuteAction::CreateSP(AssetView, &SAssetView::IsToggleIncludeClassNamesAllowed),
+								FIsActionChecked::CreateSP(AssetView, &SAssetView::IsIncludingClassNames)
+							),
+							EUserInterfaceActionType::ToggleButton
+						);
+
+						Section.AddMenuEntry(
+							"IncludeAssetPath",
+							LOCTEXT("IncludeAssetPathOption", "Search Asset Path"),
+							LOCTEXT("IncludeAssetPathOptionTooltip", "Include entire asset path in search criteria?"),
+							FSlateIcon(),
+							FUIAction(
+								FExecuteAction::CreateSP(AssetView, &SAssetView::ToggleIncludeAssetPaths),
+								FCanExecuteAction::CreateSP(AssetView, &SAssetView::IsToggleIncludeAssetPathsAllowed),
+								FIsActionChecked::CreateSP(AssetView, &SAssetView::IsIncludingAssetPaths)
+							),
+							EUserInterfaceActionType::ToggleButton
+						);
+
+						Section.AddMenuEntry(
+							"IncludeCollectionName",
+							LOCTEXT("IncludeCollectionNameOption", "Search Collection Names"),
+							LOCTEXT("IncludeCollectionNameOptionTooltip", "Include Collection names in search criteria?"),
+							FSlateIcon(),
+							FUIAction(
+								FExecuteAction::CreateSP(AssetView, &SAssetView::ToggleIncludeCollectionNames),
+								FCanExecuteAction::CreateSP(AssetView, &SAssetView::IsToggleIncludeCollectionNamesAllowed),
+								FIsActionChecked::CreateSP(AssetView, &SAssetView::IsIncludingCollectionNames)
+							),
+							EUserInterfaceActionType::ToggleButton
+						);
+					}
+
+					// Asset Thumbnails Section
+					{
+						FToolMenuSection& Section = InMenu->AddSection("AssetThumbnails", LOCTEXT("ThumbnailsHeading", "Thumbnails"));
+
+						auto CreateThumbnailSizeSubMenu = [AssetView](UToolMenu* SubMenu)
+						{
+							FToolMenuSection& SizeSection = SubMenu->AddSection("ThumbnailSizes");
+
+							for (int32 EnumValue = (int32)EThumbnailSize::Tiny; EnumValue < (int32)EThumbnailSize::MAX; ++EnumValue)
+							{
+								if ((EThumbnailSize)EnumValue == EThumbnailSize::XLarge)
+								{
+									continue;
+								}
+
+								SizeSection.AddMenuEntry(
+									NAME_None,
+									SAssetView::ThumbnailSizeToDisplayName((EThumbnailSize)EnumValue),
+									FText::GetEmpty(),
+									FSlateIcon(),
+									FUIAction(
+										FExecuteAction::CreateSP(AssetView, &SAssetView::OnThumbnailSizeChanged, (EThumbnailSize)EnumValue),
+										FCanExecuteAction::CreateSP(AssetView, &SAssetView::IsThumbnailScalingAllowed),
+										FIsActionChecked::CreateSP(AssetView, &SAssetView::IsThumbnailSizeChecked, (EThumbnailSize)EnumValue)
+									),
+									EUserInterfaceActionType::RadioButton
+								);
+							}
+						};
+						Section.AddEntry(FToolMenuEntry::InitSubMenu(
+							"ThumbnailSize",
+							LOCTEXT("ThumbnailSize", "Thumbnail Size"),
+							LOCTEXT("ThumbnailSizeToolTip", "Adjust the size of thumbnails."),
+							FNewToolMenuDelegate::CreateLambda(CreateThumbnailSizeSubMenu)
+						));
+
+						Section.AddMenuEntry(
+							"ThumbnailEditMode",
+							LOCTEXT("ThumbnailEditModeOption", "Thumbnail Edit Mode"),
+							LOCTEXT("ThumbnailEditModeOptionToolTip", "Toggle thumbnail editing mode. When in this mode you can rotate the camera on 3D thumbnails by dragging them."),
+							FSlateIcon(),
+							FUIAction(
+								FExecuteAction::CreateSP(AssetView, &SAssetView::ToggleThumbnailEditMode),
+								FCanExecuteAction::CreateSP(AssetView, &SAssetView::IsThumbnailEditModeAllowed),
+								FIsActionChecked::CreateSP(AssetView, &SAssetView::IsThumbnailEditMode)
+							),
+							EUserInterfaceActionType::ToggleButton
+						);
+
+						Section.AddMenuEntry(
+							"RealTimeThumbnails",
+							LOCTEXT("RealTimeThumbnailsOption", "Real-Time Thumbnails"),
+							LOCTEXT("RealTimeThumbnailsOptionToolTip", "Renders the assets thumbnails in real-time"),
+							FSlateIcon(),
+							FUIAction(
+								FExecuteAction::CreateSP(AssetView, &SAssetView::ToggleRealTimeThumbnails),
+								FCanExecuteAction::CreateSP(AssetView, &SAssetView::CanShowRealTimeThumbnails),
+								FIsActionChecked::CreateSP(AssetView, &SAssetView::IsShowingRealTimeThumbnails)
+							),
+							EUserInterfaceActionType::ToggleButton
+						);
+					}
+
+					// Asset Columns Section
+					if (AssetView->GetColumnViewVisibility() == EVisibility::Visible)
+					{
+						{
+							FToolMenuSection& Section = InMenu->AddSection("AssetColumns", LOCTEXT("ToggleColumnsHeading", "Columns"));
+
+							Section.AddMenuEntry(
+								"ResetColumns",
+								LOCTEXT("ResetColumns", "Reset Columns"),
+								LOCTEXT("ResetColumnsToolTip", "Reset all columns to be visible again."),
+								FSlateIcon(),
+								FUIAction(FExecuteAction::CreateSP(AssetView, &SAssetView::ResetColumns)),
+								EUserInterfaceActionType::Button
+							);
+
+							Section.AddMenuEntry(
+								"ExportColumns",
+								LOCTEXT("ExportColumns", "Export to CSV"),
+								LOCTEXT("ExportColumnsToolTip", "Export column data to CSV."),
+								FSlateIcon(),
+								FUIAction(FExecuteAction::CreateSP(AssetView, &SAssetView::ExportColumns)),
+								EUserInterfaceActionType::Button
+							);
+						}
+					}
+
+					AssetView->PopulateFilterAdditionalParams(FilterParams);
+				}
+
+				if (Context->OwningContentBrowser.IsValid())
+				{
+					OwningContentBrowserName = Context->OwningContentBrowser.Pin()->GetInstanceName();
+				}
+			}
+
+			ContentBrowserMenuUtils::AddFiltersToMenu(InMenu, OwningContentBrowserName, FilterParams);
+		}));
+	});
+
+#undef LOCTEXT_NAMESPACE
diff --git a/Engine/Source/Editor/ContentBrowser/Private/Menus/ContentBrowserMenus.cpp b/Engine/Source/Editor/ContentBrowser/Private/Menus/ContentBrowserMenus.cpp
new file mode 100644
--- /dev/null
+++ b/Engine/Source/Editor/ContentBrowser/Private/Menus/ContentBrowserMenus.cpp
@@ -0,0 +1,47 @@
+// Copyright Epic Games, Inc. All Rights Reserved.
+
+#pragma once
+
+#include "ContentBrowserMenuContexts.h"
+#include "Misc/DelayedAutoRegister.h"
+#include "SContentBrowser.h"
+#include "ToolMenu.h"
+#include "ToolMenus.h"
+#include "UObject/NameTypes.h"
+
+#define LOCTEXT_NAMESPACE "ContentBrowser"
+
+FDelayedAutoRegisterHelper SContentBrowser::PathViewFiltersMenuRegistration(
+	EDelayedRegisterRunPhase::EndOfEngineInit,
+	[]
+	{
+		static const FName PathViewFiltersMenuName("ContentBrowser.AssetViewOptions.PathViewFilters");
+		if (!ensureAlways(!UToolMenus::Get()->IsMenuRegistered(PathViewFiltersMenuName)))
+		{
+			return;
+		}
+
+		FToolMenuOwnerScoped OwnerScoped(UE_MODULE_NAME);
+
+		UToolMenu* Menu = UToolMenus::Get()->RegisterMenu(PathViewFiltersMenuName);
+
+		Menu->AddDynamicSection(NAME_None, FNewToolMenuDelegate::CreateLambda([](UToolMenu* InMenu)
+		{
+			if (const UContentBrowserAssetViewContextMenuContext* AssetViewContext = InMenu->FindContext<UContentBrowserAssetViewContextMenuContext>())
+			{
+				if (const TSharedPtr<SContentBrowser> ContentBrowser = AssetViewContext->OwningContentBrowser.Pin())
+				{
+					ContentBrowser->PopulatePathViewFiltersMenu(InMenu);
+				}
+			}
+			else if (const UContentBrowserMenuContext* ContentBrowserContext = InMenu->FindContext<UContentBrowserMenuContext>())
+			{
+				if (const TSharedPtr<SContentBrowser> ContentBrowser = ContentBrowserContext->ContentBrowser.Pin())
+				{
+					ContentBrowser->PopulatePathViewFiltersMenu(InMenu);
+				}
+			}
+		}));
+	});
+
+#undef LOCTEXT_NAMESPACE
diff --git a/Engine/Source/Editor/ContentBrowser/Private/SAssetView.cpp b/Engine/Source/Editor/ContentBrowser/Private/SAssetView.cpp
--- a/Engine/Source/Editor/ContentBrowser/Private/SAssetView.cpp
+++ b/Engine/Source/Editor/ContentBrowser/Private/SAssetView.cpp
@@ -3764,8 +3764,6 @@ FLinearColor SAssetView::GetThumbnailHintColorAndOpacity() const
 
 TSharedRef<SWidget> SAssetView::GetViewButtonContent()
 {
-	SAssetView::RegisterGetViewButtonMenu();
-
 	// Get all menu extenders for this context menu from the content browser module
 	FContentBrowserModule& ContentBrowserModule = FModuleManager::GetModuleChecked<FContentBrowserModule>( TEXT("ContentBrowser") );
 	const TArray<FContentBrowserMenuExtender>& MenuExtenderDelegates = ContentBrowserModule.GetAllAssetViewViewMenuExtenders();
@@ -3801,361 +3799,6 @@ TSharedRef<SWidget> SAssetView::GetViewButtonContent()
 	return UToolMenus::Get()->GenerateWidget("ContentBrowser.AssetViewOptions", MenuContext);
 }
 
-void SAssetView::RegisterGetViewButtonMenu()
-{
-	if (!UToolMenus::Get()->IsMenuRegistered("ContentBrowser.AssetViewOptions"))
-	{
-		UToolMenu* Menu = UToolMenus::Get()->RegisterMenu("ContentBrowser.AssetViewOptions");
-		Menu->bCloseSelfOnly = true;
-		Menu->AddDynamicSection("DynamicContent", FNewToolMenuDelegate::CreateLambda([](UToolMenu* InMenu)
-		{
-			FName OwningContentBrowserName = NAME_None;
-			FFiltersAdditionalParams FilterParams;
-			if (UContentBrowserAssetViewContextMenuContext* Context = InMenu->FindContext<UContentBrowserAssetViewContextMenuContext>())
-			{
-				if (Context->AssetView.IsValid())
-				{
-					const TSharedPtr<SAssetView> AssetView = Context->AssetView.Pin();
-					AssetView->PopulateViewButtonMenu(InMenu);
-					AssetView->PopulateFilterAdditionalParams(FilterParams);
-				}
-
-				if (Context->OwningContentBrowser.IsValid())
-				{
-					OwningContentBrowserName = Context->OwningContentBrowser.Pin()->GetInstanceName();
-				}
-			}
-
-			ContentBrowserMenuUtils::AddFiltersToMenu(InMenu, OwningContentBrowserName, FilterParams);
-		}));
-	}
-}
-
-void SAssetView::PopulateViewButtonMenu(UToolMenu* Menu)
-{
-	{
-		FToolMenuSection& Section = Menu->AddSection("AssetViewType", LOCTEXT("ViewTypeHeading", "View Type"));
-		Section.AddMenuEntry(
-			"TileView",
-			LOCTEXT("TileViewOption", "Tiles"),
-			LOCTEXT("TileViewOptionToolTip", "View assets as tiles in a grid."),
-			FSlateIcon(),
-			FUIAction(
-				FExecuteAction::CreateSP( this, &SAssetView::SetCurrentViewTypeFromMenu, EAssetViewType::Tile ),
-				FCanExecuteAction(),
-				FIsActionChecked::CreateSP( this, &SAssetView::IsCurrentViewType, EAssetViewType::Tile )
-				),
-			EUserInterfaceActionType::RadioButton
-			);
-
-		Section.AddMenuEntry(
-			"ListView",
-			LOCTEXT("ListViewOption", "List"),
-			LOCTEXT("ListViewOptionToolTip", "View assets in a list with thumbnails."),
-			FSlateIcon(),
-			FUIAction(
-				FExecuteAction::CreateSP( this, &SAssetView::SetCurrentViewTypeFromMenu, EAssetViewType::List ),
-				FCanExecuteAction(),
-				FIsActionChecked::CreateSP( this, &SAssetView::IsCurrentViewType, EAssetViewType::List )
-				),
-			EUserInterfaceActionType::RadioButton
-			);
-
-		Section.AddMenuEntry(
-			"ColumnView",
-			LOCTEXT("ColumnViewOption", "Columns"),
-			LOCTEXT("ColumnViewOptionToolTip", "View assets in a list with columns of details."),
-			FSlateIcon(),
-			FUIAction(
-				FExecuteAction::CreateSP( this, &SAssetView::SetCurrentViewTypeFromMenu, EAssetViewType::Column ),
-				FCanExecuteAction(),
-				FIsActionChecked::CreateSP( this, &SAssetView::IsCurrentViewType, EAssetViewType::Column )
-				),
-			EUserInterfaceActionType::RadioButton
-			);
-
-		const FText CustomViewLabel = ViewExtender ? ViewExtender->GetViewDisplayName() : LOCTEXT("CustomViewOption", "Custom");
-		const FText CustomViewTooltip = ViewExtender ? ViewExtender->GetViewTooltipText() : LOCTEXT("CustomViewOptionToolTip", "A user specified custom view.");
-
-		Section.AddMenuEntry(
-			"CustomView",
-			CustomViewLabel,
-			CustomViewTooltip,
-			FSlateIcon(),
-			FUIAction(
-				FExecuteAction::CreateSP( this, &SAssetView::SetCurrentViewTypeFromMenu, EAssetViewType::Custom ),
-				FCanExecuteAction(),
-				FIsActionChecked::CreateSP( this, &SAssetView::IsCurrentViewType, EAssetViewType::Custom ),
-				FIsActionButtonVisible::CreateSP(this, &SAssetView::IsCustomViewSet)
-				),
-			EUserInterfaceActionType::RadioButton
-			);
-	}
-
-	if(TSharedPtr<SFilterList> FilterBarPinned = FilterBar.Pin())
-	{
-		FToolMenuSection& Section = Menu->AddSection("FilterBar", LOCTEXT("FilterBarHeading", "Filter Display"));
-
-		Section.AddMenuEntry(
-		"VerticalLayout",
-		LOCTEXT("FilterListVerticalLayout", "Vertical"),
-		LOCTEXT("FilterListVerticalLayoutToolTip", "Swap to a vertical layout for the filter bar"),
-		FSlateIcon(),
-		FUIAction(
-			FExecuteAction::CreateLambda([FilterBarPinned]()
-			{
-				if(FilterBarPinned->GetFilterLayout() != EFilterBarLayout::Vertical)
-				{
-					FilterBarPinned->SetFilterLayout(EFilterBarLayout::Vertical);
-				}
-			}),
-			FCanExecuteAction(),
-			FIsActionChecked::CreateLambda([FilterBarPinned]() { return FilterBarPinned->GetFilterLayout() == EFilterBarLayout::Vertical; })),
-		EUserInterfaceActionType::RadioButton
-	);
-
-		Section.AddMenuEntry(
-			"HorizontalLayout",
-			LOCTEXT("FilterListHorizontalLayout", "Horizontal"),
-			LOCTEXT("FilterListHorizontalLayoutToolTip", "Swap to a Horizontal layout for the filter bar"),
-			FSlateIcon(),
-			FUIAction(
-			FExecuteAction::CreateLambda([FilterBarPinned]()
-				{
-					if(FilterBarPinned->GetFilterLayout() != EFilterBarLayout::Horizontal)
-					{
-						FilterBarPinned->SetFilterLayout(EFilterBarLayout::Horizontal);
-					}
-				}),
-				FCanExecuteAction(),
-				FIsActionChecked::CreateLambda([FilterBarPinned]() { return FilterBarPinned->GetFilterLayout() == EFilterBarLayout::Horizontal; })),
-			EUserInterfaceActionType::RadioButton
-		);
-	}
-
-	{
-		FToolMenuSection& Section = Menu->AddSection("View", LOCTEXT("ViewHeading", "View"));
-
-		Section.AddMenuEntry(
-			"ShowFolders",
-			LOCTEXT("ShowFoldersOption", "Show Folders"),
-			LOCTEXT("ShowFoldersOptionToolTip", "Show folders in the view as well as assets?"),
-			FSlateIcon(),
-			FUIAction(
-				FExecuteAction::CreateSP( this, &SAssetView::ToggleShowFolders ),
-				FCanExecuteAction::CreateSP( this, &SAssetView::IsToggleShowFoldersAllowed ),
-				FIsActionChecked::CreateSP( this, &SAssetView::IsShowingFolders )
-			),
-			EUserInterfaceActionType::ToggleButton
-		);
-
-		Section.AddMenuEntry(
-			"ShowEmptyFolders",
-			LOCTEXT("ShowEmptyFoldersOption", "Show Empty Folders"),
-			LOCTEXT("ShowEmptyFoldersOptionToolTip", "Show empty folders in the view as well as assets?"),
-			FSlateIcon(),
-			FUIAction(
-				FExecuteAction::CreateSP(this, &SAssetView::ToggleShowEmptyFolders),
-				FCanExecuteAction::CreateSP(this, &SAssetView::IsToggleShowEmptyFoldersAllowed),
-				FIsActionChecked::CreateSP(this, &SAssetView::IsShowingEmptyFolders)
-			),
-			EUserInterfaceActionType::ToggleButton
-		);
-
-		Section.AddMenuEntry(
-			"ShowFavorite",
-			LOCTEXT("ShowFavoriteOptions", "Show Favorites"),
-			LOCTEXT("ShowFavoriteOptionToolTip", "Show the favorite folders in the view?"),
-			FSlateIcon(),
-			FUIAction(
-				FExecuteAction::CreateSP(this, &SAssetView::ToggleShowFavorites),
-				FCanExecuteAction::CreateSP(this, &SAssetView::IsToggleShowFavoritesAllowed),
-				FIsActionChecked::CreateSP(this, &SAssetView::IsShowingFavorites)
-			),
-			EUserInterfaceActionType::ToggleButton
-		);
-
-		Section.AddMenuEntry(
-			"FilterRecursively",
-			LOCTEXT("FilterRecursivelyOption", "Filter Recursively"),
-			LOCTEXT("FilterRecursivelyOptionToolTip", "Should filters apply recursively in the view?"),
-			FSlateIcon(),
-			FUIAction(
-				FExecuteAction::CreateSP(this, &SAssetView::ToggleFilteringRecursively),
-				FCanExecuteAction::CreateSP(this, &SAssetView::IsToggleFilteringRecursivelyAllowed),
-				FIsActionChecked::CreateSP(this, &SAssetView::IsFilteringRecursively)
-			),
-			EUserInterfaceActionType::ToggleButton
-		);
-
-		Section.AddMenuEntry(
-			"ShowAllFolder",
-			LOCTEXT("ShowAllFolderOption", "Show All Folder"),
-			LOCTEXT("ShowAllFolderOptionToolTip", "Show the all folder in the view?"),
-			FSlateIcon(),
-			FUIAction(
-				FExecuteAction::CreateSP(this, &SAssetView::ToggleShowAllFolder),
-				FCanExecuteAction(),
-				FIsActionChecked::CreateSP(this, &SAssetView::IsShowingAllFolder)
-			),
-			EUserInterfaceActionType::ToggleButton
-		);
-
-		Section.AddMenuEntry(
-			"OrganizeFolders",
-			LOCTEXT("OrganizeFoldersOption", "Organize Folders"),
-			LOCTEXT("OrganizeFoldersOptionToolTip", "Organize folders in the view?"),
-			FSlateIcon(),
-			FUIAction(
-				FExecuteAction::CreateSP(this, &SAssetView::ToggleOrganizeFolders),
-				FCanExecuteAction(),
-				FIsActionChecked::CreateSP(this, &SAssetView::IsOrganizingFolders)
-			),
-			EUserInterfaceActionType::ToggleButton
-		);
-
-		if (bShowPathViewFilters)
-		{
-			Section.AddSubMenu(
-				"PathViewFilters",
-				LOCTEXT("PathViewFilters", "Path View Filters"),
-				LOCTEXT("PathViewFilters_ToolTip", "Path View Filters"),
-				FNewToolMenuDelegate());
-		}
-	}
-
-	{
-		FToolMenuSection& Section = Menu->AddSection("Search", LOCTEXT("SearchHeading", "Search"));
-		Section.AddMenuEntry(
-			"IncludeClassName",
-			LOCTEXT("IncludeClassNameOption", "Search Asset Class Names"),
-			LOCTEXT("IncludeClassesNameOptionTooltip", "Include asset type names in search criteria?  (e.g. Blueprint, Texture, Sound)"),
-			FSlateIcon(),
-			FUIAction(
-				FExecuteAction::CreateSP(this, &SAssetView::ToggleIncludeClassNames),
-				FCanExecuteAction::CreateSP(this, &SAssetView::IsToggleIncludeClassNamesAllowed),
-				FIsActionChecked::CreateSP(this, &SAssetView::IsIncludingClassNames)
-			),
-			EUserInterfaceActionType::ToggleButton
-		);
-
-		Section.AddMenuEntry(
-			"IncludeAssetPath",
-			LOCTEXT("IncludeAssetPathOption", "Search Asset Path"),
-			LOCTEXT("IncludeAssetPathOptionTooltip", "Include entire asset path in search criteria?"),
-			FSlateIcon(),
-			FUIAction(
-				FExecuteAction::CreateSP(this, &SAssetView::ToggleIncludeAssetPaths),
-				FCanExecuteAction::CreateSP(this, &SAssetView::IsToggleIncludeAssetPathsAllowed),
-				FIsActionChecked::CreateSP(this, &SAssetView::IsIncludingAssetPaths)
-			),
-			EUserInterfaceActionType::ToggleButton
-		);
-
-		Section.AddMenuEntry(
-			"IncludeCollectionName",
-			LOCTEXT("IncludeCollectionNameOption", "Search Collection Names"),
-			LOCTEXT("IncludeCollectionNameOptionTooltip", "Include Collection names in search criteria?"),
-			FSlateIcon(),
-			FUIAction(
-				FExecuteAction::CreateSP(this, &SAssetView::ToggleIncludeCollectionNames),
-				FCanExecuteAction::CreateSP(this, &SAssetView::IsToggleIncludeCollectionNamesAllowed),
-				FIsActionChecked::CreateSP(this, &SAssetView::IsIncludingCollectionNames)
-			),
-			EUserInterfaceActionType::ToggleButton
-		);
-	}
-
-	{
-		FToolMenuSection& Section = Menu->AddSection("AssetThumbnails", LOCTEXT("ThumbnailsHeading", "Thumbnails"));
-
-		auto CreateThumbnailSizeSubMenu = [this](UToolMenu* SubMenu)
-		{
-			FToolMenuSection& SizeSection = SubMenu->AddSection("ThumbnailSizes");
-			
-			for (int32 EnumValue = (int32)EThumbnailSize::Tiny; EnumValue < (int32)EThumbnailSize::MAX; ++EnumValue)
-			{
-#if !UE_CONTENTBROWSER_NEW_STYLE
-				if ((EThumbnailSize)EnumValue == EThumbnailSize::XLarge)
-				{
-					continue;
-				}
-#endif
-
-				SizeSection.AddMenuEntry(
-					NAME_None,
-					SAssetView::ThumbnailSizeToDisplayName((EThumbnailSize)EnumValue),
-					FText::GetEmpty(),
-					FSlateIcon(),
-					FUIAction(
-						FExecuteAction::CreateSP(this, &SAssetView::OnThumbnailSizeChanged, (EThumbnailSize)EnumValue),
-						FCanExecuteAction::CreateSP(this, &SAssetView::IsThumbnailScalingAllowed),
-						FIsActionChecked::CreateSP(this, &SAssetView::IsThumbnailSizeChecked, (EThumbnailSize)EnumValue)
-					),
-					EUserInterfaceActionType::RadioButton
-				);
-			}
-		};
-		Section.AddEntry(FToolMenuEntry::InitSubMenu(
-			"ThumbnailSize",
-			LOCTEXT("ThumbnailSize", "Thumbnail Size"),
-			LOCTEXT("ThumbnailSizeToolTip", "Adjust the size of thumbnails."),
-			FNewToolMenuDelegate::CreateLambda(CreateThumbnailSizeSubMenu)
-		));
-
-		Section.AddMenuEntry(
-			"ThumbnailEditMode",
-			LOCTEXT("ThumbnailEditModeOption", "Thumbnail Edit Mode"),
-			LOCTEXT("ThumbnailEditModeOptionToolTip", "Toggle thumbnail editing mode. When in this mode you can rotate the camera on 3D thumbnails by dragging them."),
-			FSlateIcon(),
-			FUIAction(
-				FExecuteAction::CreateSP( this, &SAssetView::ToggleThumbnailEditMode ),
-				FCanExecuteAction::CreateSP( this, &SAssetView::IsThumbnailEditModeAllowed ),
-				FIsActionChecked::CreateSP( this, &SAssetView::IsThumbnailEditMode )
-				),
-			EUserInterfaceActionType::ToggleButton
-			);
-
-		Section.AddMenuEntry(
-			"RealTimeThumbnails",
-			LOCTEXT("RealTimeThumbnailsOption", "Real-Time Thumbnails"),
-			LOCTEXT("RealTimeThumbnailsOptionToolTip", "Renders the assets thumbnails in real-time"),
-			FSlateIcon(),
-			FUIAction(
-				FExecuteAction::CreateSP( this, &SAssetView::ToggleRealTimeThumbnails ),
-				FCanExecuteAction::CreateSP( this, &SAssetView::CanShowRealTimeThumbnails ),
-				FIsActionChecked::CreateSP( this, &SAssetView::IsShowingRealTimeThumbnails )
-			),
-			EUserInterfaceActionType::ToggleButton
-			);
-	}
-
-	if (GetColumnViewVisibility() == EVisibility::Visible)
-	{
-		{
-			FToolMenuSection& Section = Menu->AddSection("AssetColumns", LOCTEXT("ToggleColumnsHeading", "Columns"));
-
-			Section.AddMenuEntry(
-				"ResetColumns",
-				LOCTEXT("ResetColumns", "Reset Columns"),
-				LOCTEXT("ResetColumnsToolTip", "Reset all columns to be visible again."),
-				FSlateIcon(),
-				FUIAction(FExecuteAction::CreateSP(this, &SAssetView::ResetColumns)),
-				EUserInterfaceActionType::Button
-				);
-
-			Section.AddMenuEntry(
-				"ExportColumns",
-				LOCTEXT("ExportColumns", "Export to CSV"),
-				LOCTEXT("ExportColumnsToolTip", "Export column data to CSV."),
-				FSlateIcon(),
-				FUIAction(FExecuteAction::CreateSP(this, &SAssetView::ExportColumns)),
-				EUserInterfaceActionType::Button
-			);
-		}
-	}
-}
-
 void SAssetView::PopulateFilterAdditionalParams(FFiltersAdditionalParams& OutParams)
 {
 	OutParams.CanShowCPPClasses = FCanExecuteAction::CreateSP(this, &SAssetView::IsToggleShowCppContentAllowed);
diff --git a/Engine/Source/Editor/ContentBrowser/Private/SContentBrowser.cpp b/Engine/Source/Editor/ContentBrowser/Private/SContentBrowser.cpp
--- a/Engine/Source/Editor/ContentBrowser/Private/SContentBrowser.cpp
+++ b/Engine/Source/Editor/ContentBrowser/Private/SContentBrowser.cpp
@@ -1026,8 +1026,6 @@ void SContentBrowser::Construct( const FArguments& InArgs, const FName& InInstan
 	// Update the breadcrumb trail path
 	OnContentBrowserSettingsChanged(NAME_None);
 
-	RegisterPathViewFiltersMenu();
-
 	// Initialize the search options
 	HandleAssetViewSearchOptionsChanged();
 }
@@ -3389,32 +3387,6 @@ FText SContentBrowser::GetAddNewToolTipText() const
 	return LOCTEXT( "AddNewToolTip_NoPath", "No path is selected as an add target." );
 }
 
-void SContentBrowser::RegisterPathViewFiltersMenu()
-{
-	static const FName PathViewFiltersMenuName = TEXT("ContentBrowser.AssetViewOptions.PathViewFilters");
-	if (!UToolMenus::Get()->IsMenuRegistered(PathViewFiltersMenuName))
-	{
-		UToolMenu* Menu = UToolMenus::Get()->RegisterMenu(PathViewFiltersMenuName);
-		Menu->AddDynamicSection("DynamicContent", FNewToolMenuDelegate::CreateLambda([](UToolMenu* InMenu)
-		{
-			if (const UContentBrowserAssetViewContextMenuContext* AssetViewContext = InMenu->FindContext<UContentBrowserAssetViewContextMenuContext>())
-			{
-				if (TSharedPtr<SContentBrowser> ContentBrowser = AssetViewContext->OwningContentBrowser.Pin())
-				{
-					ContentBrowser->PopulatePathViewFiltersMenu(InMenu);
-				}
-			}
-			else if (const UContentBrowserMenuContext* ContentBrowserContext = InMenu->FindContext<UContentBrowserMenuContext>())
-			{
-				if (TSharedPtr<SContentBrowser> ContentBrowser = ContentBrowserContext->ContentBrowser.Pin())
-				{
-					ContentBrowser->PopulatePathViewFiltersMenu(InMenu);
-				}
-			}
-		}));
-	}
-}
-
 void SContentBrowser::PopulatePathViewFiltersMenu(UToolMenu* Menu)
 {
 	if (PathViewPtr.IsValid())
diff --git a/Engine/Source/Editor/ContentBrowser/Private/SContentBrowser.h b/Engine/Source/Editor/ContentBrowser/Private/SContentBrowser.h
--- a/Engine/Source/Editor/ContentBrowser/Private/SContentBrowser.h
+++ b/Engine/Source/Editor/ContentBrowser/Private/SContentBrowser.h
@@ -583,9 +583,6 @@ private:
 	/** Called when Asset View Options "Search" options change */
 	void HandleAssetViewSearchOptionsChanged();
 
-	/** Register menu for filtering path view */
-	static void RegisterPathViewFiltersMenu();
-
 	/** Fill menu for filtering path view with items */
 	void PopulatePathViewFiltersMenu(UToolMenu* Menu);
 
@@ -786,6 +783,10 @@ private:
 
 	/** Handle to a callback for when CVars change */
 	FConsoleVariableSinkHandle CVarSinkHandle;
+
+	/** Registers a sub-menu within the settings button menu */
+	static FDelayedAutoRegisterHelper PathViewFiltersMenuRegistration;
+
 public: 
 
 	/** The section of EditorPerProjectUserSettings in which to save content browser settings */
diff --git a/Engine/Source/Editor/ContentBrowser/Public/SAssetView.h b/Engine/Source/Editor/ContentBrowser/Public/SAssetView.h
--- a/Engine/Source/Editor/ContentBrowser/Public/SAssetView.h
+++ b/Engine/Source/Editor/ContentBrowser/Public/SAssetView.h
@@ -561,11 +561,8 @@ private:
 	/** Returns the thumbnails hint color and opacity */
 	FLinearColor GetThumbnailHintColorAndOpacity() const;
 
-	/** Register menu for when the view combo button is clicked */
-	static void RegisterGetViewButtonMenu();
-
-	/** Fill in menu content for when the view combo button is clicked */
-	void PopulateViewButtonMenu(UToolMenu* Menu);
+	/** Registers the menu used by the settings button */
+	static FDelayedAutoRegisterHelper AssetViewOptionsMenuRegistration;
 
 	/** Populate the given params for this AssetView */
 	void PopulateFilterAdditionalParams(FFiltersAdditionalParams& OutParams);
