Added filter bar override name to asset picker

diff --git a/Engine/Source/Editor/ContentBrowser/Private/SAssetPicker.cpp b/Engine/Source/Editor/ContentBrowser/Private/SAssetPicker.cpp
--- a/Engine/Source/Editor/ContentBrowser/Private/SAssetPicker.cpp
+++ b/Engine/Source/Editor/ContentBrowser/Private/SAssetPicker.cpp
@@ -80,6 +80,7 @@ void SAssetPicker::Construct( const FArguments& InArgs )
 	bPendingFocusNextFrame = InArgs._AssetPickerConfig.bFocusSearchBoxWhenOpened;
 	DefaultFilterMenuExpansion = InArgs._AssetPickerConfig.DefaultFilterMenuExpansion;
 	SaveSettingsName = InArgs._AssetPickerConfig.SaveSettingsName;
+	FilterBarSaveSettingsNameOverride = InArgs._AssetPickerConfig.FilterBarSaveSettingsNameOverride;
 	OnFolderEnteredDelegate = InArgs._AssetPickerConfig.OnFolderEntered;
 	OnGetAssetContextMenu = InArgs._AssetPickerConfig.OnGetAssetContextMenu;
 	OnGetFolderContextMenu = InArgs._AssetPickerConfig.OnGetFolderContextMenu;
@@ -209,7 +210,7 @@ void SAssetPicker::Construct( const FArguments& InArgs )
 			.OnFilterChanged(this, &SAssetPicker::OnFilterChanged)
 			.FrontendFilters(FrontendFilters)
 			.InitialClassFilters(FilterClassList)
-			.FilterBarIdentifier(FName(SaveSettingsName))
+			.FilterBarIdentifier(FName(FilterBarSaveSettingsNameOverride.Get(SaveSettingsName)))
 			.ExtraFrontendFilters(InArgs._AssetPickerConfig.ExtraFrontendFilters)
 			.DefaultMenuExpansionCategory(ConvertAssetTypeCategoryToAssetCategoryPath(DefaultFilterMenuExpansion).Get(EAssetCategoryPaths::Basic))
 			.OnExtendAddFilterMenu(InArgs._AssetPickerConfig.OnExtendAddFilterMenu)
diff --git a/Engine/Source/Editor/ContentBrowser/Private/SAssetPicker.h b/Engine/Source/Editor/ContentBrowser/Private/SAssetPicker.h
--- a/Engine/Source/Editor/ContentBrowser/Private/SAssetPicker.h
+++ b/Engine/Source/Editor/ContentBrowser/Private/SAssetPicker.h
@@ -206,6 +206,9 @@ private:
 	/** If set, view settings will be saved and loaded for the asset view using this name in ini files */
 	FString SaveSettingsName;
 
+	/** If set, uses this name for the filter bar identifier instead of the SaveSettingsName */
+	TOptional<FString> FilterBarSaveSettingsNameOverride;
+
 	/** Wether to allow renaming of items */
 	bool bAllowRename;
 };
diff --git a/Engine/Source/Editor/ContentBrowser/Public/IContentBrowserSingleton.h b/Engine/Source/Editor/ContentBrowser/Public/IContentBrowserSingleton.h
--- a/Engine/Source/Editor/ContentBrowser/Public/IContentBrowserSingleton.h
+++ b/Engine/Source/Editor/ContentBrowser/Public/IContentBrowserSingleton.h
@@ -255,6 +255,9 @@ struct FAssetPickerConfig
 	/** If set, view settings will be saved and loaded for the asset view using this name in ini files */
 	FString SaveSettingsName;
 
+	/** If set, will use this name for filter bar saving purposes instead of the SaveSettingsName. */
+	TOptional<FString> FilterBarSaveSettingsNameOverride;
+	
 	/** If true, the search box will gain focus when the asset picker is created */
 	bool bFocusSearchBoxWhenOpened;
 
