SAssetPicker: Invert the logic for the Other Developers filter because it was changed to be an inverse filter recently

diff --git a/Engine/Source/Editor/ContentBrowser/Private/SAssetPicker.cpp b/Engine/Source/Editor/ContentBrowser/Private/SAssetPicker.cpp
--- a/Engine/Source/Editor/ContentBrowser/Private/SAssetPicker.cpp
+++ b/Engine/Source/Editor/ContentBrowser/Private/SAssetPicker.cpp
@@ -664,7 +664,8 @@ bool SAssetPicker::IsShowingOtherDevelopersContent() const
 {
 	if (FilterListPtr.IsValid())
 	{
-		return !FilterListPtr->IsFrontendFilterActive(OtherDevelopersFilter);
+		// OtherDevelopersFilter is an inverse filter, so if it is active we are not showing other developers content
+		return FilterListPtr->IsFrontendFilterActive(OtherDevelopersFilter);
 	}
 	else
 	{
@@ -672,10 +673,10 @@ bool SAssetPicker::IsShowingOtherDevelopersContent() const
 		{
 			if (FrontendFilters->GetFilterAtIndex(i) == OtherDevelopersFilter)
 			{
-				return true;
+				return false;
 			}
 		}
-		return false;
+		return true;
 	}
 }
 
