[Content Browser] Menu restyle and refactor (add new, toolbar, navigation, folder context)

diff --git a/Engine/Source/Editor/ContentBrowser/Private/ContentBrowserMenuUtils.cpp b/Engine/Source/Editor/ContentBrowser/Private/ContentBrowserMenuUtils.cpp
--- a/Engine/Source/Editor/ContentBrowser/Private/ContentBrowserMenuUtils.cpp
+++ b/Engine/Source/Editor/ContentBrowser/Private/ContentBrowserMenuUtils.cpp
@@ -202,6 +202,16 @@ namespace ContentBrowserMenuUtils::Filters
 
 void ContentBrowserMenuUtils::AddFiltersToMenu(UToolMenu* InMenu, const FName& InOwningContentBrowserName, FFiltersAdditionalParams InFiltersAdditionalParams)
 {
+#if UE_CONTENTBROWSER_NEW_STYLE
+	FToolMenuSection& Section = InMenu->FindOrAddSection("Show");
+
+	// New style omits the "Show" prefix
+	const FText ShowCppLabel = LOCTEXT("ShowCppClassesOption", "C++ Classes");
+	const FText ShowDevelopersContentLabel = LOCTEXT("ShowDevelopersContentOption", "Developers Content");
+	const FText ShowEngineContentLabel = LOCTEXT("ShowEngineFolderOption", "Engine Content");
+	const FText ShowPluginContentLabel = LOCTEXT("ShowPluginFolderOption", "Plugin Content");
+	const FText ShowLocalizedContentLabel = LOCTEXT("ShowLocalizedContentOption", "Localized Content");
+#else
 	FToolMenuInsert ContentMenuInsert = FToolMenuInsert();
 	if (InMenu->ContainsSection("View"))
 	{
@@ -210,9 +220,17 @@ void ContentBrowserMenuUtils::AddFiltersToMenu(UToolMenu* InMenu, const FName& I
 	}
 
 	FToolMenuSection& Section = InMenu->FindOrAddSection("Content", LOCTEXT("ContentHeading", "Content"), ContentMenuInsert);
-	Section.AddMenuEntry(
+
+	const FText ShowCppLabel = LOCTEXT("ShowCppClassesOption", "Show C++ Classes");
+	const FText ShowDevelopersContentLabel = LOCTEXT("ShowDevelopersContentOption", "Show Developers Content");
+	const FText ShowEngineContentLabel = LOCTEXT("ShowEngineFolderOption", "Show Engine Content");
+	const FText ShowPluginContentLabel = LOCTEXT("ShowPluginFolderOption", "Show Plugin Content");
+	const FText ShowLocalizedContentLabel = LOCTEXT("ShowLocalizedContentOption", "Show Localized Content");
+#endif
+
+		Section.AddMenuEntry(
 		"ShowCppClasses",
-		LOCTEXT("ShowCppClassesOption", "Show C++ Classes"),
+		ShowCppLabel,
 		LOCTEXT("ShowCppClassesOptionToolTip", "Show C++ classes in the view?"),
 		FSlateIcon(),
 		FUIAction(
@@ -225,7 +243,7 @@ void ContentBrowserMenuUtils::AddFiltersToMenu(UToolMenu* InMenu, const FName& I
 
 	Section.AddMenuEntry(
 		"ShowDevelopersContent",
-		LOCTEXT("ShowDevelopersContentOption", "Show Developers Content"),
+		ShowDevelopersContentLabel,
 		LOCTEXT("ShowDevelopersContentOptionToolTip", "Show developers content in the view?"),
 		FSlateIcon(),
 		FUIAction(
@@ -238,7 +256,7 @@ void ContentBrowserMenuUtils::AddFiltersToMenu(UToolMenu* InMenu, const FName& I
 
 	Section.AddMenuEntry(
 		"ShowEngineFolder",
-		LOCTEXT("ShowEngineFolderOption", "Show Engine Content"),
+		ShowEngineContentLabel,
 		LOCTEXT("ShowEngineFolderOptionToolTip", "Show engine content in the view?"),
 		FSlateIcon(),
 		FUIAction(
@@ -251,7 +269,7 @@ void ContentBrowserMenuUtils::AddFiltersToMenu(UToolMenu* InMenu, const FName& I
 
 	Section.AddMenuEntry(
 		"ShowPluginFolder",
-		LOCTEXT("ShowPluginFolderOption", "Show Plugin Content"),
+		ShowPluginContentLabel,
 		LOCTEXT("ShowPluginFolderOptionToolTip", "Show plugin content in the view?"),
 		FSlateIcon(),
 		FUIAction(
@@ -264,7 +282,7 @@ void ContentBrowserMenuUtils::AddFiltersToMenu(UToolMenu* InMenu, const FName& I
 
 	Section.AddMenuEntry(
 		"ShowLocalizedContent",
-		LOCTEXT("ShowLocalizedContentOption", "Show Localized Content"),
+		ShowLocalizedContentLabel,
 		LOCTEXT("ShowLocalizedContentOptionToolTip", "Show localized content in the view?"),
 		FSlateIcon(),
 		FUIAction(
diff --git a/Engine/Source/Editor/ContentBrowser/Private/ContentBrowserStyle.cpp b/Engine/Source/Editor/ContentBrowser/Private/ContentBrowserStyle.cpp
--- a/Engine/Source/Editor/ContentBrowser/Private/ContentBrowserStyle.cpp
+++ b/Engine/Source/Editor/ContentBrowser/Private/ContentBrowserStyle.cpp
@@ -312,7 +312,11 @@ namespace UE::ContentBrowser::Private
 
 		ContentBrowserToolBarStyle.SetSeparatorBrush(FSlateNoResource());
 		ContentBrowserToolBarStyle.SetSeparatorPadding(FMargin(4.0f, 0.0f));
+#if UE_CONTENTBROWSER_NEW_STYLE
+		ContentBrowserToolBarStyle.SetBackgroundPadding(FMargin(4.0f, 2.0f, 4.0f, 2.0f));
+#else
 		ContentBrowserToolBarStyle.SetBackgroundPadding(FMargin(4.0f, 2.0f, 0.0f, 2.0f));
+#endif
 
 		Set("ContentBrowser.ToolBar", ContentBrowserToolBarStyle);
 
diff --git a/Engine/Source/Editor/ContentBrowser/Private/Menus/AssetViewMenus.cpp b/Engine/Source/Editor/ContentBrowser/Private/Menus/AssetViewMenus.cpp
--- a/Engine/Source/Editor/ContentBrowser/Private/Menus/AssetViewMenus.cpp
+++ b/Engine/Source/Editor/ContentBrowser/Private/Menus/AssetViewMenus.cpp
@@ -37,9 +37,9 @@ FDelayedAutoRegisterHelper SAssetView::AssetViewOptionsMenuRegistration(
 				{
 					const TSharedRef<SAssetView> AssetView = Context->AssetView.Pin().ToSharedRef();
 
-					// AssetViewType Section
+						// ViewType Section
 					{
-						FToolMenuSection& Section = InMenu->AddSection("AssetViewType", LOCTEXT("ViewTypeHeading", "View Type"));
+						FToolMenuSection& Section = InMenu->FindOrAddSection("AssetViewType", LOCTEXT("ViewTypeHeading", "View Type"));
 						Section.AddMenuEntry(
 							"TileView",
 							LOCTEXT("TileViewOption", "Tiles"),
@@ -77,10 +77,15 @@ FDelayedAutoRegisterHelper SAssetView::AssetViewOptionsMenuRegistration(
 								FIsActionChecked::CreateSP(AssetView, &SAssetView::IsCurrentViewType, EAssetViewType::Column)
 							),
 							EUserInterfaceActionType::RadioButton
-							);
+						);
+
+						const FText CustomViewLabel = AssetView->ViewExtender
+							? AssetView->ViewExtender->GetViewDisplayName()
+							: LOCTEXT("CustomViewOption", "Custom");
 
-						const FText CustomViewLabel = AssetView->ViewExtender ? AssetView->ViewExtender->GetViewDisplayName() : LOCTEXT("CustomViewOption", "Custom");
-						const FText CustomViewTooltip = AssetView->ViewExtender ? AssetView->ViewExtender->GetViewTooltipText() : LOCTEXT("CustomViewOptionToolTip", "A user specified custom view.");
+						const FText CustomViewTooltip = AssetView->ViewExtender
+							? AssetView->ViewExtender->GetViewTooltipText()
+							: LOCTEXT("CustomViewOptionToolTip", "A user specified custom view.");
 
 						Section.AddMenuEntry(
 							"CustomView",
@@ -97,6 +102,7 @@ FDelayedAutoRegisterHelper SAssetView::AssetViewOptionsMenuRegistration(
 						);
 					}
 
+#if !UE_CONTENTBROWSER_NEW_STYLE
 					// Filter Bar Section
 					if (TSharedPtr<SFilterList> FilterBarPinned = AssetView->FilterBar.Pin())
 					{
@@ -144,50 +150,128 @@ FDelayedAutoRegisterHelper SAssetView::AssetViewOptionsMenuRegistration(
 							EUserInterfaceActionType::RadioButton
 						);
 					}
+#endif
 
-					// View Section
+				// Thumbnails Section
+				{
+#if UE_CONTENTBROWSER_NEW_STYLE
+					FToolMenuInsert SectionInsert("AssetViewType", EToolMenuInsertType::After);
+#else
+					FToolMenuInsert SectionInsert("Search", EToolMenuInsertType::After);
+#endif
+					FToolMenuSection& Section = InMenu->FindOrAddSection("AssetThumbnails", LOCTEXT("ThumbnailsHeading", "Thumbnails"), SectionInsert);
+
+					auto CreateThumbnailSizeSubMenu = [AssetView](UToolMenu* SubMenu)
 					{
-						FToolMenuSection& Section = InMenu->AddSection("View", LOCTEXT("ViewHeading", "View"));
+						FToolMenuSection& SizeSection = SubMenu->AddSection("ThumbnailSizes");
 
-						Section.AddMenuEntry(
-							"ShowFolders",
-							LOCTEXT("ShowFoldersOption", "Show Folders"),
-							LOCTEXT("ShowFoldersOptionToolTip", "Show folders in the view as well as assets?"),
-							FSlateIcon(),
-							FUIAction(
-								FExecuteAction::CreateSP(AssetView, &SAssetView::ToggleShowFolders),
-								FCanExecuteAction::CreateSP(AssetView, &SAssetView::IsToggleShowFoldersAllowed),
-								FIsActionChecked::CreateSP(AssetView, &SAssetView::IsShowingFolders)
-							),
-							EUserInterfaceActionType::ToggleButton
-						);
+						for (int32 EnumValue = static_cast<int32>(EThumbnailSize::Tiny); EnumValue < static_cast<int32>(EThumbnailSize::MAX); ++EnumValue)
+						{
+#if !UE_CONTENTBROWSER_NEW_STYLE
+							if ((EThumbnailSize)EnumValue == EThumbnailSize::XLarge)
+							{
+								continue;
+							}
+#endif
+
+							SizeSection.AddMenuEntry(
+								NAME_None,
+								SAssetView::ThumbnailSizeToDisplayName((EThumbnailSize)EnumValue),
+								FText::GetEmpty(),
+								FSlateIcon(),
+								FUIAction(
+									FExecuteAction::CreateSP(AssetView, &SAssetView::OnThumbnailSizeChanged, (EThumbnailSize)EnumValue),
+									FCanExecuteAction::CreateSP(AssetView, &SAssetView::IsThumbnailScalingAllowed),
+									FIsActionChecked::CreateSP(AssetView, &SAssetView::IsThumbnailSizeChecked, (EThumbnailSize)EnumValue)
+								),
+								EUserInterfaceActionType::RadioButton
+							);
+						}
+					};
+
+					Section.AddEntry(FToolMenuEntry::InitSubMenu(
+						"ThumbnailSize",
+						LOCTEXT("ThumbnailSize", "Thumbnail Size"),
+						LOCTEXT("ThumbnailSizeToolTip", "Adjust the size of thumbnails."),
+						FNewToolMenuDelegate::CreateLambda(CreateThumbnailSizeSubMenu)
+					));
 
+#if !UE_CONTENTBROWSER_NEW_STYLE
 						Section.AddMenuEntry(
-							"ShowEmptyFolders",
-							LOCTEXT("ShowEmptyFoldersOption", "Show Empty Folders"),
-							LOCTEXT("ShowEmptyFoldersOptionToolTip", "Show empty folders in the view as well as assets?"),
+							"ThumbnailEditMode",
+							LOCTEXT("ThumbnailEditModeOption", "Thumbnail Edit Mode"),
+							LOCTEXT("ThumbnailEditModeOptionToolTip", "Toggle thumbnail editing mode. When in this mode you can rotate the camera on 3D thumbnails by dragging them."),
 							FSlateIcon(),
 							FUIAction(
-								FExecuteAction::CreateSP(AssetView, &SAssetView::ToggleShowEmptyFolders),
-								FCanExecuteAction::CreateSP(AssetView, &SAssetView::IsToggleShowEmptyFoldersAllowed),
-								FIsActionChecked::CreateSP(AssetView, &SAssetView::IsShowingEmptyFolders)
+								FExecuteAction::CreateSP(AssetView, &SAssetView::ToggleThumbnailEditMode),
+								FCanExecuteAction::CreateSP(AssetView, &SAssetView::IsThumbnailEditModeAllowed),
+								FIsActionChecked::CreateSP(AssetView, &SAssetView::IsThumbnailEditMode)
 							),
 							EUserInterfaceActionType::ToggleButton
 						);
 
 						Section.AddMenuEntry(
-							"ShowFavorite",
-							LOCTEXT("ShowFavoriteOptions", "Show Favorites"),
-							LOCTEXT("ShowFavoriteOptionToolTip", "Show the favorite folders in the view?"),
+							"RealTimeThumbnails",
+							LOCTEXT("RealTimeThumbnailsOption", "Real-Time Thumbnails"),
+							LOCTEXT("RealTimeThumbnailsOptionToolTip", "Renders the assets thumbnails in real-time"),
 							FSlateIcon(),
 							FUIAction(
-								FExecuteAction::CreateSP(AssetView, &SAssetView::ToggleShowFavorites),
-								FCanExecuteAction::CreateSP(AssetView, &SAssetView::IsToggleShowFavoritesAllowed),
-								FIsActionChecked::CreateSP(AssetView, &SAssetView::IsShowingFavorites)
+								FExecuteAction::CreateSP(AssetView, &SAssetView::ToggleRealTimeThumbnails),
+								FCanExecuteAction::CreateSP(AssetView, &SAssetView::CanShowRealTimeThumbnails),
+								FIsActionChecked::CreateSP(AssetView, &SAssetView::IsShowingRealTimeThumbnails)
 							),
 							EUserInterfaceActionType::ToggleButton
 						);
+#endif
+				}
 
+				// Show/View Section
+				{
+#if UE_CONTENTBROWSER_NEW_STYLE
+					FToolMenuSection& Section = InMenu->FindOrAddSection("Show", LOCTEXT("ShowHeading", "Show"));
+#else
+					FToolMenuSection& Section = InMenu->AddSection("View", LOCTEXT("ViewHeading", "View"));
+#endif
+					Section.AddMenuEntry(
+						"ShowFolders",
+						LOCTEXT("ShowFoldersOption", "Folders"),
+						LOCTEXT("ShowFoldersOptionToolTip", "Show folders in the view as well as assets?"),
+						FSlateIcon(),
+						FUIAction(
+							FExecuteAction::CreateSP(AssetView, &SAssetView::ToggleShowFolders),
+							FCanExecuteAction::CreateSP(AssetView, &SAssetView::IsToggleShowFoldersAllowed),
+							FIsActionChecked::CreateSP(AssetView, &SAssetView::IsShowingFolders)
+						),
+						EUserInterfaceActionType::ToggleButton
+					);
+
+					Section.AddMenuEntry(
+						"ShowEmptyFolders",
+						LOCTEXT("ShowEmptyFoldersOption", "Empty Folders"),
+						LOCTEXT("ShowEmptyFoldersOptionToolTip", "Show empty folders in the view as well as assets?"),
+						FSlateIcon(),
+						FUIAction(
+							FExecuteAction::CreateSP(AssetView, &SAssetView::ToggleShowEmptyFolders),
+							FCanExecuteAction::CreateSP(AssetView, &SAssetView::IsToggleShowEmptyFoldersAllowed),
+							FIsActionChecked::CreateSP(AssetView, &SAssetView::IsShowingEmptyFolders)
+						),
+						EUserInterfaceActionType::ToggleButton
+					);
+
+					Section.AddMenuEntry(
+						"ShowFavorite",
+						LOCTEXT("ShowFavoriteOptions", "Show Favorites"),
+						LOCTEXT("ShowFavoriteOptionToolTip", "Show the favorite folders in the view?"),
+						FSlateIcon(),
+						FUIAction(
+							FExecuteAction::CreateSP(AssetView, &SAssetView::ToggleShowFavorites),
+							FCanExecuteAction::CreateSP(AssetView, &SAssetView::IsToggleShowFavoritesAllowed),
+							FIsActionChecked::CreateSP(AssetView, &SAssetView::IsShowingFavorites)
+						),
+						EUserInterfaceActionType::ToggleButton
+					);
+
+#if !UE_CONTENTBROWSER_NEW_STYLE
 						Section.AddMenuEntry(
 							"FilterRecursively",
 							LOCTEXT("FilterRecursivelyOption", "Filter Recursively"),
@@ -235,11 +319,12 @@ FDelayedAutoRegisterHelper SAssetView::AssetViewOptionsMenuRegistration(
 								LOCTEXT("PathViewFilters_ToolTip", "Path View Filters"),
 								FNewToolMenuDelegate());
 						}
-					}
+#endif
+				}
 
 					// Search Section
 					{
-						FToolMenuSection& Section = InMenu->AddSection("Search", LOCTEXT("SearchHeading", "Search"));
+						FToolMenuSection& Section = InMenu->FindOrAddSection("Search", LOCTEXT("SearchHeading", "Search"));
 						Section.AddMenuEntry(
 							"IncludeClassName",
 							LOCTEXT("IncludeClassNameOption", "Search Asset Class Names"),
@@ -280,69 +365,27 @@ FDelayedAutoRegisterHelper SAssetView::AssetViewOptionsMenuRegistration(
 						);
 					}
 
-					// Asset Thumbnails Section
-					{
-						FToolMenuSection& Section = InMenu->AddSection("AssetThumbnails", LOCTEXT("ThumbnailsHeading", "Thumbnails"));
-
-						auto CreateThumbnailSizeSubMenu = [AssetView](UToolMenu* SubMenu)
-						{
-							FToolMenuSection& SizeSection = SubMenu->AddSection("ThumbnailSizes");
-
-							for (int32 EnumValue = (int32)EThumbnailSize::Tiny; EnumValue < (int32)EThumbnailSize::MAX; ++EnumValue)
-							{
-								if ((EThumbnailSize)EnumValue == EThumbnailSize::XLarge)
-								{
-									continue;
-								}
-
-								SizeSection.AddMenuEntry(
-									NAME_None,
-									SAssetView::ThumbnailSizeToDisplayName((EThumbnailSize)EnumValue),
-									FText::GetEmpty(),
-									FSlateIcon(),
-									FUIAction(
-										FExecuteAction::CreateSP(AssetView, &SAssetView::OnThumbnailSizeChanged, (EThumbnailSize)EnumValue),
-										FCanExecuteAction::CreateSP(AssetView, &SAssetView::IsThumbnailScalingAllowed),
-										FIsActionChecked::CreateSP(AssetView, &SAssetView::IsThumbnailSizeChecked, (EThumbnailSize)EnumValue)
-									),
-									EUserInterfaceActionType::RadioButton
-								);
-							}
-						};
-						Section.AddEntry(FToolMenuEntry::InitSubMenu(
-							"ThumbnailSize",
-							LOCTEXT("ThumbnailSize", "Thumbnail Size"),
-							LOCTEXT("ThumbnailSizeToolTip", "Adjust the size of thumbnails."),
-							FNewToolMenuDelegate::CreateLambda(CreateThumbnailSizeSubMenu)
-						));
-
-						Section.AddMenuEntry(
-							"ThumbnailEditMode",
-							LOCTEXT("ThumbnailEditModeOption", "Thumbnail Edit Mode"),
-							LOCTEXT("ThumbnailEditModeOptionToolTip", "Toggle thumbnail editing mode. When in this mode you can rotate the camera on 3D thumbnails by dragging them."),
-							FSlateIcon(),
-							FUIAction(
-								FExecuteAction::CreateSP(AssetView, &SAssetView::ToggleThumbnailEditMode),
-								FCanExecuteAction::CreateSP(AssetView, &SAssetView::IsThumbnailEditModeAllowed),
-								FIsActionChecked::CreateSP(AssetView, &SAssetView::IsThumbnailEditMode)
-							),
-							EUserInterfaceActionType::ToggleButton
-						);
 
+#if UE_CONTENTBROWSER_NEW_STYLE
+					// Manage Section
+					{
+						FToolMenuSection& Section = InMenu->FindOrAddSection("Manage", LOCTEXT("ManageHeading", "Manage"));
 						Section.AddMenuEntry(
-							"RealTimeThumbnails",
-							LOCTEXT("RealTimeThumbnailsOption", "Real-Time Thumbnails"),
-							LOCTEXT("RealTimeThumbnailsOptionToolTip", "Renders the assets thumbnails in real-time"),
+							"OrganizeFolders",
+							LOCTEXT("OrganizeFoldersOption", "Organize Folders"),
+							LOCTEXT("OrganizeFoldersOptionToolTip", "Organize folders in the view?"),
 							FSlateIcon(),
 							FUIAction(
-								FExecuteAction::CreateSP(AssetView, &SAssetView::ToggleRealTimeThumbnails),
-								FCanExecuteAction::CreateSP(AssetView, &SAssetView::CanShowRealTimeThumbnails),
-								FIsActionChecked::CreateSP(AssetView, &SAssetView::IsShowingRealTimeThumbnails)
+								FExecuteAction::CreateSP(AssetView, &SAssetView::ToggleOrganizeFolders),
+								FCanExecuteAction(),
+								FIsActionChecked::CreateSP(AssetView, &SAssetView::IsOrganizingFolders)
 							),
 							EUserInterfaceActionType::ToggleButton
 						);
 					}
+#endif
 
+#if !UE_CONTENTBROWSER_NEW_STYLE
 					// Asset Columns Section
 					if (AssetView->GetColumnViewVisibility() == EVisibility::Visible)
 					{
@@ -368,6 +411,7 @@ FDelayedAutoRegisterHelper SAssetView::AssetViewOptionsMenuRegistration(
 							);
 						}
 					}
+#endif
 
 					AssetView->PopulateFilterAdditionalParams(FilterParams);
 				}
@@ -382,4 +426,41 @@ FDelayedAutoRegisterHelper SAssetView::AssetViewOptionsMenuRegistration(
 		}));
 	});
 
+#if UE_CONTENTBROWSER_NEW_STYLE
+
+FDelayedAutoRegisterHelper SAssetView::NavigationBarMenuExtensionRegistration(
+	EDelayedRegisterRunPhase::EndOfEngineInit,
+	[]
+	{
+		static const FName NavigationBarName("ContentBrowser.NavigationBar");
+		if (UToolMenu* NavigationBarMenu = UToolMenus::Get()->ExtendMenu(NavigationBarName))
+		{
+			FToolMenuSection& AssetCountSection = NavigationBarMenu->FindOrAddSection("AssetCount");
+			AssetCountSection.InsertPosition = FToolMenuInsert(NAME_None, EToolMenuInsertType::Last);
+			AssetCountSection.Alignment = EToolMenuSectionAlign::Last;
+
+			AssetCountSection.AddDynamicEntry(
+				"AssetCount",
+				FNewToolMenuSectionDelegate::CreateLambda([](FToolMenuSection& InSection)
+				{
+					if (const UContentBrowserToolbarMenuContext* Context = InSection.FindContext<UContentBrowserToolbarMenuContext>())
+					{
+						if (Context->AssetView.IsValid())
+						{
+							const TSharedRef<SAssetView> AssetView = Context->AssetView.Pin().ToSharedRef();
+
+							InSection.AddEntry(
+								FToolMenuEntry::InitWidget(
+									"AssetCount",
+									SNew(STextBlock)
+									.Text(AssetView, &SAssetView::GetAssetCountText),
+									FText::GetEmpty()));
+						}
+					}
+				}));
+		}
+	});
+
+#endif
+
 #undef LOCTEXT_NAMESPACE
diff --git a/Engine/Source/Editor/ContentBrowser/Private/Menus/ContentBrowserMenus.cpp b/Engine/Source/Editor/ContentBrowser/Private/Menus/ContentBrowserMenus.cpp
--- a/Engine/Source/Editor/ContentBrowser/Private/Menus/ContentBrowserMenus.cpp
+++ b/Engine/Source/Editor/ContentBrowser/Private/Menus/ContentBrowserMenus.cpp
@@ -7,8 +7,173 @@
 #include "ToolMenus.h"
 #include "UObject/NameTypes.h"
 
+#include "SPositiveActionButton.h"
+
+#if UE_CONTENTBROWSER_NEW_STYLE
+#include "SActionButton.h"
+#include "SAssetSearchBox.h"
+#include "SNavigationBar.h"
+#endif
+
 #define LOCTEXT_NAMESPACE "ContentBrowser"
 
+#if UE_CONTENTBROWSER_NEW_STYLE
+FDelayedAutoRegisterHelper SContentBrowser::NavigationBarMenuRegistration(
+	EDelayedRegisterRunPhase::EndOfEngineInit,
+	[]
+	{
+		static const FName NavigationBarMenuName("ContentBrowser.NavigationBar");
+		if (!ensureAlways(!UToolMenus::Get()->IsMenuRegistered(NavigationBarMenuName)))
+		{
+			return;
+		}
+
+		FToolMenuOwnerScoped OwnerScoped(UE_MODULE_NAME);
+
+		UToolMenu* ToolBar = UToolMenus::Get()->RegisterMenu(NavigationBarMenuName, NAME_None, EMultiBoxType::SlimHorizontalToolBar);
+		ToolBar->StyleName = "ContentBrowser.ToolBar";
+
+		ToolBar->AddSection("History")
+		.AddDynamicEntry(NAME_None,
+			FNewToolMenuSectionDelegate::CreateLambda([](FToolMenuSection& InSection)
+			{
+				if (const UContentBrowserToolbarMenuContext* Context = InSection.FindContext<UContentBrowserToolbarMenuContext>())
+				{
+					if (Context->ContentBrowser.IsValid())
+					{
+						const TSharedRef<SContentBrowser> ContentBrowser = Context->ContentBrowser.Pin().ToSharedRef();
+
+						const TSharedRef<SActionButton> HistoryBackButton =
+							SNew(SActionButton)
+							.ActionButtonType(EActionButtonType::Simple)
+							.ToolTipText(ContentBrowser, &SContentBrowser::GetHistoryBackTooltip)
+							.IsEnabled(ContentBrowser, &SContentBrowser::IsBackEnabled)
+							.OnClicked(ContentBrowser, &SContentBrowser::BackClicked)
+							.Icon(FAppStyle::Get().GetBrush("Icons.CircleArrowLeft"))
+							.AddMetaData<FTagMetaData>(FTagMetaData(TEXT("ContentBrowserHistoryBack")));
+
+						InSection.AddEntry(
+							FToolMenuEntry::InitWidget(
+								"HistoryBackButton",
+								HistoryBackButton,
+								FText::GetEmpty()
+							));
+
+						const TSharedRef<SActionButton> HistoryForwardButton =
+							SNew(SActionButton)
+							.ActionButtonType(EActionButtonType::Simple)
+							.ToolTipText(ContentBrowser, &SContentBrowser::GetHistoryForwardTooltip)
+							.IsEnabled(ContentBrowser, &SContentBrowser::IsForwardEnabled)
+							.OnClicked(ContentBrowser, &SContentBrowser::ForwardClicked)
+							.Icon(FAppStyle::Get().GetBrush("Icons.CircleArrowRight"))
+							.AddMetaData<FTagMetaData>(FTagMetaData(TEXT("ContentBrowserHistoryForward")));
+
+						InSection.AddEntry(
+							FToolMenuEntry::InitWidget(
+								"HistoryForwardButton",
+								HistoryForwardButton,
+								FText::GetEmpty()
+							));
+					}
+				}
+			}));
+
+		ToolBar->AddSection("Path")
+		.AddDynamicEntry(NAME_None, FNewToolMenuSectionDelegate::CreateLambda(
+			[](FToolMenuSection& InSection)
+			{
+				static constexpr float NavigationPathMinWidth = 180.0f;
+				static constexpr float NavigationPathMaxWidth = 640.0f;
+
+				if (const UContentBrowserToolbarMenuContext* Context = InSection.FindContext<UContentBrowserToolbarMenuContext>())
+				{
+					if (Context->ContentBrowser.IsValid())
+					{
+						const TSharedRef<SContentBrowser> ContentBrowser = Context->ContentBrowser.Pin().ToSharedRef();
+
+						// Should have been initialized in SContentBrowser::CreateAssetView
+						if (!ensure(ContentBrowser->NavigationBar.IsValid()))
+						{
+							return;
+						}
+
+						FToolMenuEntry& PathBoxEntry =
+							InSection.AddEntry(
+							FToolMenuEntry::InitWidget(
+								"PathBox",
+								ContentBrowser->NavigationBar.ToSharedRef(),
+								FText::GetEmpty()));
+
+						PathBoxEntry.WidgetData.StyleParams.SizeRule = FSizeParam::SizeRule_StretchContent;
+						PathBoxEntry.WidgetData.StyleParams.MinSize = NavigationPathMinWidth;
+						PathBoxEntry.WidgetData.StyleParams.MaxSize = NavigationPathMaxWidth;
+						PathBoxEntry.WidgetData.StyleParams.DesiredWidth = NavigationPathMaxWidth;
+					}
+				}
+			}));
+	});
+#endif
+
+FDelayedAutoRegisterHelper SContentBrowser::AddNewContextMenuRegistration(
+	EDelayedRegisterRunPhase::EndOfEngineInit,
+	[]
+	{
+		static const FName AddNewContextMenuName("ContentBrowser.AddNewContextMenu");
+		if (!ensureAlways(!UToolMenus::Get()->IsMenuRegistered(AddNewContextMenuName)))
+		{
+			return;
+		}
+
+		FToolMenuOwnerScoped OwnerScoped(UE_MODULE_NAME);
+
+		UToolMenu* Menu = UToolMenus::Get()->RegisterMenu(AddNewContextMenuName);
+		Menu->AddDynamicSection("DynamicSection_Common", FNewToolMenuDelegate::CreateLambda([](UToolMenu* InMenu)
+		{
+			TSharedPtr<SContentBrowser> ContentBrowser;
+			if (const UContentBrowserMenuContext* MenuContext = InMenu->FindContext<UContentBrowserMenuContext>())
+			{
+				ContentBrowser = MenuContext->ContentBrowser.Pin();
+			}
+			else
+			{
+				if (const UContentBrowserToolbarMenuContext* ToolbarContext  = InMenu->FindContext<UContentBrowserToolbarMenuContext>())
+				{
+					ContentBrowser = ToolbarContext->ContentBrowser.Pin();
+				}
+			}
+
+			if (ContentBrowser)
+			{
+				ContentBrowser->PopulateAddNewContextMenu(InMenu);
+			}
+		}));
+	});
+
+FDelayedAutoRegisterHelper SContentBrowser::FolderContextMenuRegistration(
+	EDelayedRegisterRunPhase::EndOfEngineInit,
+	[]
+	{
+		static const FName FolderContextMenuName("ContentBrowser.FolderContextMenu");
+		if (!ensureAlways(!UToolMenus::Get()->IsMenuRegistered(FolderContextMenuName)))
+		{
+			return;
+		}
+
+		FToolMenuOwnerScoped OwnerScoped(UE_MODULE_NAME);
+
+		UToolMenu* Menu = UToolMenus::Get()->RegisterMenu(FolderContextMenuName);
+		Menu->bCloseSelfOnly = true;
+
+		Menu->AddDynamicSection(NAME_None, FNewToolMenuDelegate::CreateLambda([](UToolMenu* InMenu)
+		{
+			UContentBrowserFolderContext* Context = InMenu->FindContext<UContentBrowserFolderContext>();
+			if (Context && Context->ContentBrowser.IsValid())
+			{
+				Context->ContentBrowser.Pin()->PopulateFolderContextMenu(InMenu);
+			}
+		}));
+	});
+
 FDelayedAutoRegisterHelper SContentBrowser::PathViewFiltersMenuRegistration(
 	EDelayedRegisterRunPhase::EndOfEngineInit,
 	[]
@@ -42,4 +207,221 @@ FDelayedAutoRegisterHelper SContentBrowser::PathViewFiltersMenuRegistration(
 		}));
 	});
 
+FDelayedAutoRegisterHelper SContentBrowser::ToolBarMenuRegistration(
+	EDelayedRegisterRunPhase::EndOfEngineInit,
+	[]
+	{
+		static const FName ToolBarMenuName("ContentBrowser.ToolBar");
+		if (!ensureAlways(!UToolMenus::Get()->IsMenuRegistered(ToolBarMenuName)))
+		{
+			return;
+		}
+
+		FToolMenuOwnerScoped OwnerScoped(UE_MODULE_NAME);
+
+		UToolMenu* ToolBar = UToolMenus::Get()->RegisterMenu(ToolBarMenuName, NAME_None, EMultiBoxType::SlimHorizontalToolBar);
+		ToolBar->StyleName = ToolBarMenuName; // Style name is the same as the menu entry name
+
+		ToolBar->AddSection("New")
+		.AddDynamicEntry(NAME_None,
+			FNewToolMenuSectionDelegate::CreateLambda([](FToolMenuSection& InSection)
+			{
+				if (UContentBrowserToolbarMenuContext* Context = InSection.FindContext<UContentBrowserToolbarMenuContext>())
+				{
+					if (Context->ContentBrowser.IsValid())
+					{
+						const TSharedRef<SContentBrowser> ContentBrowser = Context->ContentBrowser.Pin().ToSharedRef();
+
+						const TSharedRef<SPositiveActionButton> NewButton =
+							SNew(SPositiveActionButton)
+							.OnGetMenuContent_Lambda([Context]
+							{
+								return Context->ContentBrowser.Pin()->MakeAddNewContextMenu(EContentBrowserDataMenuContext_AddNewMenuDomain::Toolbar, Context);
+							})
+							.ToolTipText(ContentBrowser, &SContentBrowser::GetAddNewToolTipText)
+							.IsEnabled(ContentBrowser, &SContentBrowser::IsAddNewEnabled)
+							.AddMetaData<FTagMetaData>(FTagMetaData(TEXT("ContentBrowserNewAsset")))
+							.Icon(FAppStyle::Get().GetBrush("Icons.Plus"))
+							.Text(LOCTEXT("AddAssetButton", "Add"));
+
+						InSection.AddEntry(
+							FToolMenuEntry::InitWidget(
+								"NewButton",
+								NewButton,
+								FText::GetEmpty(),
+								true,
+								false
+							));
+					}
+				}
+			}));
+
+		ToolBar->AddSection("Save")
+		.AddDynamicEntry(NAME_None,
+			FNewToolMenuSectionDelegate::CreateLambda([](FToolMenuSection& InSection)
+			{
+				if (const UContentBrowserToolbarMenuContext* Context = InSection.FindContext<UContentBrowserToolbarMenuContext>())
+				{
+					if (Context->ContentBrowser.IsValid())
+					{
+						const TSharedRef<SContentBrowser> ContentBrowser = Context->ContentBrowser.Pin().ToSharedRef();
+
+#if UE_CONTENTBROWSER_NEW_STYLE
+						TSharedRef<SWidget> SaveButton =
+							SNew(SActionButton)
+							.ToolTipText(LOCTEXT("SaveDirtyPackagesTooltip", "Save all modified assets."))
+							.OnClicked(ContentBrowser, &SContentBrowser::OnSaveClicked)
+							.AddMetaData<FTagMetaData>(FTagMetaData(TEXT("ContentBrowserSaveDirtyPackages")))
+							.Icon(FAppStyle::Get().GetBrush("MainFrame.SaveAll"))
+							.Text(LOCTEXT("SaveAll", "Save All"));
+#else
+						TSharedRef<SWidget> SaveButton =
+							SNew(SButton)
+							.ButtonStyle(FAppStyle::Get(), "SimpleButton")
+							.ToolTipText(LOCTEXT("SaveDirtyPackagesTooltip", "Save all modified assets."))
+							.ContentPadding(2.f)
+							.OnClicked(ContentBrowser, &SContentBrowser::OnSaveClicked)
+							.AddMetaData<FTagMetaData>(FTagMetaData(TEXT("ContentBrowserSaveDirtyPackages")))
+							[
+								SNew(SHorizontalBox)
+								// Save All Icon
+								+ SHorizontalBox::Slot()
+								.AutoWidth()
+								.HAlign(HAlign_Center)
+								.VAlign(VAlign_Center)
+								[
+									SNew(SImage)
+									.Image(FAppStyle::Get().GetBrush("MainFrame.SaveAll"))
+									.ColorAndOpacity(FSlateColor::UseForeground())
+								]
+
+								// Save All Text
+								+ SHorizontalBox::Slot()
+								.Padding(FMargin(3, 0, 0, 0))
+								.VAlign(VAlign_Center)
+								.AutoWidth()
+								[
+									SNew(STextBlock)
+									.TextStyle(FAppStyle::Get(), "NormalText")
+									.Text(LOCTEXT("SaveAll", "Save All"))
+								]
+							];
+#endif
+
+						InSection.AddEntry(
+							FToolMenuEntry::InitWidget(
+								"SaveButton",
+								SaveButton,
+								FText::GetEmpty(),
+								true,
+								false
+							));
+					}
+				}
+			}));
+
+#if UE_CONTENTBROWSER_NEW_STYLE
+		ToolBar->AddSection("Filter")
+		.AddDynamicEntry(NAME_None, FNewToolMenuSectionDelegate::CreateLambda(
+			[](FToolMenuSection& InSection)
+			{
+				if (const UContentBrowserToolbarMenuContext* Context = InSection.FindContext<UContentBrowserToolbarMenuContext>())
+				{
+					if (Context->ContentBrowser.IsValid())
+					{
+						const TSharedRef<SContentBrowser> ContentBrowser = Context->ContentBrowser.Pin().ToSharedRef();
+
+						if (!ContentBrowser->FilterComboButton.IsValid())
+						{
+							return;
+						}
+
+						FToolMenuEntry& FilterEntry = InSection.AddEntry(
+							FToolMenuEntry::InitWidget(
+								"Filter",
+								SNew(SBox)
+								.HAlign(HAlign_Left)
+								.VAlign(VAlign_Center)
+								[
+									ContentBrowser->FilterComboButton.ToSharedRef()
+								],
+								LOCTEXT("Filter", "Filter")
+							));
+
+						FilterEntry.WidgetData.StyleParams.HorizontalAlignment = HAlign_Left;
+						FilterEntry.WidgetData.StyleParams.SizeRule = FSizeParam::ESizeRule::SizeRule_Auto;
+
+						// ... no label when in the toolbar
+						FilterEntry.ToolBarData.LabelOverride = FText::GetEmpty();
+					}
+				}
+			}));
+
+		ToolBar->AddSection("Search")
+		.AddDynamicEntry(NAME_None, FNewToolMenuSectionDelegate::CreateLambda(
+			[](FToolMenuSection& InSection)
+			{
+				static constexpr float SearchBoxMinWidth = 180.0f;
+				static constexpr float SearchBoxMaxWidth = 640.0f;
+
+				if (const UContentBrowserToolbarMenuContext* Context = InSection.FindContext<UContentBrowserToolbarMenuContext>())
+				{
+					if (Context->ContentBrowser.IsValid())
+					{
+						const TSharedRef<SContentBrowser> ContentBrowser = Context->ContentBrowser.Pin().ToSharedRef();
+
+						// Should have been initialized in SContentBrowser::CreateAssetView
+						if (!ensure(ContentBrowser->SearchBoxPtr.IsValid()))
+						{
+							return;
+						}
+
+						FToolMenuEntry& SearchBoxEntry =
+							InSection.AddEntry(
+							FToolMenuEntry::InitWidget(
+								"SearchBox",
+								ContentBrowser->SearchBoxPtr.ToSharedRef(),
+								FText::GetEmpty()));
+
+						SearchBoxEntry.WidgetData.StyleParams.SizeRule = FSizeParam::SizeRule_StretchContent;
+						SearchBoxEntry.WidgetData.StyleParams.MinSize = SearchBoxMinWidth;
+						SearchBoxEntry.WidgetData.StyleParams.MaxSize = SearchBoxMaxWidth;
+						SearchBoxEntry.WidgetData.StyleParams.DesiredWidth = SearchBoxMaxWidth;
+					}
+				}
+			}));
+
+		FToolMenuSection& SettingsSection = ToolBar->AddSection("Settings");
+		SettingsSection.Alignment = EToolMenuSectionAlign::Last;
+
+		SettingsSection.AddDynamicEntry(NAME_None,
+		FNewToolMenuSectionDelegate::CreateLambda([](FToolMenuSection& InSection)
+		{
+			if (const UContentBrowserToolbarMenuContext* Context = InSection.FindContext<UContentBrowserToolbarMenuContext>())
+			{
+				if (Context->ContentBrowser.IsValid())
+				{
+					const TSharedRef<SContentBrowser> ContentBrowser = Context->ContentBrowser.Pin().ToSharedRef();
+
+					const TSharedRef<SWidget> LockToggleButton =
+						SNew(SActionButton)
+						.ActionButtonType(EActionButtonType::Simple)
+						.Icon(ContentBrowser, &SContentBrowser::GetLockIconBrush)
+						.IconColorAndOpacity(FSlateColor::UseStyle())
+						.ToolTipText(LOCTEXT("LockToggleTooltip", "Toggle lock. If locked, this browser will ignore Find in Content Browser requests."))
+						.OnClicked(ContentBrowser, &SContentBrowser::ToggleLockClicked)
+						.AddMetaData<FTagMetaData>(FTagMetaData(TEXT("ContentBrowserLock")));
+
+					InSection.AddEntry(
+						FToolMenuEntry::InitWidget(
+							"Lock",
+							LockToggleButton,
+							FText::GetEmpty()
+						));
+				}
+			}
+		}));
+#endif
+	});
+
 #undef LOCTEXT_NAMESPACE
diff --git a/Engine/Source/Editor/ContentBrowser/Private/SAssetView.cpp b/Engine/Source/Editor/ContentBrowser/Private/SAssetView.cpp
--- a/Engine/Source/Editor/ContentBrowser/Private/SAssetView.cpp
+++ b/Engine/Source/Editor/ContentBrowser/Private/SAssetView.cpp
@@ -76,6 +76,10 @@
 #include "ISourceControlModule.h"
 #include "RevisionControlStyle/RevisionControlStyle.h"
 
+#if UE_CONTENTBROWSER_NEW_STYLE
+#include "SPathView.h"
+#endif
+
 #define LOCTEXT_NAMESPACE "ContentBrowser"
 #define MAX_THUMBNAIL_SIZE 4096
 
@@ -1376,6 +1380,7 @@ void SAssetView::Construct( const FArguments& InArgs )
 		[
 			SNew(SHorizontalBox)
 
+#if !UE_CONTENTBROWSER_NEW_STYLE
 			// Asset count
 			+SHorizontalBox::Slot()
 			.FillWidth(1.f)
@@ -1385,6 +1390,7 @@ void SAssetView::Construct( const FArguments& InArgs )
 				SNew(STextBlock)
 				.Text(this, &SAssetView::GetAssetCountText)
 			]
+#endif
 
 			// View mode combo button
 			+SHorizontalBox::Slot()
diff --git a/Engine/Source/Editor/ContentBrowser/Private/SContentBrowser.cpp b/Engine/Source/Editor/ContentBrowser/Private/SContentBrowser.cpp
--- a/Engine/Source/Editor/ContentBrowser/Private/SContentBrowser.cpp
+++ b/Engine/Source/Editor/ContentBrowser/Private/SContentBrowser.cpp
@@ -157,6 +157,7 @@
 #include "Widgets/Text/STextBlock.h"
 #if UE_CONTENTBROWSER_NEW_STYLE
 #include "SActionButton.h"
+#include "Widgets/Colors/SColorBlock.h"
 #endif
 
 class FTreeItem;
@@ -169,8 +170,8 @@ const FString SContentBrowser::SettingsIniSection = TEXT("ContentBrowser");
 
 namespace ContentBrowserSourcesWidgetSwitcherIndex
 {
-	static const int32 PathView = 0;
-	static const int32 CollectionsView = 1;
+	static constexpr int32 PathView = 0;
+	static constexpr int32 CollectionsView = 1;
 }
 
 namespace ContentBrowser
@@ -490,6 +491,13 @@ void SContentBrowser::Construct( const FArguments& InArgs, const FName& InInstan
 	// Note, for backwards compatibility ShowBottomToolbar controls the visibility of view options so we respect that here
 	if (bShowBottomToolbar)
 	{
+#if UE_CONTENTBROWSER_NEW_STYLE
+		ViewOptions =
+			SNew(SActionButton)
+			.ActionButtonType(EActionButtonType::Simple)
+			.OnGetMenuContent(AssetViewPtr.ToSharedRef(), &SAssetView::GetViewButtonContent)
+			.Icon(FAppStyle::Get().GetBrush("Icons.Settings"));
+#else
 		ViewOptions =
 			SNew(SComboButton)
 			.ComboButtonStyle(&FAppStyle::Get().GetWidgetStyle<FComboButtonStyle>("SimpleComboButton"))
@@ -507,7 +515,7 @@ void SContentBrowser::Construct( const FArguments& InArgs, const FName& InInstan
 					.ColorAndOpacity(FSlateColor::UseForeground())
 					.Image(FAppStyle::Get().GetBrush("Icons.Settings"))
 				]
-#if !UE_CONTENTBROWSER_NEW_STYLE
+
 				+ SHorizontalBox::Slot()
 				.VAlign(VAlign_Center)
 				.Padding(4.0, 0.0f)
@@ -516,15 +524,15 @@ void SContentBrowser::Construct( const FArguments& InArgs, const FName& InInstan
 					.Text(LOCTEXT("Settings", "Settings"))
 					.ColorAndOpacity(FSlateColor::UseForeground())
 				]
-#endif
 			];
+#endif
 	}
 
 #if UE_CONTENTBROWSER_NEW_STYLE
-	constexpr float PanelInsetPadding = 4.0f;
+	constexpr float ToolBarVerticalPadding = 4.0f; // ToolBar total height should be 36.0f
+	constexpr float PanelInsetPadding = 2.0f;
 
-	ChildSlot
-	[
+	const TSharedRef<SWidget> PanelContentWidget =
 		SNew(SVerticalBox)
 
 		+ SVerticalBox::Slot()
@@ -534,223 +542,199 @@ void SContentBrowser::Construct( const FArguments& InArgs, const FName& InInstan
 			.Thickness(2.0f)
 		]
 
-		// Assets/tree
+		// Source / Tree + Assets
 		+ SVerticalBox::Slot()
 		.FillHeight(1.0f)
 		.Padding(0.0f)
 		[
-			// The tree/assets splitter
-			SAssignNew(PathAssetSplitterPtr, SSplitter)
-			.PhysicalSplitterHandleSize(2.0f)
+			SNew(SHorizontalBox)
 
-			// Sources View
-			+ SSplitter::Slot()
-			.Resizable(true)
-			.SizeRule(SSplitter::SizeToContent)
-			.OnSlotResized(this, &SContentBrowser::OnPathViewBoxColumnResized)
+			// Tree + Assets + Navigation Bar
+			+ SHorizontalBox::Slot()
+			.FillWidth(1.0f)
+			.Padding(FMargin(PanelInsetPadding, 0.0f, 0.0f, 0.0f))
 			[
-				SNew(SBox)
-				.Padding(FMargin(PanelInsetPadding, 0.0f, PanelInsetPadding, PanelInsetPadding))
-				.Visibility(this, &SContentBrowser::GetSourcesViewVisibility)
-				.WidthOverride(this, &SContentBrowser::GetPathViewBoxWidthOverride)
+				SNew(SVerticalBox)
+
+				// Assets/tree
+				+ SVerticalBox::Slot()
+				.FillHeight(1.0f)
+				.Padding(FMargin(0.0f, 0.0f, 0.0f, PanelInsetPadding))
 				[
-					SNew(SBorder)
-					.Padding(FMargin(0))
-					.BorderImage(FAppStyle::GetBrush("Brushes.Recessed"))
+					// The tree/assets splitter
+					SAssignNew(PathAssetSplitterPtr, SSplitter)
+					.PhysicalSplitterHandleSize(2.0f)
+
+					// Sources View
+					+ SSplitter::Slot()
+					.Resizable(true)
+					.SizeRule(SSplitter::SizeToContent)
+					.OnSlotResized(this, &SContentBrowser::OnPathViewBoxColumnResized)
 					[
-						// Note: If adding more widgets here, fix ContentBrowserSourcesWidgetSwitcherIndex and the code that uses it!
-						SAssignNew(SourcesWidgetSwitcher, SWidgetSwitcher)
-
-						// Paths View
-						+SWidgetSwitcher::Slot()
+						SNew(SBox)
+						.Padding(0.0f)
+						.Visibility(this, &SContentBrowser::GetSourcesViewVisibility)
+						.WidthOverride(this, &SContentBrowser::GetPathViewBoxWidthOverride)
 						[
-							SAssignNew(PathFavoriteSplitterPtr, SSplitter)
-							.Clipping(EWidgetClipping::ClipToBounds)
-							.PhysicalSplitterHandleSize(2.0f)
-							.HitDetectionSplitterHandleSize(8.0f)
-							.Orientation(EOrientation::Orient_Vertical)
-							.MinimumSlotHeight(26.0f)
-							.Visibility( this, &SContentBrowser::GetSourcesViewVisibility )
-							+SSplitter::Slot()
-							.SizeRule(TAttribute<SSplitter::ESizeRule>(this, &SContentBrowser::GetFavoritesAreaSizeRule))
-							.MinSize(TAttribute<float>(this, &SContentBrowser::GetFavoritesAreaMinSize))
-							.Value(0.2f)
+							SNew(SBorder)
+							.Padding(FMargin(0))
+							.BorderImage(FAppStyle::GetBrush("Brushes.Recessed"))
 							[
-								SNew(SBorder)
-								.BorderImage(FAppStyle::Get().GetBrush("Brushes.Header"))
-								.Padding(0.0f, 2.0f, 0.0f, 0.0f)
-								[
-									CreateFavoritesView(Config)
-								]
-							]
+								// Note: If adding more widgets here, fix ContentBrowserSourcesWidgetSwitcherIndex and the code that uses it!
+								SAssignNew(SourcesWidgetSwitcher, SWidgetSwitcher)
 
-							+SSplitter::Slot()
-							.SizeRule(TAttribute<SSplitter::ESizeRule>(this, &SContentBrowser::GetPathAreaSizeRule))
-							.MinSize(29.0f)
-							.Value(0.8f)
-							[
-								SNew(SBorder)
-								.BorderImage(FAppStyle::Get().GetBrush("Brushes.Header"))
-								.Padding(0.0f, 2.0f, 0.0f, 0.0f)
+								// Paths View
+								+ SWidgetSwitcher::Slot()
 								[
-									CreatePathView(Config)
+									SAssignNew(PathFavoriteSplitterPtr, SSplitter)
+									.Clipping(EWidgetClipping::ClipToBounds)
+									.PhysicalSplitterHandleSize(2.0f)
+									.HitDetectionSplitterHandleSize(8.0f)
+									.Orientation(EOrientation::Orient_Vertical)
+									.MinimumSlotHeight(26.0f)
+									.Visibility(this, &SContentBrowser::GetSourcesViewVisibility)
+									+ SSplitter::Slot()
+									.SizeRule(TAttribute<SSplitter::ESizeRule>(this, &SContentBrowser::GetFavoritesAreaSizeRule))
+									.MinSize(TAttribute<float>(this, &SContentBrowser::GetFavoritesAreaMinSize))
+									.Value(0.2f)
+									[
+										SNew(SBorder)
+										.BorderImage(FAppStyle::Get().GetBrush("Brushes.Header"))
+										.Padding(0.0f, 2.0f, 0.0f, 0.0f)
+										[
+											CreateFavoritesView(Config)
+										]
+									]
+
+									+ SSplitter::Slot()
+									.SizeRule(TAttribute<SSplitter::ESizeRule>(this, &SContentBrowser::GetPathAreaSizeRule))
+									.MinSize(29.0f)
+									.Value(0.8f)
+									[
+										SNew(SBorder)
+										.BorderImage(FAppStyle::Get().GetBrush("Brushes.Header"))
+										.Padding(0.0f, 2.0f, 0.0f, 0.0f)
+										[
+											CreatePathView(Config)
+										]
+									]
+
+									+ SSplitter::Slot()
+									.SizeRule(TAttribute<SSplitter::ESizeRule>(this, &SContentBrowser::GetCollectionsAreaSizeRule))
+									.MinSize(TAttribute<float>(this, &SContentBrowser::GetCollectionsAreaMinSize))
+									.Value(0.4f)
+									[
+										SNew(SBorder)
+										.BorderImage(FAppStyle::Get().GetBrush("Brushes.Header"))
+										.Padding(0.0f, 2.0f, 0.0f, 0.0f)
+										[
+											CreateDockedCollectionsView(Config)
+										]
+									]
 								]
-							]
 
-							+SSplitter::Slot()
-							.SizeRule(TAttribute<SSplitter::ESizeRule>(this, &SContentBrowser::GetCollectionsAreaSizeRule))
-							.MinSize(TAttribute<float>(this, &SContentBrowser::GetCollectionsAreaMinSize))
-							.Value(0.4f)
-							[
-								SNew(SBorder)
-								.BorderImage(FAppStyle::Get().GetBrush("Brushes.Header"))
-								.Padding(0.0f, 2.0f, 0.0f, 0.0f)
+								// Collections View
+								+ SWidgetSwitcher::Slot()
 								[
-									CreateDockedCollectionsView(Config)
+									SNew(SBox)
+									.Padding(FMargin(0.0f, PanelInsetPadding, 0.0f, 0.0f))
+									[
+										CollectionViewPtr.ToSharedRef()
+									]
 								]
 							]
 						]
-
-						// Collections View
-						+SWidgetSwitcher::Slot()
-						[
-							SNew(SBox)
-							.Padding(FMargin(0.0f, PanelInsetPadding, 0.0f, 0.0f))
-							[
-								CollectionViewPtr.ToSharedRef()
-							]
-						]
 					]
-				]
-			]
-
-			// Asset View
-			+ SSplitter::Slot()
-			.Value(0.75f)
-			[
-				SNew(SBox)
-				.Padding(FMargin(0.0f, 0.0f, 0.0f, PanelInsetPadding))
-				[
-					SNew(SVerticalBox)
 
-					+ SVerticalBox::Slot()
-					.AutoHeight()
+					// Asset View
+					+ SSplitter::Slot()
+					.Value(0.75f)
 					[
-						SNew(SBorder)
-						.Padding(FMargin(3))
-						.BorderImage(bIsDrawer ? FStyleDefaults::GetNoBrush() : FAppStyle::Get().GetBrush("Brushes.Panel"))
+						SNew(SBox)
+						.Padding(0.0f)
 						[
-							SNew(SHorizontalBox)
+							SNew(SVerticalBox)
 
-							+ SHorizontalBox::Slot()
-							.AutoWidth()
-							.HAlign(HAlign_Left)
-							.Padding(5, 0, 0, 0)
+							+ SVerticalBox::Slot()
+							.AutoHeight()
 							[
-								CreateToolBar(Config)
-							]
-
-							+ SHorizontalBox::Slot()
-							.AutoWidth()
-							.HAlign(HAlign_Right)
-							.VAlign(VAlign_Center)
-							[
-								CreateLockButton(Config)
+								SNew(SBorder)
+								.Padding(FMargin(3, ToolBarVerticalPadding))
+								.BorderImage(bIsDrawer ? FStyleDefaults::GetNoBrush() : FAppStyle::Get().GetBrush("Brushes.Panel"))
+								[
+									SNew(SHorizontalBox)
+
+									+ SHorizontalBox::Slot()
+									.FillWidth(1.0f)
+									.HAlign(HAlign_Fill)
+									.VAlign(VAlign_Center)
+									.Padding(5, 0, 0, 0)
+									[
+										CreateToolBar(Config)
+									]
+
+									+ SHorizontalBox::Slot()
+									.AutoWidth()
+									.HAlign(HAlign_Right)
+									.VAlign(VAlign_Center)
+									[
+										CreateDrawerDockButton(Config)
+									]
+
+									+ SHorizontalBox::Slot()
+									.AutoWidth()
+									.HAlign(HAlign_Right)
+									.VAlign(VAlign_Center)
+									.Padding(5.0f, 0.0f, 0.0f, 0.0f)
+									[
+										ViewOptions
+									]
+								]
 							]
 
-							+ SHorizontalBox::Slot()
-							.AutoWidth()
-							.HAlign(HAlign_Right)
-							.VAlign(VAlign_Center)
+							+ SVerticalBox::Slot()
 							[
-								CreateDrawerDockButton(Config)
+								CreateAssetView(Config)
 							]
 
-							+ SHorizontalBox::Slot()
-							.Padding(5.0f, 0.0f, 0.0f, 0.0f)
-							.HAlign(HAlign_Right)
-							.VAlign(VAlign_Center)
+							// Navigation Bar
+							+ SVerticalBox::Slot()
+							.AutoHeight()
+							.HAlign(HAlign_Fill)
+							.Padding(FMargin(0, PanelInsetPadding, 0, 0))
 							[
-								ViewOptions
+								SNew(SBorder)
+								.Padding(FMargin(3))
+								.BorderImage(bIsDrawer ? FStyleDefaults::GetNoBrush() : FAppStyle::Get().GetBrush("Brushes.Panel"))
+								[
+									CreateNavigationBar(Config)
+								]
 							]
 						]
 					]
-
-					+ SVerticalBox::Slot()
-					[
-						CreateAssetView(Config)
-					]
 				]
 			]
-		]
-
-		// Navigation Bar
-		+ SVerticalBox::Slot()
-		.AutoHeight()
-		[
-			SNew(SBorder)
-			.Padding(FMargin(3))
-			.BorderImage(bIsDrawer ? FStyleDefaults::GetNoBrush() : FAppStyle::Get().GetBrush("Brushes.Panel"))
-			[
-				SNew(SHorizontalBox)
+		];
 
-				// History Back Button
-				+SHorizontalBox::Slot()
-				.Padding(10, 0, 0, 0)
-				.AutoWidth()
-				[
-					SNew(SButton)
-					.VAlign(EVerticalAlignment::VAlign_Center)
-					.ButtonStyle(FAppStyle::Get(), "SimpleButton")
-					.ToolTipText( this, &SContentBrowser::GetHistoryBackTooltip )
-					.ContentPadding( FMargin(1, 0) )
-					.OnClicked(this, &SContentBrowser::BackClicked)
-					.IsEnabled(this, &SContentBrowser::IsBackEnabled)
-					.AddMetaData<FTagMetaData>(FTagMetaData(TEXT("ContentBrowserHistoryBack")))
-					[
-						SNew(SImage)
-						.Image(FAppStyle::Get().GetBrush("Icons.CircleArrowLeft"))
-						.ColorAndOpacity(FSlateColor::UseForeground())
-					]
-				]
+	ChildSlot
+	[
+		SNew(SOverlay)
 
-				// History Forward Button
-				+ SHorizontalBox::Slot()
-				.Padding(2, 0, 0, 0)
-				.AutoWidth()
-				[
-					SNew(SButton)
-					.VAlign(EVerticalAlignment::VAlign_Center)
-					.ButtonStyle(FAppStyle::Get(), "SimpleButton")
-					.ToolTipText( this, &SContentBrowser::GetHistoryForwardTooltip )
-					.ContentPadding( FMargin(1, 0) )
-					.OnClicked(this, &SContentBrowser::ForwardClicked)
-					.IsEnabled(this, &SContentBrowser::IsForwardEnabled)
-					.AddMetaData<FTagMetaData>(FTagMetaData(TEXT("ContentBrowserHistoryForward")))
-					[
-						SNew(SImage)
-						.Image(FAppStyle::Get().GetBrush("Icons.CircleArrowRight"))
-						.ColorAndOpacity(FSlateColor::UseForeground())
-					]
-				]
+		+ SOverlay::Slot()
+		.HAlign(HAlign_Fill)
+		.VAlign(VAlign_Fill)
+		.Padding(0)
+		[
+			SNew(SColorBlock)
+			.Color(FStyleColors::Background.GetSpecifiedColor()) // Always apply the Background color, rather than relying that set by the Tab which is context sensitive
+		]
 
-				// Path
-				+ SHorizontalBox::Slot()
-				.VAlign(VAlign_Center)
-				.HAlign(HAlign_Fill)
-				.FillWidth(1.0f)
-				.Padding(2, 0, 0, 0)
-				[
-					SAssignNew(NavigationBar, SNavigationBar)
-					.OnPathClicked(this, &SContentBrowser::OnPathClicked)
-					.GetPathMenuContent(this, &SContentBrowser::OnGetCrumbDelimiterContent)
-					.GetComboOptions(this, &SContentBrowser::GetRecentPaths)
-					.OnNavigateToPath(this, &SContentBrowser::OnNavigateToPath)
-					.OnCompletePrefix(this, &SContentBrowser::OnCompletePathPrefix)
-					.OnCanEditPathAsText(this, &SContentBrowser::OnCanEditPathAsText)
-					.AddMetaData<FTagMetaData>(FTagMetaData(TEXT("ContentBrowserPath")))
-				]
-			]
+		+ SOverlay::Slot()
+		.HAlign(HAlign_Fill)
+		.VAlign(VAlign_Fill)
+		.Padding(0)
+		[
+			PanelContentWidget
 		]
 	];
 #else
@@ -1205,18 +1189,34 @@ void SContentBrowser::HandleAssetViewSearchOptionsChanged()
 
 TSharedRef<SWidget> SContentBrowser::CreateToolBar(const FContentBrowserConfig* Config)
 {
-	RegisterContentBrowserToolBar();
-
 	FToolMenuContext MenuContext;
 
 	UContentBrowserToolbarMenuContext* CommonContextObject = NewObject<UContentBrowserToolbarMenuContext>();
 	CommonContextObject->ContentBrowser = SharedThis(this);
+	CommonContextObject->AssetView = AssetViewPtr;
+	CommonContextObject->ContentBrowserConfig = Config;
 
 	MenuContext.AddObject(CommonContextObject);
 
 	return UToolMenus::Get()->GenerateWidget("ContentBrowser.ToolBar", MenuContext);
 }
 
+#if UE_CONTENTBROWSER_NEW_STYLE
+TSharedRef<SWidget> SContentBrowser::CreateNavigationBar(const FContentBrowserConfig* Config)
+{
+	FToolMenuContext MenuContext;
+
+	UContentBrowserToolbarMenuContext* CommonContextObject = NewObject<UContentBrowserToolbarMenuContext>();
+	CommonContextObject->ContentBrowser = SharedThis(this);
+	CommonContextObject->AssetView = AssetViewPtr;
+	CommonContextObject->ContentBrowserConfig = Config;
+
+	MenuContext.AddObject(CommonContextObject);
+
+	return UToolMenus::Get()->GenerateWidget("ContentBrowser.NavigationBar", MenuContext);
+}
+#endif
+
 TSharedRef<SWidget> SContentBrowser::CreateLockButton(const FContentBrowserConfig* Config)
 {
 	if(Config == nullptr || Config->bCanShowLockButton)
@@ -1241,12 +1241,96 @@ TSharedRef<SWidget> SContentBrowser::CreateLockButton(const FContentBrowserConfi
 void SContentBrowser::OnFilterBarLayoutChanging(EFilterBarLayout NewLayout)
 {
 #if UE_CONTENTBROWSER_NEW_STYLE
-	const FOptionalSize SearchBoxDesiredWidth;
-	constexpr float SearchBoxMaxWidth = 640.0f;
+	if (NewLayout == EFilterBarLayout::Horizontal)
+	{
+		AssetViewBorder->SetContent(
+
+			SNew(SVerticalBox)
+
+			+ SVerticalBox::Slot()
+			.Padding(0.0f)
+			.AutoHeight()
+			[
+				FilterListPtr.ToSharedRef()
+			]
+
+			+ SVerticalBox::Slot()
+			.FillHeight(1.0f)
+			.Padding(0, 0)
+			[
+				AssetViewPtr.ToSharedRef()
+			]
+		);
+	}
+	else
+	{
+		AssetViewBorder->SetContent(
+
+			SNew(SSplitter)
+			.PhysicalSplitterHandleSize(2.0f)
+
+			/* Filters in an SScrollBox */
+			+ SSplitter::Slot()
+			.MinSize(95.f)
+			.Resizable(true)
+			.SizeRule(SSplitter::SizeToContent)
+			.OnSlotResized(this, &SContentBrowser::OnFilterBoxColumnResized)
+			[
+				SNew(SBox)
+				.WidthOverride(this, &SContentBrowser::GetFilterViewBoxWidthOverride)
+				// Don't take up space when there are no filters
+				.Visibility_Lambda([this]
+				{
+					return FilterListPtr->HasAnyFilters() ? EVisibility::Visible : EVisibility::Collapsed;
+				})
+				[
+					SNew(SVerticalBox)
+
+					// Header
+					+ SVerticalBox::Slot()
+					.Padding(0.0f, 2.0f)
+					.AutoHeight()
+					[
+
+						SNew(SBorder)
+						.BorderImage(FAppStyle::Get().GetBrush("Brushes.Header"))
+						.Padding(FMargin(8.0f, 6.0f))
+						.Content()
+						[
+							SNew(STextBlock)
+							.Text(LOCTEXT("FilterListVerticalHeader", "Filters"))
+							.TextStyle(FAppStyle::Get(), "ButtonText")
+							.Font(FAppStyle::Get().GetFontStyle("NormalFontBold"))
+						]
+					]
+
+					// Filter List
+					+ SVerticalBox::Slot()
+					.Padding(0.0f)
+					.FillHeight(1.0f)
+					[
+						FilterListPtr.ToSharedRef()
+					]
+				]
+			]
+
+			+ SSplitter::Slot()
+			.Value(0.88f)
+			[
+				SNew(SVerticalBox)
+
+				+ SVerticalBox::Slot()
+				.FillHeight(1.0f)
+				.Padding(0, 0)
+				[
+					AssetViewPtr.ToSharedRef()
+				]
+			]
+		);
+	}
 #else
 	const FOptionalSize SearchBoxDesiredWidth = 500.0f;
 	constexpr float SearchBoxMaxWidth = 0.0f;
-#endif
 
 	if (NewLayout == EFilterBarLayout::Horizontal)
 	{
@@ -1296,7 +1380,6 @@ void SContentBrowser::OnFilterBarLayoutChanging(EFilterBarLayout NewLayout)
 			[
 				AssetViewPtr.ToSharedRef()
 			]
-
 		);
 	}
 	else
@@ -1391,6 +1474,7 @@ void SContentBrowser::OnFilterBarLayoutChanging(EFilterBarLayout NewLayout)
 			]
 		);
 	}
+#endif
 }
 
 TSharedRef<SWidget> SContentBrowser::CreateAssetView(const FContentBrowserConfig* Config)
@@ -1437,12 +1521,28 @@ TSharedRef<SWidget> SContentBrowser::CreateAssetView(const FContentBrowserConfig
 				   .Visibility((Config != nullptr ? Config->bCanShowAssetSearch : true) ? EVisibility::Visible : EVisibility::Collapsed)
 				   .AddMetaData<FTagMetaData>(FTagMetaData(TEXT("ContentBrowserSearchAssets")));
 
-	/* Create the Border that the Asset View will live in, the actual layout is populated in OnFilterBarLayoutChanging, 
+	NavigationBar = SNew(SNavigationBar)
+					.OnPathClicked(this, &SContentBrowser::OnPathClicked)
+					.GetPathMenuContent(this, &SContentBrowser::OnGetCrumbDelimiterContent)
+					.GetComboOptions(this, &SContentBrowser::GetRecentPaths)
+					.OnNavigateToPath(this, &SContentBrowser::OnNavigateToPath)
+					.OnCompletePrefix(this, &SContentBrowser::OnCompletePathPrefix)
+					.OnCanEditPathAsText(this, &SContentBrowser::OnCanEditPathAsText)
+					.AddMetaData<FTagMetaData>(FTagMetaData(TEXT("ContentBrowserPath")));
+
+	/* Create the Border that the Asset View will live in, the actual layout is populated in OnFilterBarLayoutChanging,
 	 * which is called initially called through SAssetFilterBar::LoadSettings through SContentBrowser::LoadSettings()
 	 */
+
+#if UE_CONTENTBROWSER_NEW_STYLE
+	static const FMargin AssetViewPadding(2.0f, 0.0f, 2.0f, 0.0f);
+#else
+	static const FMargin AssetViewPadding(2.0f, 2.0f, 2.0f, 0.0f);
+#endif
+
 	AssetViewBorder = SNew(SBorder)
 			.BorderImage(FAppStyle::Get().GetBrush("Brushes.Panel"))
-			.Padding(FMargin(2.0f, 2.0f, 2.0f, 0.0f));
+			.Padding(AssetViewPadding);
 
 	return AssetViewBorder.ToSharedRef();
 }
@@ -1733,7 +1833,9 @@ void SContentBrowser::ExtendViewOptionsMenu(const FContentBrowserConfig* Config)
 		return;
 	}
 
-	Menu->AddDynamicSection("ContentBrowserViewOptionsSection", FNewToolMenuDelegate::CreateLambda([bShowLockButton, bShowSourcesView](UToolMenu* InMenu)
+	Menu->AddDynamicSection("ContentBrowserViewOptionsSection",
+		FNewToolMenuDelegate::CreateLambda(
+			[bShowLockButton, bShowSourcesView](UToolMenu* InMenu)
 		{
 			if (UContentBrowserAssetViewContextMenuContext* Context = InMenu->FindContext<UContentBrowserAssetViewContextMenuContext>())
 			{
@@ -1741,6 +1843,18 @@ void SContentBrowser::ExtendViewOptionsMenu(const FContentBrowserConfig* Config)
 				{
 					if (bShowLockButton)
 					{
+#if UE_CONTENTBROWSER_NEW_STYLE
+						FToolMenuSection& Section = InMenu->FindOrAddSection("Manage");
+						Section.AddMenuEntry(
+							"ToggleLock",
+							TAttribute<FText>(ContentBrowser.ToSharedRef(), &SContentBrowser::GetLockMenuText),
+							LOCTEXT("LockToggleTooltip", "Toggle lock. If locked, this browser will ignore Find in Content Browser requests."),
+							TAttribute<FSlateIcon>(ContentBrowser.ToSharedRef(), &SContentBrowser::GetLockIcon),
+							FUIAction(
+								FExecuteAction::CreateLambda([ContentBrowser = Context->OwningContentBrowser]() { ContentBrowser.Pin()->ToggleLockClicked(); })
+							)
+						);
+#else
 						FToolMenuSection& Section = InMenu->AddSection("Locking", LOCTEXT("LockingMenuHeader", "Locking"), FToolMenuInsert("AssetViewType", EToolMenuInsertType::After));
 						Section.AddMenuEntry(
 							"ToggleLock",
@@ -1751,8 +1865,9 @@ void SContentBrowser::ExtendViewOptionsMenu(const FContentBrowserConfig* Config)
 								FExecuteAction::CreateLambda([ContentBrowser = Context->OwningContentBrowser]() { ContentBrowser.Pin()->ToggleLockClicked(); })
 							)
 						);
+#endif
 					}
-
+#if !UE_CONTENTBROWSER_NEW_STYLE
 					if (bShowSourcesView)
 					{
 						FToolMenuSection& Section = InMenu->FindOrAddSection("View");
@@ -1768,113 +1883,13 @@ void SContentBrowser::ExtendViewOptionsMenu(const FContentBrowserConfig* Config)
 							EUserInterfaceActionType::ToggleButton
 						);
 					}
+#endif
 				}
 			}
 		}
 	));
 }
 
-void SContentBrowser::RegisterContentBrowserToolBar()
-{
-	static const FName ToolBarName("ContentBrowser.ToolBar");
-	UToolMenus* ToolMenus = UToolMenus::Get();
-	if (ToolMenus->IsMenuRegistered(ToolBarName))
-	{
-		return;
-	}
-
-	UToolMenu* ToolBar = UToolMenus::Get()->RegisterMenu(ToolBarName, NAME_None, EMultiBoxType::SlimHorizontalToolBar);
-	ToolBar->StyleName = "ContentBrowser.ToolBar";
-
-	{
-		FToolMenuSection& Section = ToolBar->AddSection("New");
-
-		Section.AddDynamicEntry("New", FNewToolMenuSectionDelegate::CreateLambda([](FToolMenuSection& InSection)
-		{
-			UContentBrowserToolbarMenuContext* Context = InSection.FindContext<UContentBrowserToolbarMenuContext>();
-			TSharedRef<SContentBrowser> ContentBrowser = Context->ContentBrowser.Pin().ToSharedRef();
-		
-			TSharedRef<SPositiveActionButton> NewButton =
-				SNew(SPositiveActionButton)
-				.OnGetMenuContent_Lambda([Context] { return Context->ContentBrowser.Pin()->MakeAddNewContextMenu(EContentBrowserDataMenuContext_AddNewMenuDomain::Toolbar, Context); })
-				.ToolTipText(ContentBrowser, &SContentBrowser::GetAddNewToolTipText)
-				.IsEnabled(ContentBrowser, &SContentBrowser::IsAddNewEnabled)
-				.AddMetaData<FTagMetaData>(FTagMetaData(TEXT("ContentBrowserNewAsset")))
-				.Icon(FAppStyle::Get().GetBrush("Icons.Plus"))
-				.Text(LOCTEXT("AddAssetButton", "Add"));
-		
-			InSection.AddEntry(
-				FToolMenuEntry::InitWidget(
-					"NewButton",
-					NewButton,
-					FText::GetEmpty(),
-					true,
-					false
-				));
-		}));
-	}
-
-	{
-		FToolMenuSection& Section = ToolBar->AddSection("Save");
-		Section.AddDynamicEntry("Save", FNewToolMenuSectionDelegate::CreateLambda([](FToolMenuSection& InSection)
-		{
-			UContentBrowserToolbarMenuContext* Context = InSection.FindContext<UContentBrowserToolbarMenuContext>();
-			TSharedRef<SContentBrowser> ContentBrowser = Context->ContentBrowser.Pin().ToSharedRef();
-
-#if UE_CONTENTBROWSER_NEW_STYLE
-				TSharedRef<SWidget> SaveButton =
-					SNew(SActionButton)
-					.ToolTipText(LOCTEXT("SaveDirtyPackagesTooltip", "Save all modified assets."))
-					.OnClicked(ContentBrowser, &SContentBrowser::OnSaveClicked)
-					.AddMetaData<FTagMetaData>(FTagMetaData(TEXT("ContentBrowserSaveDirtyPackages")))
-					.Icon(FAppStyle::Get().GetBrush("MainFrame.SaveAll"))
-					.Text(LOCTEXT("SaveAll", "Save All"));
-#else
-				TSharedRef<SWidget> SaveButton =
-					SNew(SButton)
-					.ButtonStyle(FAppStyle::Get(), "SimpleButton")
-					.ToolTipText(LOCTEXT("SaveDirtyPackagesTooltip", "Save all modified assets."))
-					.ContentPadding(2.f)
-					.OnClicked(ContentBrowser, &SContentBrowser::OnSaveClicked)
-					.AddMetaData<FTagMetaData>(FTagMetaData(TEXT("ContentBrowserSaveDirtyPackages")))
-					[
-						SNew(SHorizontalBox)
-						// Save All Icon
-						+ SHorizontalBox::Slot()
-						.AutoWidth()
-						.HAlign(HAlign_Center)
-						.VAlign(VAlign_Center)
-						[
-							SNew(SImage)
-							.Image(FAppStyle::Get().GetBrush("MainFrame.SaveAll"))
-							.ColorAndOpacity(FSlateColor::UseForeground())
-						]
-
-						// Save All Text
-						+ SHorizontalBox::Slot()
-						.Padding(FMargin(3, 0, 0, 0))
-						.VAlign(VAlign_Center)
-						.AutoWidth()
-						[
-							SNew(STextBlock)
-							.TextStyle(FAppStyle::Get(), "NormalText")
-							.Text(LOCTEXT("SaveAll", "Save All"))
-						]
-					];
-#endif
-
-			InSection.AddEntry(
-				FToolMenuEntry::InitWidget(
-					"SaveButton",
-					SaveButton,
-					FText::GetEmpty(),
-					true,
-					false
-				));
-		}));
-	}
-}
-
 SSplitter::ESizeRule SContentBrowser::GetFavoritesAreaSizeRule() const
 {
 	// Make sure the area is expanded and visible 
@@ -3161,33 +3176,6 @@ FString SContentBrowser::GetCurrentPath(const EContentBrowserPathType PathType)
 
 void SContentBrowser::AppendNewMenuContextObjects(const EContentBrowserDataMenuContext_AddNewMenuDomain InDomain, const TArray<FName>& InSelectedPaths, FToolMenuContext& InOutMenuContext, UContentBrowserToolbarMenuContext* CommonContext, bool bCanBeModified)
 {
-	if (!UToolMenus::Get()->IsMenuRegistered("ContentBrowser.AddNewContextMenu"))
-	{
-		UToolMenu* Menu = UToolMenus::Get()->RegisterMenu("ContentBrowser.AddNewContextMenu");
-		Menu->AddDynamicSection("DynamicSection_Common", FNewToolMenuDelegate::CreateLambda([](UToolMenu* InMenu)
-		{
-			TSharedPtr<SContentBrowser> ContentBrowser;
-			const UContentBrowserMenuContext* MenuContext = InMenu->FindContext<UContentBrowserMenuContext>();
-			if (MenuContext)
-			{
-				ContentBrowser = MenuContext->ContentBrowser.Pin();
-			}
-			else
-			{
-				const UContentBrowserToolbarMenuContext* ToolbarContext  = InMenu->FindContext<UContentBrowserToolbarMenuContext>();
-				if (ToolbarContext)
-				{
-					ContentBrowser = ToolbarContext->ContentBrowser.Pin();
-				}
-			}
-
-			if (ContentBrowser)
-			{
-				ContentBrowser->PopulateAddNewContextMenu(InMenu);
-			}
-		}));
-	}
-
 	if(!CommonContext)
 	{
 		UContentBrowserMenuContext* CommonContextObject = NewObject<UContentBrowserMenuContext>();
@@ -4616,20 +4604,6 @@ TSharedPtr<SWidget> SContentBrowser::GetItemContextMenu(TArrayView<const FConten
 		// Ensure the path context menu has the up-to-date list of paths being worked on
 		PathContextMenu->SetSelectedFolders(SelectedFolders);
 
-		if (!UToolMenus::Get()->IsMenuRegistered("ContentBrowser.FolderContextMenu"))
-		{
-			UToolMenu* Menu = UToolMenus::Get()->RegisterMenu("ContentBrowser.FolderContextMenu");
-			Menu->bCloseSelfOnly = true;
-			Menu->AddDynamicSection("Section", FNewToolMenuDelegate::CreateLambda([](UToolMenu* InMenu)
-			{
-				UContentBrowserFolderContext* Context = InMenu->FindContext<UContentBrowserFolderContext>();
-				if (Context && Context->ContentBrowser.IsValid())
-				{
-					Context->ContentBrowser.Pin()->PopulateFolderContextMenu(InMenu);
-				}
-			}));
-		}
-
 		TArray<FString> SelectedPackagePaths;
 		bool bPhysicalPathExists = false;
 		for (const FContentBrowserItem& SelectedFolder : SelectedFolders)
diff --git a/Engine/Source/Editor/ContentBrowser/Private/SContentBrowser.h b/Engine/Source/Editor/ContentBrowser/Private/SContentBrowser.h
--- a/Engine/Source/Editor/ContentBrowser/Private/SContentBrowser.h
+++ b/Engine/Source/Editor/ContentBrowser/Private/SContentBrowser.h
@@ -591,6 +591,9 @@ private:
 
 	/** Creates various widgets for the content browser main view */
 	TSharedRef<SWidget> CreateToolBar(const FContentBrowserConfig* Config);
+#if UE_CONTENTBROWSER_NEW_STYLE
+	TSharedRef<SWidget> CreateNavigationBar(const FContentBrowserConfig* Config);
+#endif
 	TSharedRef<SWidget> CreateLockButton(const FContentBrowserConfig* Config);
 	TSharedRef<SWidget> CreateAssetView(const FContentBrowserConfig* Config);
 	TSharedRef<SWidget> CreateFavoritesView(const FContentBrowserConfig* Config);
@@ -604,9 +607,6 @@ private:
 	/** Adds menu options to the view menu */
 	void ExtendViewOptionsMenu(const FContentBrowserConfig* Config);
 
-	/** Creates the content browser toolbar menu */
-	void RegisterContentBrowserToolBar();
-
 	/** Gets the size rule for various areas. When areas are a collapsed the splitter slot becomes auto sized, otherwise it is user sized */
 	SSplitter::ESizeRule GetFavoritesAreaSizeRule() const;
 	SSplitter::ESizeRule GetPathAreaSizeRule() const;
@@ -784,11 +784,24 @@ private:
 	/** Handle to a callback for when CVars change */
 	FConsoleVariableSinkHandle CVarSinkHandle;
 
+	/** Registers the "Add New" drop-down menu */
+	static FDelayedAutoRegisterHelper AddNewContextMenuRegistration;
+
+#if UE_CONTENTBROWSER_NEW_STYLE
+	/** Registers the bottom navigation bar */
+	static FDelayedAutoRegisterHelper NavigationBarMenuRegistration;
+#endif
+
+	/** Registers the context menu for folders */
+	static FDelayedAutoRegisterHelper FolderContextMenuRegistration;
+
 	/** Registers a sub-menu within the settings button menu */
 	static FDelayedAutoRegisterHelper PathViewFiltersMenuRegistration;
 
-public: 
+	/** Registers the top toolbar */
+	static FDelayedAutoRegisterHelper ToolBarMenuRegistration;
 
+public:
 	/** The section of EditorPerProjectUserSettings in which to save content browser settings */
 	static const FString SettingsIniSection;
 };
diff --git a/Engine/Source/Editor/ContentBrowser/Public/ContentBrowserMenuContexts.h b/Engine/Source/Editor/ContentBrowser/Public/ContentBrowserMenuContexts.h
--- a/Engine/Source/Editor/ContentBrowser/Public/ContentBrowserMenuContexts.h
+++ b/Engine/Source/Editor/ContentBrowser/Public/ContentBrowserMenuContexts.h
@@ -23,6 +23,7 @@ class SContentBrowser;
 class SFilterList;
 class UClass;
 class UAssetDefinition;
+struct FContentBrowserConfig;
 struct FFrame;
 struct FToolMenuSection;
 struct FToolMenuContext;
@@ -328,6 +329,8 @@ public:
 	bool CanWriteToCurrentPath() const;
 
 	TWeakPtr<SContentBrowser> ContentBrowser;
+	TWeakPtr<SAssetView> AssetView;
+	const FContentBrowserConfig* ContentBrowserConfig = nullptr;
 };
 
 class UToolMenu;
diff --git a/Engine/Source/Editor/ContentBrowser/Public/SAssetView.h b/Engine/Source/Editor/ContentBrowser/Public/SAssetView.h
--- a/Engine/Source/Editor/ContentBrowser/Public/SAssetView.h
+++ b/Engine/Source/Editor/ContentBrowser/Public/SAssetView.h
@@ -106,7 +106,8 @@ enum class EAssetViewCopyType
 /**
  * A widget to display a list of filtered assets
  */
-class CONTENTBROWSER_API SAssetView : public SCompoundWidget
+class CONTENTBROWSER_API SAssetView
+	: public SCompoundWidget
 {
 public:
 	SLATE_BEGIN_ARGS( SAssetView )
@@ -315,9 +316,10 @@ public:
 	SLATE_END_ARGS()
 
 	friend FAssetContextMenu;
+	friend class UContentBrowserAssetViewContextMenuContext;
 
 	SAssetView();
-	~SAssetView();
+	virtual ~SAssetView() override;
 
 	/** Constructs this widget with InArgs */
 	void Construct( const FArguments& InArgs );
@@ -561,11 +563,16 @@ private:
 	/** Returns the thumbnails hint color and opacity */
 	FLinearColor GetThumbnailHintColorAndOpacity() const;
 
+	/** Populate the given params for this AssetView */
+	void PopulateFilterAdditionalParams(FFiltersAdditionalParams& OutParams);
+
 	/** Registers the menu used by the settings button */
 	static FDelayedAutoRegisterHelper AssetViewOptionsMenuRegistration;
 
-	/** Populate the given params for this AssetView */
-	void PopulateFilterAdditionalParams(FFiltersAdditionalParams& OutParams);
+#if UE_CONTENTBROWSER_NEW_STYLE
+	/** Extend the bottom navigation bar */
+	static FDelayedAutoRegisterHelper NavigationBarMenuExtensionRegistration;
+#endif
 
 	/** Toggle whether folders should be shown or not */
 	void ToggleShowFolders();
