Add tests for creating and deleting custom text filters

#virtualized

diff --git a/Engine/Source/Editor/ContentBrowser/Private/ContentBrowserSingleton.cpp b/Engine/Source/Editor/ContentBrowser/Private/ContentBrowserSingleton.cpp
--- a/Engine/Source/Editor/ContentBrowser/Private/ContentBrowserSingleton.cpp
+++ b/Engine/Source/Editor/ContentBrowser/Private/ContentBrowserSingleton.cpp
@@ -22,6 +22,7 @@
 #include "Editor/EditorEngine.h"
 #include "EditorDirectories.h"
 #include "EngineLogs.h"
+#include "Filters/CustomTextFilters.h"
 #include "Framework/Application/SlateApplication.h"
 #include "Framework/Commands/UIAction.h"
 #include "Framework/Commands/UICommandInfo.h"
@@ -1360,4 +1361,60 @@ TArray<FString> FContentBrowserSingleton::GetAliasesForPath(const FSoftObjectPat
 	return OutPaths;
 }
 
+void FContentBrowserSingleton::AddCustomTextFilter(const FCustomTextFilterData& FilterData, bool bApplyFilter)
+{
+	if (!HasPrimaryContentBrowser())
+	{
+		ChooseNewPrimaryBrowser();
+	}
+
+	if (PrimaryContentBrowser.IsValid())
+	{
+		PrimaryContentBrowser.Pin()->AddCustomTextFilter(FilterData, bApplyFilter);
+	}
+}
+
+void FContentBrowserSingleton::DeleteCustomTextFilterByLabel(const FText& FilterLabel)
+{
+	if (!HasPrimaryContentBrowser())
+	{
+		ChooseNewPrimaryBrowser();
+	}
+
+	if (PrimaryContentBrowser.IsValid())
+	{
+		PrimaryContentBrowser.Pin()->DeleteCustomTextFilterByLabel(FilterLabel);
+	}
+}
+
+void FContentBrowserSingleton::ModifyCustomTextFilterByLabel(const FCustomTextFilterData& NewFilterData, const FText& FilterLabel)
+{
+	if (!HasPrimaryContentBrowser())
+	{
+		ChooseNewPrimaryBrowser();
+	}
+
+	if (PrimaryContentBrowser.IsValid())
+	{
+		PrimaryContentBrowser.Pin()->ModifyCustomTextFilterByLabel(NewFilterData, FilterLabel);
+	}
+}
+
+bool FContentBrowserSingleton::IsAssetViewDoneFiltering()
+{
+	bool isDoneFiltering = false;
+
+	if (!HasPrimaryContentBrowser())
+	{
+		ChooseNewPrimaryBrowser();
+	}
+
+	if (PrimaryContentBrowser.IsValid())
+	{
+		isDoneFiltering = PrimaryContentBrowser.Pin()->IsAssetViewDoneFiltering();
+	}
+
+	return isDoneFiltering;
+}
+
 #undef LOCTEXT_NAMESPACE
diff --git a/Engine/Source/Editor/ContentBrowser/Private/ContentBrowserSingleton.h b/Engine/Source/Editor/ContentBrowser/Private/ContentBrowserSingleton.h
--- a/Engine/Source/Editor/ContentBrowser/Private/ContentBrowserSingleton.h
+++ b/Engine/Source/Editor/ContentBrowser/Private/ContentBrowserSingleton.h
@@ -118,6 +118,10 @@ PRAGMA_ENABLE_DEPRECATION_WARNINGS
 	virtual FDelegateHandle RegisterOnFavoritesChangedHandler(FSimpleDelegate OnFavoritesChanged) override;
 	virtual void UnregisterOnFavoritesChangedDelegate(FDelegateHandle Handle) override;
 	virtual TArray<FString> GetAliasesForPath(const FSoftObjectPath& InPath) const override;
+	virtual void AddCustomTextFilter(const FCustomTextFilterData& FilterData, bool bApplyFilter) override;
+	virtual void DeleteCustomTextFilterByLabel(const FText& FilterLabel) override;
+	virtual void ModifyCustomTextFilterByLabel(const FCustomTextFilterData& NewFilterData, const FText& FilterLabel) override;
+	virtual bool IsAssetViewDoneFiltering() override;
 
 	/** Broadcast that the favorites have changed. */
 	void BroadcastFavoritesChanged() const;
diff --git a/Engine/Source/Editor/ContentBrowser/Private/SContentBrowser.cpp b/Engine/Source/Editor/ContentBrowser/Private/SContentBrowser.cpp
--- a/Engine/Source/Editor/ContentBrowser/Private/SContentBrowser.cpp
+++ b/Engine/Source/Editor/ContentBrowser/Private/SContentBrowser.cpp
@@ -3349,6 +3349,42 @@ bool SContentBrowser::CanWriteToPath(const FContentBrowserItemPath InPath) const
 	return true;
 }
 
+void SContentBrowser::AddCustomTextFilter(const FCustomTextFilterData& FilterData, bool bApplyFilter)
+{
+	if (FilterListPtr.IsValid())
+	{
+		FilterListPtr->OnCreateCustomTextFilter(FilterData, bApplyFilter);
+	}
+}
+
+void SContentBrowser::DeleteCustomTextFilterByLabel(const FText& FilterLabel)
+{
+	if (FilterListPtr.IsValid())
+	{
+		FilterListPtr->DeleteCustomTextFilterByLabel(FilterLabel);
+	}
+}
+
+void SContentBrowser::ModifyCustomTextFilterByLabel(const FCustomTextFilterData& NewFilterData, const FText& FilterLabel)
+{
+	if (FilterListPtr.IsValid())
+	{
+		FilterListPtr->ModifyCustomTextFilterByLabel(NewFilterData, FilterLabel);
+	}
+}
+
+bool SContentBrowser::IsAssetViewDoneFiltering()
+{
+	bool isDoneFiltering = false;
+
+	if (AssetViewPtr.IsValid())
+	{
+		isDoneFiltering = !AssetViewPtr->HasItemsPendingFilter();
+	}
+
+	return isDoneFiltering;
+}
+
 bool SContentBrowser::IsAddNewEnabled() const
 {
 	return CanWriteToCurrentPath();
diff --git a/Engine/Source/Editor/ContentBrowser/Private/SContentBrowser.h b/Engine/Source/Editor/ContentBrowser/Private/SContentBrowser.h
--- a/Engine/Source/Editor/ContentBrowser/Private/SContentBrowser.h
+++ b/Engine/Source/Editor/ContentBrowser/Private/SContentBrowser.h
@@ -9,6 +9,7 @@
 #include "ContentBrowserDataMenuContexts.h"
 #include "ContentBrowserDataSubsystem.h"
 #include "ContentBrowserDelegates.h"
+#include "Filters/CustomTextFilters.h"
 #include "HAL/Platform.h"
 #include "HistoryManager.h"
 #include "IAssetTypeActions.h"
@@ -228,6 +229,18 @@ public:
 	/** Returns true if path can be written to */
 	bool CanWriteToPath(const FContentBrowserItemPath InPath) const;
 
+	/** Adds a new Custom Text Filter to the Filter List */
+	void AddCustomTextFilter(const FCustomTextFilterData& FilterData, bool bApplyFilter);
+
+	/** Deletes the specified Custom Text Filter */
+	void DeleteCustomTextFilterByLabel(const FText& FilterLabel);
+
+	/** Modifies the specified Custom Text Filter */
+	void ModifyCustomTextFilterByLabel(const FCustomTextFilterData& NewFilterData, const FText& FilterLabel);
+
+	/** Gets whether the AssetView in the Content Browser is done filtering */
+	bool IsAssetViewDoneFiltering();
+
 private:
 
 	/** Called prior to syncing the selection in this Content Browser */
diff --git a/Engine/Source/Editor/ContentBrowser/Private/SFilterList.cpp b/Engine/Source/Editor/ContentBrowser/Private/SFilterList.cpp
--- a/Engine/Source/Editor/ContentBrowser/Private/SFilterList.cpp
+++ b/Engine/Source/Editor/ContentBrowser/Private/SFilterList.cpp
@@ -412,6 +412,21 @@ void SFilterList::OnModifyCustomTextFilter(const FCustomTextFilterData& InFilter
 	}
 }
 
+void SFilterList::ModifyCustomTextFilterByLabel(const FCustomTextFilterData& NewFilterData, const FText& FilterLabel)
+{
+	// Find the filter with the matching label
+	TSharedRef<ICustomTextFilter<FAssetFilterType>>* MatchingFilter = CustomTextFilters
+		.FindByPredicate([&FilterLabel](const TSharedRef<ICustomTextFilter<FAssetFilterType>>& FilterElement)
+			{
+				return FilterElement->CreateCustomTextFilterData().FilterLabel.EqualTo(FilterLabel);
+			});
+
+	if (MatchingFilter)
+	{
+		OnModifyCustomTextFilter(NewFilterData, MatchingFilter->ToSharedPtr());
+	}
+}
+
 void SFilterList::OnDeleteCustomTextFilter(const TSharedPtr<ICustomTextFilter<FAssetFilterType>> InFilter)
 {
 	SAssetFilterBar<FAssetFilterType>::OnDeleteCustomTextFilter(InFilter);
@@ -425,6 +440,21 @@ void SFilterList::OnDeleteCustomTextFilter(const TSharedPtr<ICustomTextFilter<FA
 	}
 }
 
+void SFilterList::DeleteCustomTextFilterByLabel(const FText& FilterLabel)
+{
+	// Find the filter with the matching label
+	TSharedRef<ICustomTextFilter<FAssetFilterType>>* MatchingFilter = CustomTextFilters
+		.FindByPredicate([&FilterLabel](const TSharedRef<ICustomTextFilter<FAssetFilterType>>& FilterElement)
+			{
+				return FilterElement->CreateCustomTextFilterData().FilterLabel.EqualTo(FilterLabel);
+			});
+
+	if (MatchingFilter)
+	{
+		OnDeleteCustomTextFilter(MatchingFilter->ToSharedPtr());
+	}
+}
+
 bool SFilterList::RestoreCustomTextFilterState(const FCustomTextFilterState& InFilterState)
 {
 	// Find the filter associated with the current instance data from our list of custom text filters
diff --git a/Engine/Source/Editor/ContentBrowser/Private/SFilterList.h b/Engine/Source/Editor/ContentBrowser/Private/SFilterList.h
--- a/Engine/Source/Editor/ContentBrowser/Private/SFilterList.h
+++ b/Engine/Source/Editor/ContentBrowser/Private/SFilterList.h
@@ -159,15 +159,23 @@ public:
 
 	virtual void SetFilterLayout(EFilterBarLayout InFilterBarLayout) override;
 
+	/** Handler for when a custom text filter is created */
+	virtual void OnCreateCustomTextFilter(const FCustomTextFilterData& InFilterData, bool bApplyFilter) override;
+
+	/** Deletes the custom text filter with the matching label */
+	void DeleteCustomTextFilterByLabel(const FText& FilterLabel);
+
+	/** Modifies the custom text filter with the matching label */
+	void ModifyCustomTextFilterByLabel(const FCustomTextFilterData& NewFilterData, const FText& FilterLabel);
+
 protected:
 	
 	/** Handler for when the add filter button was clicked */
 	//virtual TSharedRef<SWidget> MakeAddFilterMenu() override;
 	
-	/** Handler for when a custom text filter is created */
-	virtual void OnCreateCustomTextFilter(const FCustomTextFilterData& InFilterData, bool bApplyFilter) override;
 	/** Handler for when a custom text filter is modified */
 	virtual void OnModifyCustomTextFilter(const FCustomTextFilterData& InFilterData, TSharedPtr<ICustomTextFilter<FAssetFilterType>> InFilter) override;
+
 	/** Handler for when a custom text filter is deleted */
 	virtual void OnDeleteCustomTextFilter(const TSharedPtr<ICustomTextFilter<FAssetFilterType>> InFilter) override;
 
diff --git a/Engine/Source/Editor/ContentBrowser/Public/IContentBrowserSingleton.h b/Engine/Source/Editor/ContentBrowser/Public/IContentBrowserSingleton.h
--- a/Engine/Source/Editor/ContentBrowser/Public/IContentBrowserSingleton.h
+++ b/Engine/Source/Editor/ContentBrowser/Public/IContentBrowserSingleton.h
@@ -12,6 +12,7 @@
 #include "ContentBrowserItem.h"
 #include "ContentBrowserDelegates.h"
 #include "ContentBrowserDataSubsystem.h"
+#include "Filters/CustomTextFilters.h"
 #include "Misc/FilterCollection.h"
 #include "Framework/Views/ITypedTableView.h"
 #include "AssetThumbnail.h"
@@ -769,4 +770,16 @@ public:
 	 * @return All alternative paths that represent the input path (not including the input path itself)
 	 */
 	virtual TArray<FString> GetAliasesForPath(const FSoftObjectPath& InPath) const = 0;
+
+	/** Adds a Custom Text Filter to the Filter List */
+	virtual void AddCustomTextFilter(const FCustomTextFilterData& FilterData, bool bApplyFilter) = 0;
+
+	/** Deletes the specified Custom Text Filter */
+	virtual void DeleteCustomTextFilterByLabel(const FText& FilterLabel) = 0;
+
+	/** Modifies the specified Custom Text Filter */
+	virtual void ModifyCustomTextFilterByLabel(const FCustomTextFilterData& FilterData, const FText& FilterLabel) = 0;
+
+	/** Gets whether the AssetView in the Content Browser is done filtering */
+	virtual bool IsAssetViewDoneFiltering() = 0;
 };
