Fixing some duplicated Localization keys

diff --git a/Engine/Source/Editor/ContentBrowser/Private/ContentBrowserMenuUtils.cpp b/Engine/Source/Editor/ContentBrowser/Private/ContentBrowserMenuUtils.cpp
--- a/Engine/Source/Editor/ContentBrowser/Private/ContentBrowserMenuUtils.cpp
+++ b/Engine/Source/Editor/ContentBrowser/Private/ContentBrowserMenuUtils.cpp
@@ -206,11 +206,11 @@ void ContentBrowserMenuUtils::AddFiltersToMenu(UToolMenu* InMenu, const FName& I
 	FToolMenuSection& Section = InMenu->FindOrAddSection("Show");
 
 	// New style omits the "Show" prefix
-	const FText ShowCppLabel = LOCTEXT("ShowCppClassesOption", "C++ Classes");
-	const FText ShowDevelopersContentLabel = LOCTEXT("ShowDevelopersContentOption", "Developers Content");
-	const FText ShowEngineContentLabel = LOCTEXT("ShowEngineFolderOption", "Engine Content");
-	const FText ShowPluginContentLabel = LOCTEXT("ShowPluginFolderOption", "Plugin Content");
-	const FText ShowLocalizedContentLabel = LOCTEXT("ShowLocalizedContentOption", "Localized Content");
+	const FText ShowCppLabel = LOCTEXT("ShowCppClassesOption_NewStyle", "C++ Classes");
+	const FText ShowDevelopersContentLabel = LOCTEXT("ShowDevelopersContentOption_NewStyle", "Developers Content");
+	const FText ShowEngineContentLabel = LOCTEXT("ShowEngineFolderOption_NewStyle", "Engine Content");
+	const FText ShowPluginContentLabel = LOCTEXT("ShowPluginFolderOption_NewStyle", "Plugin Content");
+	const FText ShowLocalizedContentLabel = LOCTEXT("ShowLocalizedContentOption_NewStyle", "Localized Content");
 #else
 	FToolMenuInsert ContentMenuInsert = FToolMenuInsert();
 	if (InMenu->ContainsSection("View"))
