[ContentBrowser] Fix ensure being hit while creating the menu in the CB 2.0

diff --git a/Engine/Source/Editor/ContentBrowser/Private/SContentBrowser.cpp b/Engine/Source/Editor/ContentBrowser/Private/SContentBrowser.cpp
--- a/Engine/Source/Editor/ContentBrowser/Private/SContentBrowser.cpp
+++ b/Engine/Source/Editor/ContentBrowser/Private/SContentBrowser.cpp
@@ -532,6 +532,8 @@ void SContentBrowser::Construct( const FArguments& InArgs, const FName& InInstan
 	constexpr float ToolBarVerticalPadding = 4.0f; // ToolBar total height should be 36.0f
 	constexpr float PanelInsetPadding = 2.0f;
 
+	const TSharedRef<SWidget> AssetView = CreateAssetView(Config);
+
 	const TSharedRef<SWidget> PanelContentWidget =
 		SNew(SVerticalBox)
 
@@ -694,7 +696,7 @@ void SContentBrowser::Construct( const FArguments& InArgs, const FName& InInstan
 
 							+ SVerticalBox::Slot()
 							[
-								CreateAssetView(Config)
+								AssetView
 							]
 
 							// Navigation Bar
