[ContentBrowser] Switched to cvar (from define) for new style

diff --git a/Engine/Plugins/Editor/ContentBrowser/ContentBrowserAssetDataSource/Source/ContentBrowserAssetDataSource/Private/ContentBrowserAssetDataSource.cpp b/Engine/Plugins/Editor/ContentBrowser/ContentBrowserAssetDataSource/Source/ContentBrowserAssetDataSource/Private/ContentBrowserAssetDataSource.cpp
--- a/Engine/Plugins/Editor/ContentBrowser/ContentBrowserAssetDataSource/Source/ContentBrowserAssetDataSource/Private/ContentBrowserAssetDataSource.cpp
+++ b/Engine/Plugins/Editor/ContentBrowser/ContentBrowserAssetDataSource/Source/ContentBrowserAssetDataSource/Private/ContentBrowserAssetDataSource.cpp
@@ -34,6 +34,7 @@
 #include "ContentBrowserMenuContexts.h"
 #include "Misc/Paths.h"
 #include "Widgets/Input/SButton.h"
+#include "SActionButton.h"
 #include "Subsystems/ImportSubsystem.h"
 #include "Widgets/Images/SImage.h"
 #include "Tasks/Task.h"
@@ -42,14 +43,16 @@
 #include "ToolMenuSection.h"
 #include "Widgets/SBoxPanel.h"
 #include "Widgets/Text/STextBlock.h"
-#if UE_CONTENTBROWSER_NEW_STYLE
-#include "SActionButton.h"
-#endif
 
 #include UE_INLINE_GENERATED_CPP_BY_NAME(ContentBrowserAssetDataSource)
 
 #define LOCTEXT_NAMESPACE "ContentBrowserAssetDataSource"
 
+namespace UE::Editor::ContentBrowser
+{
+	extern CONTENTBROWSER_API bool GEnableContentBrowserNewStyle;
+}
+
 namespace AssetDataSource
 {
 	bool bAllowInternalParallelism = true;
@@ -3591,44 +3594,47 @@ void UContentBrowserAssetDataSource::PopulateContentBrowserToolBar(UToolMenu* In
 	const UContentBrowserToolbarMenuContext* ContextObject = InMenu->FindContext<UContentBrowserToolbarMenuContext>();
 	checkf(ContextObject, TEXT("Required context UContentBrowserToolbarMenuContext was missing!"));
 
-#if UE_CONTENTBROWSER_NEW_STYLE
-	TSharedRef<SWidget> ImportButton =
-		SNew(SActionButton)
-		.ToolTipText(LOCTEXT("ImportTooltip", "Import assets from files to the currently selected folder"))
-		.OnClicked_UObject(this, &UContentBrowserAssetDataSource::OnImportClicked, ContextObject)
-		.IsEnabled_UObject(this, &UContentBrowserAssetDataSource::IsImportEnabled, ContextObject)
-		.Icon(FAppStyle::Get().GetBrush("Icons.Import"))
-		.Text(LOCTEXT("Import", "Import"));
-#else
-	TSharedRef<SWidget> ImportButton =
-		SNew(SButton)
-		.ButtonStyle(FAppStyle::Get(), "SimpleButton")
-		.ToolTipText(LOCTEXT("ImportTooltip", "Import assets from files to the currently selected folder"))
-		.ContentPadding(2)
-		.OnClicked_UObject(this, &UContentBrowserAssetDataSource::OnImportClicked, ContextObject)
-		.IsEnabled_UObject(this, &UContentBrowserAssetDataSource::IsImportEnabled, ContextObject)
-		[
-			SNew(SHorizontalBox)
-			+ SHorizontalBox::Slot()
-			.AutoWidth()
-			.HAlign(HAlign_Center)
-			.VAlign(VAlign_Center)
-			[
-				SNew(SImage)
-				.Image(FAppStyle::Get().GetBrush("Icons.Import"))
-				.ColorAndOpacity(FSlateColor::UseForeground())
-			]
-			+ SHorizontalBox::Slot()
-			.Padding(FMargin(3, 0, 0, 0))
-			.VAlign(VAlign_Center)
-			.AutoWidth()
+	TSharedPtr<SWidget> ImportButton = nullptr;
+
+	if (UE::Editor::ContentBrowser::GEnableContentBrowserNewStyle)
+	{
+		SAssignNew(ImportButton, SActionButton)
+			.ToolTipText(LOCTEXT("ImportTooltip", "Import assets from files to the currently selected folder"))
+			.OnClicked_UObject(this, &UContentBrowserAssetDataSource::OnImportClicked, ContextObject)
+			.IsEnabled_UObject(this, &UContentBrowserAssetDataSource::IsImportEnabled, ContextObject)
+			.Icon(FAppStyle::Get().GetBrush("Icons.Import"))
+			.Text(LOCTEXT("Import", "Import"));
+	}
+	else
+	{
+		SAssignNew(ImportButton, SButton)
+			.ButtonStyle(FAppStyle::Get(), "SimpleButton")
+			.ToolTipText(LOCTEXT("ImportTooltip", "Import assets from files to the currently selected folder"))
+			.ContentPadding(2)
+			.OnClicked_UObject(this, &UContentBrowserAssetDataSource::OnImportClicked, ContextObject)
+			.IsEnabled_UObject(this, &UContentBrowserAssetDataSource::IsImportEnabled, ContextObject)
 			[
-				SNew(STextBlock)
-				.TextStyle(FAppStyle::Get(), "NormalText")
-				.Text(LOCTEXT("Import", "Import"))
-			]
-		];
-#endif
+				SNew(SHorizontalBox)
+				+ SHorizontalBox::Slot()
+				.AutoWidth()
+				.HAlign(HAlign_Center)
+				.VAlign(VAlign_Center)
+				[
+					SNew(SImage)
+					.Image(FAppStyle::Get().GetBrush("Icons.Import"))
+					.ColorAndOpacity(FSlateColor::UseForeground())
+				]
+				+ SHorizontalBox::Slot()
+				.Padding(FMargin(3, 0, 0, 0))
+				.VAlign(VAlign_Center)
+				.AutoWidth()
+				[
+					SNew(STextBlock)
+					.TextStyle(FAppStyle::Get(), "NormalText")
+					.Text(LOCTEXT("Import", "Import"))
+				]
+			];
+	}
 
 	FToolMenuSection& Section = InMenu->FindOrAddSection("New");
 
@@ -3637,7 +3643,7 @@ void UContentBrowserAssetDataSource::PopulateContentBrowserToolBar(UToolMenu* In
 	Section.AddEntry(
 		FToolMenuEntry::InitWidget(
 			"Import",
-			ImportButton,
+			ImportButton.ToSharedRef(),
 			FText::GetEmpty(),
 			true,
 			false
diff --git a/Engine/Source/Editor/AssetDefinition/Private/AssetStatusAssetDataInfoProvider.cpp b/Engine/Source/Editor/AssetDefinition/Private/AssetStatusAssetDataInfoProvider.cpp
--- a/Engine/Source/Editor/AssetDefinition/Private/AssetStatusAssetDataInfoProvider.cpp
+++ b/Engine/Source/Editor/AssetDefinition/Private/AssetStatusAssetDataInfoProvider.cpp
@@ -2,7 +2,6 @@
 
 #include "AssetStatusAssetDataInfoProvider.h"
 
-#if UE_CONTENTBROWSER_NEW_STYLE
 #include "AssetRegistry/AssetData.h"
 #include "Misc/PackageName.h"
 #include "UObject/UObjectGlobals.h"
@@ -19,4 +18,3 @@ FString FAssetStatusAssetDataInfoProvider::TryGetFilename() const
 	FPackageName::TryConvertLongPackageNameToFilename(AssetData.PackageName.ToString(), OutFileName, *PackageExtension);
 	return OutFileName;
 }
-#endif
diff --git a/Engine/Source/Editor/AssetDefinition/Public/AssetDefinition.h b/Engine/Source/Editor/AssetDefinition/Public/AssetDefinition.h
--- a/Engine/Source/Editor/AssetDefinition/Public/AssetDefinition.h
+++ b/Engine/Source/Editor/AssetDefinition/Public/AssetDefinition.h
@@ -605,13 +605,11 @@ public:
 		return false;
 	}
 
-#if UE_CONTENTBROWSER_NEW_STYLE
 	// Implementation will change in the AssetDefinitionDefault in a later update
 	virtual void GetAssetStatusInfo(const TSharedPtr<IAssetStatusInfoProvider>& InAssetStatusInfoProvider, TArray<FAssetDisplayInfo>& OutStatusInfo) const
 	{
 		return;
 	}
-#endif
 
 	/**Optionally adds custom Asset Picker buttons to the Property Editor for the selected AssetDefinition, using the provided extension data this function returns.*/
 	virtual void GetAssetActionButtonExtensions(const FAssetData& InAssetData, TArray<FAssetButtonActionExtension>& OutExtensions) const { };
diff --git a/Engine/Source/Editor/AssetDefinition/Public/AssetDefinitionAssetInfo.h b/Engine/Source/Editor/AssetDefinition/Public/AssetDefinitionAssetInfo.h
--- a/Engine/Source/Editor/AssetDefinition/Public/AssetDefinitionAssetInfo.h
+++ b/Engine/Source/Editor/AssetDefinition/Public/AssetDefinitionAssetInfo.h
@@ -2,7 +2,6 @@
 
 #pragma once
 
-#if UE_CONTENTBROWSER_NEW_STYLE
 #include "Internationalization/Text.h"
 #include "Layout/Visibility.h"
 #include "Misc/Attribute.h"
@@ -79,4 +78,3 @@ public:
 	TAttribute<EVisibility> IsVisible;
 	TAttribute<FAssetStatusPriority> Priority;
 };
-#endif
diff --git a/Engine/Source/Editor/AssetDefinition/Public/AssetStatusAssetDataInfoProvider.h b/Engine/Source/Editor/AssetDefinition/Public/AssetStatusAssetDataInfoProvider.h
--- a/Engine/Source/Editor/AssetDefinition/Public/AssetStatusAssetDataInfoProvider.h
+++ b/Engine/Source/Editor/AssetDefinition/Public/AssetStatusAssetDataInfoProvider.h
@@ -2,7 +2,6 @@
 
 #pragma once
 
-#if UE_CONTENTBROWSER_NEW_STYLE
 #include "IAssetStatusInfoProvider.h"
 #include "AssetRegistry/AssetData.h"
 
@@ -20,4 +19,3 @@ public:
 private:
 	FAssetData AssetData;
 };
-#endif
diff --git a/Engine/Source/Editor/AssetDefinition/Public/IAssetStatusInfoProvider.h b/Engine/Source/Editor/AssetDefinition/Public/IAssetStatusInfoProvider.h
--- a/Engine/Source/Editor/AssetDefinition/Public/IAssetStatusInfoProvider.h
+++ b/Engine/Source/Editor/AssetDefinition/Public/IAssetStatusInfoProvider.h
@@ -2,7 +2,6 @@
 
 #pragma once
 
-#if UE_CONTENTBROWSER_NEW_STYLE
 #include "Containers/UnrealString.h"
 
 class UPackage;
@@ -18,4 +17,3 @@ public:
 	/** Try to get the filename, return empty otherwise */
 	virtual FString TryGetFilename() const = 0;
 };
-#endif
diff --git a/Engine/Source/Editor/AssetDefinition/Public/IAssetSystemInfoProvider.h b/Engine/Source/Editor/AssetDefinition/Public/IAssetSystemInfoProvider.h
--- a/Engine/Source/Editor/AssetDefinition/Public/IAssetSystemInfoProvider.h
+++ b/Engine/Source/Editor/AssetDefinition/Public/IAssetSystemInfoProvider.h
@@ -2,7 +2,6 @@
 
 #pragma once
 
-#if UE_CONTENTBROWSER_NEW_STYLE
 #include "Containers/Array.h"
 
 struct FAssetDisplayInfo;
@@ -15,4 +14,3 @@ public:
 	/** Populate the OutAssetDisplayInfo with the information of the Asset like the Path/Size etc... */
 	virtual void PopulateAssetInfo(TArray<FAssetDisplayInfo>& OutAssetDisplayInfo) const = 0;
 };
-#endif
diff --git a/Engine/Source/Editor/ContentBrowser/ContentBrowser.Build.cs b/Engine/Source/Editor/ContentBrowser/ContentBrowser.Build.cs
--- a/Engine/Source/Editor/ContentBrowser/ContentBrowser.Build.cs
+++ b/Engine/Source/Editor/ContentBrowser/ContentBrowser.Build.cs
@@ -6,8 +6,6 @@ public class ContentBrowser : ModuleRules
 {
 	public ContentBrowser(ReadOnlyTargetRules Target) : base(Target)
 	{
-		PublicDefinitions.Add("UE_CONTENTBROWSER_NEW_STYLE=0");
-
 		UnsafeTypeCastWarningLevel = WarningLevel.Error;
 
 		PrivateIncludePathModuleNames.AddRange(
diff --git a/Engine/Source/Editor/ContentBrowser/Private/AssetSystemContentBrowserInfoProvider.cpp b/Engine/Source/Editor/ContentBrowser/Private/AssetSystemContentBrowserInfoProvider.cpp
--- a/Engine/Source/Editor/ContentBrowser/Private/AssetSystemContentBrowserInfoProvider.cpp
+++ b/Engine/Source/Editor/ContentBrowser/Private/AssetSystemContentBrowserInfoProvider.cpp
@@ -2,7 +2,6 @@
 
 #include "AssetSystemContentBrowserInfoProvider.h"
 
-#if UE_CONTENTBROWSER_NEW_STYLE
 #include "ActorFolder.h"
 #include "ActorFolderDesc.h"
 #include "AssetDefinition.h"
@@ -574,4 +573,3 @@ FText FAssetSystemContentBrowserInfoProvider::GetAssetUserDescription() const
 }
 
 #undef LOCTEXT_NAMESPACE
-#endif
diff --git a/Engine/Source/Editor/ContentBrowser/Private/AssetSystemContentBrowserInfoProvider.h b/Engine/Source/Editor/ContentBrowser/Private/AssetSystemContentBrowserInfoProvider.h
--- a/Engine/Source/Editor/ContentBrowser/Private/AssetSystemContentBrowserInfoProvider.h
+++ b/Engine/Source/Editor/ContentBrowser/Private/AssetSystemContentBrowserInfoProvider.h
@@ -2,7 +2,7 @@
 
 #pragma once
 
-#if UE_CONTENTBROWSER_NEW_STYLE
+#include "AssetDefinition.h"
 #include "Containers/Array.h"
 #include "Containers/UnrealString.h"
 #include "Delegates/IDelegateInstance.h"
@@ -71,4 +71,3 @@ private:
 	/** AssetViewItem of this provider */
 	TSharedPtr<FAssetViewItem> AssetItem;
 };
-#endif
diff --git a/Engine/Source/Editor/ContentBrowser/Private/AssetViewWidgets.cpp b/Engine/Source/Editor/ContentBrowser/Private/AssetViewWidgets.cpp
--- a/Engine/Source/Editor/ContentBrowser/Private/AssetViewWidgets.cpp
+++ b/Engine/Source/Editor/ContentBrowser/Private/AssetViewWidgets.cpp
@@ -90,11 +90,12 @@ class FDragDropEvent;
 
 #define LOCTEXT_NAMESPACE "ContentBrowser"
 
-#if UE_CONTENTBROWSER_NEW_STYLE
-static int32 GenericThumbnailSizes[(int32)EThumbnailSize::MAX] = { 80, 96, 112, 128, 136, 160 };
-#else
-static int32 GenericThumbnailSizes[(int32)EThumbnailSize::MAX] = { 24, 32, 64, 128, 160, 200 };
-#endif
+namespace UE::Editor::ContentBrowser
+{
+	extern CONTENTBROWSER_API bool GEnableContentBrowserNewStyle;
+}
+
+using namespace UE::Editor::ContentBrowser;
 
 ///////////////////////////////
 // FAssetViewModeUtils
@@ -249,21 +250,23 @@ TSharedRef<SWidget> FAssetViewItemHelper::CreateListTileItemContents(T* const In
 			]
 		];
 
-#if !UE_CONTENTBROWSER_NEW_STYLE
-		// Dirty state
-		ItemContentsOverlay->AddSlot()
-		.HAlign(HAlign_Left)
-		.VAlign(VAlign_Bottom)
-		[
-			SNew(SBox)
-			.MaxDesiredWidth(InTileOrListItem, &T::GetStateIconImageSize)
-			.MaxDesiredHeight(InTileOrListItem, &T::GetStateIconImageSize)
+		if (!GEnableContentBrowserNewStyle)
+		{
+			// Dirty state
+			ItemContentsOverlay->AddSlot()
+			.HAlign(HAlign_Left)
+			.VAlign(VAlign_Bottom)
 			[
-				SNew(SImage)
-				.Image(InTileOrListItem, &T::GetDirtyImage)
-			]
-		];
-#endif
+				SNew(SBox)
+				.MaxDesiredWidth(InTileOrListItem, &T::GetStateIconImageSize)
+				.MaxDesiredHeight(InTileOrListItem, &T::GetStateIconImageSize)
+				[
+					SNew(SImage)
+					.Image(InTileOrListItem, &T::GetDirtyImage)
+				]
+			];
+		}
+
 		// Tools for thumbnail edit mode
 		ItemContentsOverlay->AddSlot()
 		[
@@ -298,14 +301,14 @@ public:
 
 		SToolTip::Construct(
 			SToolTip::FArguments()
-#if UE_CONTENTBROWSER_NEW_STYLE
-			.TextMargin(FMargin(12.f, 8.f, 12.f, 2.f))
-			.BorderImage(FAppStyle::GetBrush("AssetThumbnail.Tooltip.Border"))
-#else
-			.TextMargin(1.0f)
-			.BorderImage(UE::ContentBrowser::Private::FContentBrowserStyle::Get().GetBrush("ContentBrowser.TileViewTooltip.ToolTipBorder"))
-#endif
-			);
+			.TextMargin(
+				GEnableContentBrowserNewStyle
+				? FMargin(12.f, 8.f, 12.f, 2.f)
+				: FMargin(1.0f))
+			.BorderImage(
+				GEnableContentBrowserNewStyle
+				? FAppStyle::GetBrush("AssetThumbnail.Tooltip.Border")
+				: UE::ContentBrowser::Private::FContentBrowserStyle::Get().GetBrush("ContentBrowser.TileViewTooltip.ToolTipBorder")));
 	}
 
 	// IToolTip interface
@@ -1137,112 +1140,116 @@ TSharedRef<SWidget> SAssetViewItem::CreateToolTipWidget() const
 					}
 				}
 			}
-#if UE_CONTENTBROWSER_NEW_STYLE
-			return SNew(SVerticalBox)
 
-				+ SVerticalBox::Slot()
-				.AutoHeight()
-				[
-					SNew(SVerticalBox)
+			if (GEnableContentBrowserNewStyle)
+			{
+				return SNew(SVerticalBox)
 
 					+ SVerticalBox::Slot()
 					.AutoHeight()
 					[
-						SNew(SHorizontalBox)
+						SNew(SVerticalBox)
 
-						+ SHorizontalBox::Slot()
-						.AutoWidth()
-						.VAlign(VAlign_Center)
+						+ SVerticalBox::Slot()
+						.AutoHeight()
 						[
-							SNew(STextBlock)
-							.Text(FolderName)
-							.ColorAndOpacity(FStyleColors::White)
-							.Font(FAppStyle::GetFontStyle("ContentBrowser.Tooltip.EntryFont"))
-						]
+							SNew(SHorizontalBox)
 
-						+ SHorizontalBox::Slot()
-						.AutoWidth()
-						.VAlign(VAlign_Center)
-						[
-							SNew(STextBlock)
-							.Font(FAppStyle::GetFontStyle("ContentBrowser.Tooltip.EntryFont"))
-							.ColorAndOpacity(FStyleColors::White)
-							.Text(LOCTEXT("FolderNameBracketedLabel", " (Folder)"))
+							+ SHorizontalBox::Slot()
+							.AutoWidth()
+							.VAlign(VAlign_Center)
+							[
+								SNew(STextBlock)
+								.Text(FolderName)
+								.ColorAndOpacity(FStyleColors::White)
+								.Font(FAppStyle::GetFontStyle("ContentBrowser.Tooltip.EntryFont"))
+							]
+
+							+ SHorizontalBox::Slot()
+							.AutoWidth()
+							.VAlign(VAlign_Center)
+							[
+								SNew(STextBlock)
+								.Font(FAppStyle::GetFontStyle("ContentBrowser.Tooltip.EntryFont"))
+								.ColorAndOpacity(FStyleColors::White)
+								.Text(LOCTEXT("FolderNameBracketedLabel", " (Folder)"))
+							]
 						]
 					]
-				]
-
-				+ SVerticalBox::Slot()
-				.Padding(FMargin(0.f,6.f))
-				.AutoHeight()
-				[
-					SNew(SSeparator)
-					.Orientation(Orient_Horizontal)
-					.Thickness(1.f)
-					.ColorAndOpacity(COLOR("#484848FF"))
-					.SeparatorImage(FAppStyle::Get().GetBrush("WhiteBrush"))
-				]
 
-				+ SVerticalBox::Slot()
-				.AutoHeight()
-				[
-					InfoBox
-				];
-#else
-			return SNew(SBorder)
-				.Padding(6)
-				.BorderImage( UE::ContentBrowser::Private::FContentBrowserStyle::Get().GetBrush("ContentBrowser.TileViewTooltip.NonContentBorder") )
-				[
-					SNew(SVerticalBox)
+					+ SVerticalBox::Slot()
+					.Padding(FMargin(0.f,6.f))
+					.AutoHeight()
+					[
+						SNew(SSeparator)
+						.Orientation(Orient_Horizontal)
+						.Thickness(1.f)
+						.ColorAndOpacity(COLOR("#484848FF"))
+						.SeparatorImage(FAppStyle::Get().GetBrush("WhiteBrush"))
+					]
 
-					+SVerticalBox::Slot()
+					+ SVerticalBox::Slot()
 					.AutoHeight()
-					.Padding(0, 0, 0, 4)
 					[
-						SNew(SBorder)
-						.Padding(6)
-						.BorderImage( UE::ContentBrowser::Private::FContentBrowserStyle::Get().GetBrush("ContentBrowser.TileViewTooltip.ContentBorder") )
-						[
-							SNew(SVerticalBox)
+						InfoBox
+					];
+			}
+			else
+			{
+				return SNew(SBorder)
+					.Padding(6)
+					.BorderImage( UE::ContentBrowser::Private::FContentBrowserStyle::Get().GetBrush("ContentBrowser.TileViewTooltip.NonContentBorder") )
+					[
+						SNew(SVerticalBox)
 
-							+SVerticalBox::Slot()
-							.AutoHeight()
+						+SVerticalBox::Slot()
+						.AutoHeight()
+						.Padding(0, 0, 0, 4)
+						[
+							SNew(SBorder)
+							.Padding(6)
+							.BorderImage( UE::ContentBrowser::Private::FContentBrowserStyle::Get().GetBrush("ContentBrowser.TileViewTooltip.ContentBorder") )
 							[
-								SNew(SHorizontalBox)
+								SNew(SVerticalBox)
 
-								+SHorizontalBox::Slot()
-								.AutoWidth()
-								.VAlign(VAlign_Center)
-								.Padding(0, 0, 4, 0)
+								+SVerticalBox::Slot()
+								.AutoHeight()
 								[
-									SNew(STextBlock)
-									.Text( FolderName )
-									.Font( UE::ContentBrowser::Private::FContentBrowserStyle::Get().GetFontStyle("ContentBrowser.TileViewTooltip.NameFont") )
-								]
+									SNew(SHorizontalBox)
 
-								+SHorizontalBox::Slot()
-								.AutoWidth()
-								.VAlign(VAlign_Center)
-								[
-									SNew(STextBlock) 
-									.Text( LOCTEXT("FolderNameBracketed", "(Folder)") )
+									+SHorizontalBox::Slot()
+									.AutoWidth()
+									.VAlign(VAlign_Center)
+									.Padding(0, 0, 4, 0)
+									[
+										SNew(STextBlock)
+										.Text( FolderName )
+										.Font( UE::ContentBrowser::Private::FContentBrowserStyle::Get().GetFontStyle("ContentBrowser.TileViewTooltip.NameFont") )
+									]
+
+									+SHorizontalBox::Slot()
+									.AutoWidth()
+									.VAlign(VAlign_Center)
+									[
+										SNew(STextBlock)
+										.Text( LOCTEXT("FolderNameBracketed", "(Folder)") )
+									]
 								]
 							]
 						]
-					]
 
-					+SVerticalBox::Slot()
-					.AutoHeight()
-					[
-						SNew(SBorder)
-						.Padding(6)
-						.BorderImage( UE::ContentBrowser::Private::FContentBrowserStyle::Get().GetBrush("ContentBrowser.TileViewTooltip.ContentBorder") )
+						+SVerticalBox::Slot()
+						.AutoHeight()
 						[
-							InfoBox
+							SNew(SBorder)
+							.Padding(6)
+							.BorderImage( UE::ContentBrowser::Private::FContentBrowserStyle::Get().GetBrush("ContentBrowser.TileViewTooltip.ContentBorder") )
+							[
+								InfoBox
+							]
 						]
-					]
-				];
-#endif
+					];
+			}
 		}
 	}
 	else
@@ -1298,64 +1305,66 @@ FText SAssetViewItem::GetAssetUserDescription() const
 
 	return FText::GetEmpty();
 }
-#if UE_CONTENTBROWSER_NEW_STYLE
+
 void SAssetViewItem::AddToToolTipInfoBox(const TSharedRef<SVerticalBox>& InfoBox, const FText& Key, const FText& Value, bool bImportant) const
 {
-	InfoBox->AddSlot()
-	.Padding(0.f, 0.f, 0.f, 6.f)
-	.AutoHeight()
-	[
-		SNew(SHorizontalBox)
-		+SHorizontalBox::Slot()
-		.AutoWidth()
-		.Padding(0, 0, 4, 0)
+	if (GEnableContentBrowserNewStyle)
+	{
+		InfoBox->AddSlot()
+		.Padding(0.f, 0.f, 0.f, 6.f)
+		.AutoHeight()
 		[
-			SNew(STextBlock)
-			.Font(FAppStyle::GetFontStyle("ContentBrowser.Tooltip.EntryFont"))
-			.Text(FText::Format(NSLOCTEXT("AssetThumbnailToolTip", "AssetViewTooltipFormat", "{0}:"), Key))
-		]
+			SNew(SHorizontalBox)
+			+SHorizontalBox::Slot()
+			.AutoWidth()
+			.Padding(0, 0, 4, 0)
+			[
+				SNew(STextBlock)
+				.Font(FAppStyle::GetFontStyle("ContentBrowser.Tooltip.EntryFont"))
+				.Text(FText::Format(NSLOCTEXT("AssetThumbnailToolTip", "AssetViewTooltipFormat", "{0}:"), Key))
+			]
 
-		+SHorizontalBox::Slot()
-		.AutoWidth()
+			+SHorizontalBox::Slot()
+			.AutoWidth()
+			[
+				SNew(STextBlock)
+				.Font(FAppStyle::GetFontStyle("ContentBrowser.Tooltip.EntryFont"))
+				.ColorAndOpacity(FStyleColors::White)
+				.Text(Value)
+			]
+		];
+	}
+	else
+	{
+		FWidgetStyle ImportantStyle;
+		ImportantStyle.SetForegroundColor(FLinearColor(1, 0.5, 0, 1));
+
+		InfoBox->AddSlot()
+		.AutoHeight()
+		.Padding(0, 1)
 		[
-			SNew(STextBlock)
-			.Font(FAppStyle::GetFontStyle("ContentBrowser.Tooltip.EntryFont"))
-			.ColorAndOpacity(FStyleColors::White)
-			.Text(Value)
-		]
-	];
-}
-#else
-void SAssetViewItem::AddToToolTipInfoBox(const TSharedRef<SVerticalBox>& InfoBox, const FText& Key, const FText& Value, bool bImportant) const
-{
-	FWidgetStyle ImportantStyle;
-	ImportantStyle.SetForegroundColor(FLinearColor(1, 0.5, 0, 1));
+			SNew(SHorizontalBox)
 
-	InfoBox->AddSlot()
-	.AutoHeight()
-	.Padding(0, 1)
-	[
-		SNew(SHorizontalBox)
+			+SHorizontalBox::Slot()
+			.AutoWidth()
+			.Padding(0, 0, 4, 0)
+			[
+				SNew(STextBlock) .Text( FText::Format(LOCTEXT("AssetViewTooltipFormat", "{0}:"), Key ) )
+				.ColorAndOpacity(bImportant ? ImportantStyle.GetSubduedForegroundColor() : FSlateColor::UseSubduedForeground())
+			]
 
-		+SHorizontalBox::Slot()
-		.AutoWidth()
-		.Padding(0, 0, 4, 0)
-		[
-			SNew(STextBlock) .Text( FText::Format(LOCTEXT("AssetViewTooltipFormat", "{0}:"), Key ) )
-			.ColorAndOpacity(bImportant ? ImportantStyle.GetSubduedForegroundColor() : FSlateColor::UseSubduedForeground())
-		]
+			+SHorizontalBox::Slot()
+			.AutoWidth()
+			[
+				SNew(STextBlock) .Text( Value )
+				.ColorAndOpacity(bImportant ? ImportantStyle.GetForegroundColor() : FSlateColor::UseForeground())
+				.HighlightText((Key.ToString() == TEXT("Path")) ? HighlightText : FText())
+				.WrapTextAt(700.0f)
+			]
+		];
+	}
 
-		+SHorizontalBox::Slot()
-		.AutoWidth()
-		[
-			SNew(STextBlock) .Text( Value )
-			.ColorAndOpacity(bImportant ? ImportantStyle.GetForegroundColor() : FSlateColor::UseForeground())
-			.HighlightText((Key.ToString() == TEXT("Path")) ? HighlightText : FText())
-			.WrapTextAt(700.0f)
-		]
-	];
 }
-#endif
 
 void SAssetViewItem::UpdateDirtyState()
 {
@@ -1653,11 +1662,14 @@ void SAssetViewItem::CacheDisplayTags()
 
 const FSlateBrush* SAssetViewItem::GetBorderImage() const
 {
-#if UE_CONTENTBROWSER_NEW_STYLE
-	return bDraggedOver ? FAppStyle::GetBrush("AssetThumbnail.AssetBorder") : FAppStyle::GetBrush("NoBorder");
-#else
-	return bDraggedOver ? FAppStyle::GetBrush("Menu.Background") : FAppStyle::GetBrush("NoBorder");
-#endif
+	if (GEnableContentBrowserNewStyle)
+	{
+		return bDraggedOver ? FAppStyle::GetBrush("AssetThumbnail.AssetBorder") : FAppStyle::GetBrush("NoBorder");
+	}
+	else
+	{
+		return bDraggedOver ? FAppStyle::GetBrush("Menu.Background") : FAppStyle::GetBrush("NoBorder");
+	}
 }
 
 bool SAssetViewItem::IsFolder() const
@@ -2039,12 +2051,14 @@ void SAssetTileItem::Construct( const FArguments& InArgs )
 		ThumbnailConfig.HintColorAndOpacity = InArgs._ThumbnailHintColorAndOpacity;
 		ThumbnailConfig.Padding= FMargin(2.0f);
 		ThumbnailConfig.GenericThumbnailSize = MakeAttributeSP(this, &SAssetTileItem::GetGenericThumbnailSize);
-#if UE_CONTENTBROWSER_NEW_STYLE
-		ThumbnailConfig.AssetSystemInfoProvider = MakeShared<FAssetSystemContentBrowserInfoProvider>(AssetItem);
-		ThumbnailConfig.bAllowAssetStatusThumbnailOverlay = true;
-		ThumbnailConfig.bShowAssetChip = true;
-		ThumbnailConfig.AssetChipBorderImageOverride = TAttribute<const FSlateBrush*>::CreateSP(this, &SAssetTileItem::GetAssetAreaOverlayBackgroundImage);
-#endif
+
+		if (GEnableContentBrowserNewStyle)
+		{
+			ThumbnailConfig.AssetSystemInfoProvider = MakeShared<FAssetSystemContentBrowserInfoProvider>(AssetItem);
+			ThumbnailConfig.bAllowAssetStatusThumbnailOverlay = true;
+			ThumbnailConfig.bShowAssetChip = true;
+			ThumbnailConfig.AssetChipBorderImageOverride = TAttribute<const FSlateBrush*>::CreateSP(this, &SAssetTileItem::GetAssetAreaOverlayBackgroundImage);
+		}
 
 		if(!AssetItem->GetItem().IsSupported())
 		{
@@ -2066,13 +2080,15 @@ void SAssetTileItem::Construct( const FArguments& InArgs )
 		}
 
 		Thumbnail = AssetThumbnail->MakeThumbnailWidget(ThumbnailConfig);
-#if UE_CONTENTBROWSER_NEW_STYLE
-		// Use the same tooltip as the Thumbnail
-		if (const TSharedPtr<IToolTip>& ThumbnailTooltip = Thumbnail->GetToolTip())
+
+		if (GEnableContentBrowserNewStyle)
 		{
-			SetToolTip(ThumbnailTooltip);
+			// Use the same tooltip as the Thumbnail
+			if (const TSharedPtr<IToolTip>& ThumbnailTooltip = Thumbnail->GetToolTip())
+			{
+				SetToolTip(ThumbnailTooltip);
+			}
 		}
-#endif
 	}
 	else
 	{
@@ -2112,187 +2128,187 @@ void SAssetTileItem::Construct( const FArguments& InArgs )
 			.ColorAndOpacity(this, &SAssetTileItem::GetNameAreaTextColor);
 	}
 
-#if UE_CONTENTBROWSER_NEW_STYLE
-
-	constexpr float BorderPadding = 1.f;
-	constexpr float ShadowLeftTopPadding = 3.f;
-	constexpr float ShadowRightBotPadding = 4.f;
-	constexpr float ThumbnailBorderPadding = 0.f;
-	constexpr float NameAreaBoxLeftRightBotPadding = 4.f;
-	constexpr float NameAreaBoxTopPadding = 6.f;
-	constexpr float ClassNameMaxHeight = 14.f;
+	if (GEnableContentBrowserNewStyle)
+	{
+		constexpr float BorderPadding = 1.f;
+		constexpr float ShadowLeftTopPadding = 3.f;
+		constexpr float ShadowRightBotPadding = 4.f;
+		constexpr float ThumbnailBorderPadding = 0.f;
+		constexpr float NameAreaBoxLeftRightBotPadding = 4.f;
+		constexpr float NameAreaBoxTopPadding = 6.f;
+		constexpr float ClassNameMaxHeight = 14.f;
 
-	ChildSlot
-	.Padding(FMargin(BorderPadding, BorderPadding, 0.f, 0.f))
-	[
-		// Drop shadow border
-		SNew(SBorder)
-		.Padding(FMargin(ShadowLeftTopPadding, ShadowLeftTopPadding, ShadowRightBotPadding, ShadowRightBotPadding))
-		.BorderImage(IsFolder() ? TAttribute<const FSlateBrush*>(this, &SAssetTileItem::GetFolderBackgroundShadowImage) : FAppStyle::Get().GetBrush(ItemShadowBorderName))
+		ChildSlot
+		.Padding(FMargin(BorderPadding, BorderPadding, 0.f, 0.f))
 		[
-			SNew(SOverlay)
-			.AddMetaData<FTagMetaData>(FTagMetaData(AssetItem->GetItem().GetVirtualPath()))
-			+SOverlay::Slot()
+			// Drop shadow border
+			SNew(SBorder)
+			.Padding(FMargin(ShadowLeftTopPadding, ShadowLeftTopPadding, ShadowRightBotPadding, ShadowRightBotPadding))
+			.BorderImage(IsFolder() ? TAttribute<const FSlateBrush*>(this, &SAssetTileItem::GetFolderBackgroundShadowImage) : FAppStyle::Get().GetBrush(ItemShadowBorderName))
 			[
-				SNew(SBorder)
-				.Padding(ThumbnailBorderPadding)
-				.BorderImage(IsFolder() ? TAttribute<const FSlateBrush*>(this, &SAssetTileItem::GetFolderBackgroundImage) : TAttribute<const FSlateBrush*>(this, &SAssetTileItem::GetNameAreaBackgroundImage))
+				SNew(SOverlay)
+				.AddMetaData<FTagMetaData>(FTagMetaData(AssetItem->GetItem().GetVirtualPath()))
+				+SOverlay::Slot()
 				[
-					SNew(SVerticalBox)
-					// Thumbnail
-					+SVerticalBox::Slot()
-					.AutoHeight()
+					SNew(SBorder)
+					.Padding(ThumbnailBorderPadding)
+					.BorderImage(IsFolder() ? TAttribute<const FSlateBrush*>(this, &SAssetTileItem::GetFolderBackgroundImage) : TAttribute<const FSlateBrush*>(this, &SAssetTileItem::GetNameAreaBackgroundImage))
 					[
-						// The remainder of the space is reserved for the name.
-						SNew(SBox)
-						.WidthOverride(this, &SAssetTileItem::GetThumbnailBoxSize)
-						.HeightOverride(this, &SAssetTileItem::GetThumbnailBoxSize)
+						SNew(SVerticalBox)
+						// Thumbnail
+						+SVerticalBox::Slot()
+						.AutoHeight()
 						[
-							ItemContents
+							// The remainder of the space is reserved for the name.
+							SNew(SBox)
+							.WidthOverride(this, &SAssetTileItem::GetThumbnailBoxSize)
+							.HeightOverride(this, &SAssetTileItem::GetThumbnailBoxSize)
+							[
+								ItemContents
+							]
 						]
-					]
 
-					+SVerticalBox::Slot()
-					.FillHeight(1.f)
-					[
-						SNew(SBox)
-						.Visibility(this, &SAssetTileItem::GetNameAreaVisibility)
-						.Padding(FMargin(NameAreaBoxLeftRightBotPadding, NameAreaBoxTopPadding, NameAreaBoxLeftRightBotPadding, NameAreaBoxLeftRightBotPadding))
-						.Visibility(this, &SAssetTileItem::GetNameAreaVisibility)
+						+SVerticalBox::Slot()
+						.FillHeight(1.f)
 						[
-							SNew(SVerticalBox)
-							+ SVerticalBox::Slot()
-							.HAlign(IsFolder() ? HAlign_Center : HAlign_Fill)
+							SNew(SBox)
+							.Visibility(this, &SAssetTileItem::GetNameAreaVisibility)
+							.Padding(FMargin(NameAreaBoxLeftRightBotPadding, NameAreaBoxTopPadding, NameAreaBoxLeftRightBotPadding, NameAreaBoxLeftRightBotPadding))
+							.Visibility(this, &SAssetTileItem::GetNameAreaVisibility)
 							[
-								SNew(SBox)
-								.VAlign(VAlign_Top)
-								.HeightOverride(this, &SAssetTileItem::GetNameAreaMaxDesiredHeight)
+								SNew(SVerticalBox)
+								+ SVerticalBox::Slot()
+								.HAlign(IsFolder() ? HAlign_Center : HAlign_Fill)
 								[
-									InlineRenameWidget.ToSharedRef()
+									SNew(SBox)
+									.VAlign(VAlign_Top)
+									.HeightOverride(this, &SAssetTileItem::GetNameAreaMaxDesiredHeight)
+									[
+										InlineRenameWidget.ToSharedRef()
+									]
 								]
-							]
-							+ SVerticalBox::Slot()
-							.AutoHeight()
-							[
-								SNew(SBox)
-								.HeightOverride(ClassNameMaxHeight)
-								.VAlign(VAlign_Bottom)
+								+ SVerticalBox::Slot()
+								.AutoHeight()
 								[
-									SAssignNew(ClassTextWidget, STextBlock)
-									.Font(FAppStyle::GetFontStyle("ContentBrowser.AssetTileViewClassNameFont"))
-									.Visibility(this, &SAssetTileItem::GetAssetClassLabelVisibility)
-									.TextStyle(UE::ContentBrowser::Private::FContentBrowserStyle::Get().Get(), "ContentBrowser.ClassFont")
-									.OverflowPolicy(ETextOverflowPolicy::Ellipsis)
-									.Text(this, &SAssetTileItem::GetAssetClassText)
-									.ColorAndOpacity(this, &SAssetTileItem::GetAssetClassLabelTextColor)
+									SNew(SBox)
+									.HeightOverride(ClassNameMaxHeight)
+									.VAlign(VAlign_Bottom)
+									[
+										SAssignNew(ClassTextWidget, STextBlock)
+										.Font(FAppStyle::GetFontStyle("ContentBrowser.AssetTileViewClassNameFont"))
+										.Visibility(this, &SAssetTileItem::GetAssetClassLabelVisibility)
+										.TextStyle(UE::ContentBrowser::Private::FContentBrowserStyle::Get().Get(), "ContentBrowser.ClassFont")
+										.OverflowPolicy(ETextOverflowPolicy::Ellipsis)
+										.Text(this, &SAssetTileItem::GetAssetClassText)
+										.ColorAndOpacity(this, &SAssetTileItem::GetAssetClassLabelTextColor)
+									]
 								]
 							]
 						]
 					]
 				]
 			]
-		]
-	];
-#else
-	constexpr float AssetViewWidgetsBorderPadding = 4.f;
-	constexpr float AssetViewWidgetsShadowPadding = 5.f;
-	constexpr float AssetViewWidgetsRenameWidgetPadding = 2.f;
-	constexpr float AssetViewWidgetsClassTextPadding = 2.f;
-	constexpr float AssetViewWidgetSourceControlSize = 16.f;
-
-	ChildSlot
-	.Padding(FMargin(0.0f, 0.0f, AssetViewWidgetsBorderPadding, AssetViewWidgetsBorderPadding))
-	[				
-		// Drop shadow border
-		SNew(SBorder)
-		.Padding(FMargin(0.0f, 0.0f, AssetViewWidgetsShadowPadding, AssetViewWidgetsShadowPadding))
-		.BorderImage(IsFolder() ? TAttribute<const FSlateBrush*>(this, &SAssetTileItem::GetFolderBackgroundShadowImage) : FAppStyle::Get().GetBrush(ItemShadowBorderName))
+		];
+	}
+	else
+	{
+		constexpr float AssetViewWidgetsBorderPadding = 4.f;
+		constexpr float AssetViewWidgetsShadowPadding = 5.f;
+		constexpr float AssetViewWidgetsRenameWidgetPadding = 2.f;
+
+		ChildSlot
+		.Padding(FMargin(0.0f, 0.0f, AssetViewWidgetsBorderPadding, AssetViewWidgetsBorderPadding))
 		[
-			SNew(SOverlay)
-			.AddMetaData<FTagMetaData>(FTagMetaData(AssetItem->GetItem().GetVirtualPath()))
-			+SOverlay::Slot()
+			// Drop shadow border
+			SNew(SBorder)
+			.Padding(FMargin(0.0f, 0.0f, AssetViewWidgetsShadowPadding, AssetViewWidgetsShadowPadding))
+			.BorderImage(IsFolder() ? TAttribute<const FSlateBrush*>(this, &SAssetTileItem::GetFolderBackgroundShadowImage) : FAppStyle::Get().GetBrush(ItemShadowBorderName))
 			[
-				SNew(SBorder)
-				.Padding(0)
-				.BorderImage(IsFolder() ? TAttribute<const FSlateBrush*>(this, &SAssetTileItem::GetFolderBackgroundImage) : UE::ContentBrowser::Private::FContentBrowserStyle::Get().Get().GetBrush("ContentBrowser.AssetTileItem.ThumbnailAreaBackground"))
+				SNew(SOverlay)
+				.AddMetaData<FTagMetaData>(FTagMetaData(AssetItem->GetItem().GetVirtualPath()))
+				+SOverlay::Slot()
 				[
-					SNew(SVerticalBox)
-					// Thumbnail
-					+SVerticalBox::Slot()
-					.AutoHeight()
+					SNew(SBorder)
+					.Padding(0)
+					.BorderImage(IsFolder() ? TAttribute<const FSlateBrush*>(this, &SAssetTileItem::GetFolderBackgroundImage) : UE::ContentBrowser::Private::FContentBrowserStyle::Get().Get().GetBrush("ContentBrowser.AssetTileItem.ThumbnailAreaBackground"))
 					[
-						// The remainder of the space is reserved for the name.
-						SNew(SBox)
-						.Padding(0)
-						.WidthOverride(this, &SAssetTileItem::GetThumbnailBoxSize)
-						.HeightOverride(this, &SAssetTileItem::GetThumbnailBoxSize)
+						SNew(SVerticalBox)
+						// Thumbnail
+						+SVerticalBox::Slot()
+						.AutoHeight()
 						[
-							ItemContents
+							// The remainder of the space is reserved for the name.
+							SNew(SBox)
+							.Padding(0)
+							.WidthOverride(this, &SAssetTileItem::GetThumbnailBoxSize)
+							.HeightOverride(this, &SAssetTileItem::GetThumbnailBoxSize)
+							[
+								ItemContents
+							]
 						]
-					]
 
-					+SVerticalBox::Slot()
-					[
-						SNew(SBorder)
-						.Padding(FMargin(2.0f, 3.0f))
-						.BorderImage(this, &SAssetTileItem::GetNameAreaBackgroundImage)
+						+SVerticalBox::Slot()
 						[
-							SNew(SVerticalBox)
-							+ SVerticalBox::Slot()
-							.Padding(AssetViewWidgetsRenameWidgetPadding, AssetViewWidgetsRenameWidgetPadding, 0.0f, 0.0f)
-							.VAlign(VAlign_Top)
-							.HAlign(IsFolder() ? HAlign_Center : HAlign_Left)
+							SNew(SBorder)
+							.Padding(FMargin(2.0f, 3.0f))
+							.BorderImage(this, &SAssetTileItem::GetNameAreaBackgroundImage)
 							[
-								SNew(SBox)
-								.MaxDesiredHeight(this, &SAssetTileItem::GetNameAreaMaxDesiredHeight)
+								SNew(SVerticalBox)
+								+ SVerticalBox::Slot()
+								.Padding(AssetViewWidgetsRenameWidgetPadding, AssetViewWidgetsRenameWidgetPadding, 0.0f, 0.0f)
+								.VAlign(VAlign_Top)
+								.HAlign(IsFolder() ? HAlign_Center : HAlign_Left)
 								[
-									InlineRenameWidget.ToSharedRef()
+									SNew(SBox)
+									.MaxDesiredHeight(this, &SAssetTileItem::GetNameAreaMaxDesiredHeight)
+									[
+										InlineRenameWidget.ToSharedRef()
+									]
 								]
-							]
-							+ SVerticalBox::Slot()
-							.VAlign(VAlign_Bottom)
-							.AutoHeight()
-							.Padding(0.0f, 0.0f, 0.0f, 0.0f)
-							[
-								SNew(SHorizontalBox)
-								+SHorizontalBox::Slot()
+								+ SVerticalBox::Slot()
+								.VAlign(VAlign_Bottom)
+								.AutoHeight()
+								.Padding(0.0f, 0.0f, 0.0f, 0.0f)
 								[
-									SAssignNew(ClassTextWidget, STextBlock)
-									.Visibility(this, &SAssetTileItem::GetAssetClassLabelVisibility)
-									.TextStyle(UE::ContentBrowser::Private::FContentBrowserStyle::Get().Get(), "ContentBrowser.ClassFont")
-									.OverflowPolicy(ETextOverflowPolicy::Ellipsis)
-									.Text(this, &SAssetTileItem::GetAssetClassText)
-									.ColorAndOpacity(this, &SAssetTileItem::GetAssetClassLabelTextColor)
-								]
+									SNew(SHorizontalBox)
+									+SHorizontalBox::Slot()
+									[
+										SAssignNew(ClassTextWidget, STextBlock)
+										.Visibility(this, &SAssetTileItem::GetAssetClassLabelVisibility)
+										.TextStyle(UE::ContentBrowser::Private::FContentBrowserStyle::Get().Get(), "ContentBrowser.ClassFont")
+										.OverflowPolicy(ETextOverflowPolicy::Ellipsis)
+										.Text(this, &SAssetTileItem::GetAssetClassText)
+										.ColorAndOpacity(this, &SAssetTileItem::GetAssetClassLabelTextColor)
+									]
 
-								+SHorizontalBox::Slot()
-								.Padding(FMargin(3.0f, 0.0f, 0.0f, 1.0f))
-								.AutoWidth()
-								.HAlign(HAlign_Right)
-								[
-									SNew(SBox)
-									.WidthOverride(this, &SAssetTileItem::GetSourceControlIconSize)
-									.HeightOverride(this, &SAssetTileItem::GetSourceControlIconSize)
-									.Visibility(this, &SAssetTileItem::GetSCCIconVisibility)
+									+SHorizontalBox::Slot()
+									.Padding(FMargin(3.0f, 0.0f, 0.0f, 1.0f))
+									.AutoWidth()
+									.HAlign(HAlign_Right)
 									[
-										GenerateSourceControlIconWidget()
+										SNew(SBox)
+										.WidthOverride(this, &SAssetTileItem::GetSourceControlIconSize)
+										.HeightOverride(this, &SAssetTileItem::GetSourceControlIconSize)
+										.Visibility(this, &SAssetTileItem::GetSCCIconVisibility)
+										[
+											GenerateSourceControlIconWidget()
+										]
 									]
 								]
+
 							]
-							
 						]
 					]
 				]
+				+SOverlay::Slot()
+				[
+					SNew(SImage)
+					.Image(this, &SAssetViewItem::GetBorderImage)
+					.Visibility(EVisibility::HitTestInvisible)
+				]
 			]
-			+SOverlay::Slot()
-			[
-				SNew(SImage)
-				.Image(this, &SAssetViewItem::GetBorderImage)
-				.Visibility(EVisibility::HitTestInvisible)	
-			]
-		]
-	];
-#endif
+		];
+	}
 
 	HandleSourceControlStateChanged();
 
@@ -2383,11 +2399,7 @@ FOptionalSize SAssetTileItem::GetStateIconImageSize() const
 
 FOptionalSize SAssetTileItem::GetThumbnailBoxSize() const
 {
-#if UE_CONTENTBROWSER_NEW_STYLE
-	return FOptionalSize(ThumbnailDimension.Get());
-#else
-	return FOptionalSize(ItemWidth.Get()- ThumbnailPadding);
-#endif
+	return GEnableContentBrowserNewStyle ? FOptionalSize(ThumbnailDimension.Get()) : FOptionalSize(ItemWidth.Get()- ThumbnailPadding);
 }
 
 EVisibility SAssetTileItem::GetAssetClassLabelVisibility() const
@@ -2406,38 +2418,37 @@ FSlateColor SAssetTileItem::GetAssetClassLabelTextColor() const
 		return FStyleColors::White;
 	}
 
-#if UE_CONTENTBROWSER_NEW_STYLE
-	return FStyleColors::Hover2;
-#else
-	return FSlateColor::UseSubduedForeground();
-#endif
+	return GEnableContentBrowserNewStyle ? FStyleColors::Hover2 : FSlateColor::UseSubduedForeground();
 }
 
 FSlateFontInfo SAssetTileItem::GetThumbnailFont() const
 {
-#if UE_CONTENTBROWSER_NEW_STYLE
-	const static FLazyName RegularFont("ContentBrowser.AssetTileViewNameFont");
-	return FAppStyle::GetFontStyle(RegularFont);
-#else
-	FOptionalSize ThumbSize = GetThumbnailBoxSize();
-	if ( ThumbSize.IsSet() )
+	if (GEnableContentBrowserNewStyle)
 	{
-		float Size = ThumbSize.Get();
-		if ( Size < 50 )
-		{
-			const static FName SmallFontName("ContentBrowser.AssetTileViewNameFontVerySmall");
-			return FAppStyle::GetFontStyle(SmallFontName);
-		}
-		else if ( Size < 85 )
+		const static FLazyName RegularFont("ContentBrowser.AssetTileViewNameFont");
+		return FAppStyle::GetFontStyle(RegularFont);
+	}
+	else
+	{
+		FOptionalSize ThumbSize = GetThumbnailBoxSize();
+		if ( ThumbSize.IsSet() )
 		{
-			const static FName SmallFontName("ContentBrowser.AssetTileViewNameFontSmall");
-			return FAppStyle::GetFontStyle(SmallFontName);
+			float Size = ThumbSize.Get();
+			if ( Size < 50 )
+			{
+				const static FName SmallFontName("ContentBrowser.AssetTileViewNameFontVerySmall");
+				return FAppStyle::GetFontStyle(SmallFontName);
+			}
+			else if ( Size < 85 )
+			{
+				const static FName SmallFontName("ContentBrowser.AssetTileViewNameFontSmall");
+				return FAppStyle::GetFontStyle(SmallFontName);
+			}
 		}
-	}
 
-	const static FName RegularFont("ContentBrowser.AssetTileViewNameFont");
-	return FAppStyle::GetFontStyle(RegularFont);
-#endif
+		const static FName RegularFont("ContentBrowser.AssetTileViewNameFont");
+		return FAppStyle::GetFontStyle(RegularFont);
+	}
 }
 
 const FSlateBrush* SAssetTileItem::GetFolderBackgroundImage() const
@@ -2485,17 +2496,10 @@ const FSlateBrush* SAssetTileItem::GetNameAreaBackgroundImage() const
 		return FStyleDefaults::GetNoBrush();
 	}
 
-#if UE_CONTENTBROWSER_NEW_STYLE
-	static const FLazyName SelectedHover("ContentBrowser.AssetTileItem.AssetContentSelectedHoverBackground");
-	static const FLazyName Selected("ContentBrowser.AssetTileItem.AssetContentSelectedBackground");
-	static const FLazyName Hovered("ContentBrowser.AssetTileItem.AssetContentHoverBackground");
-	static const FLazyName Normal("ContentBrowser.AssetTileItem.AssetContent");
-#else
-	static const FLazyName SelectedHover("ContentBrowser.AssetTileItem.NameAreaSelectedHoverBackground");
-	static const FLazyName Selected("ContentBrowser.AssetTileItem.NameAreaSelectedBackground");
-	static const FLazyName Hovered("ContentBrowser.AssetTileItem.NameAreaHoverBackground");
-	static const FLazyName Normal("ContentBrowser.AssetTileItem.NameAreaBackground");
-#endif
+	const FName SelectedHover = GEnableContentBrowserNewStyle ? "ContentBrowser.AssetTileItem.AssetContentSelectedHoverBackground" : "ContentBrowser.AssetTileItem.NameAreaSelectedHoverBackground";
+	const FName Selected = GEnableContentBrowserNewStyle ? "ContentBrowser.AssetTileItem.AssetContentSelectedBackground" : "ContentBrowser.AssetTileItem.NameAreaSelectedBackground";
+	const FName Hovered = GEnableContentBrowserNewStyle ? "ContentBrowser.AssetTileItem.AssetContentHoverBackground" : "ContentBrowser.AssetTileItem.NameAreaHoverBackground";
+	const FName Normal = GEnableContentBrowserNewStyle ? "ContentBrowser.AssetTileItem.AssetContent" : "ContentBrowser.AssetTileItem.NameAreaBackground";
 
 	const bool bIsSelected = IsSelected.IsBound() ? IsSelected.Execute() : false;
 	const bool bIsHoveredOrDraggedOver = IsHovered() || bDraggedOver;
@@ -2578,36 +2582,48 @@ FSlateColor SAssetTileItem::GetNameAreaTextColor() const
 		return FStyleColors::White;
 	}
 
-#if UE_CONTENTBROWSER_NEW_STYLE
-	return FStyleColors::Foreground;
-#else
-	return FSlateColor::UseForeground();
-#endif
+	return GEnableContentBrowserNewStyle ? FStyleColors::Foreground : FSlateColor::UseForeground();
 }
 
 FOptionalSize SAssetTileItem::GetNameAreaMaxDesiredHeight() const
 {
-#if UE_CONTENTBROWSER_NEW_STYLE
-	constexpr int32 MaxHeightNameArea = 42;
-	return MaxHeightNameArea;
-#else
-	return AssetNameHeights[(int32)CurrentThumbnailSize.Get()];
-#endif
+	if (GEnableContentBrowserNewStyle)
+	{
+		constexpr int32 MaxHeightNameArea = 42;
+		return MaxHeightNameArea;
+	}
+	else
+	{
+		return AssetNameHeights[(int32)CurrentThumbnailSize.Get()];
+	}
 }
 
 int32 SAssetTileItem::GetGenericThumbnailSize() const
 {
-	return GenericThumbnailSizes[(int32)CurrentThumbnailSize.Get()];
+	const int32 ThumbnailSizeIdx = (int32)CurrentThumbnailSize.Get();
+	if (GEnableContentBrowserNewStyle)
+	{
+		static int32 GenericThumbnailSizes[(int32)EThumbnailSize::MAX] = { 80, 96, 112, 128, 136, 160 };
+		return GenericThumbnailSizes[ThumbnailSizeIdx];
+	}
+	else
+	{
+		static int32 GenericThumbnailSizes[(int32)EThumbnailSize::MAX] = { 24, 32, 64, 128, 160, 200 };
+		return GenericThumbnailSizes[ThumbnailSizeIdx];
+	}
 }
 
 EVisibility SAssetTileItem::GetSCCIconVisibility() const
 {
 	// Hide the scc state icon when there is no brush or in tiny size since there isn't enough space
-#if UE_CONTENTBROWSER_NEW_STYLE
-	return bHasCCStateBrush && ISourceControlModule::Get().IsEnabled() && ISourceControlModule::Get().GetProvider().IsAvailable() ? EVisibility::Visible : EVisibility::Collapsed;
-#else
-	return bHasCCStateBrush &&  CurrentThumbnailSize.Get() != EThumbnailSize::Tiny && ISourceControlModule::Get().IsEnabled() && ISourceControlModule::Get().GetProvider().IsAvailable() ? EVisibility::Visible : EVisibility::Collapsed;
-#endif
+	if (GEnableContentBrowserNewStyle)
+	{
+		return bHasCCStateBrush && ISourceControlModule::Get().IsEnabled() && ISourceControlModule::Get().GetProvider().IsAvailable() ? EVisibility::Visible : EVisibility::Collapsed;
+	}
+	else
+	{
+		return bHasCCStateBrush &&  CurrentThumbnailSize.Get() != EThumbnailSize::Tiny && ISourceControlModule::Get().IsEnabled() && ISourceControlModule::Get().GetProvider().IsAvailable() ? EVisibility::Visible : EVisibility::Collapsed;
+	}
 }
 
 EVisibility SAssetTileItem::GetNameAreaVisibility() const
@@ -2653,12 +2669,16 @@ void SAssetTileItem::InitializeAssetNameHeights()
 
 			AssetNameHeights[(int32)EThumbnailSize::Medium] = RegularFontHeight * MediumSizeMultiplier;
 			AssetNameHeights[(int32)EThumbnailSize::Large] = RegularFontHeight * LargeSizeMultiplier;
-#if UE_CONTENTBROWSER_NEW_STYLE
-			AssetNameHeights[(int32)EThumbnailSize::XLarge] = RegularFontHeight * XLargeSizeMultiplier;
-			AssetNameHeights[(int32)EThumbnailSize::Huge] = RegularFontHeight * HugeSizeMultiplier;
-#else
-			AssetNameHeights[(int32)EThumbnailSize::Huge] = RegularFontHeight * XLargeSizeMultiplier;
-#endif
+
+			if (GEnableContentBrowserNewStyle)
+			{
+				AssetNameHeights[(int32)EThumbnailSize::XLarge] = RegularFontHeight * XLargeSizeMultiplier;
+				AssetNameHeights[(int32)EThumbnailSize::Huge] = RegularFontHeight * HugeSizeMultiplier;
+			}
+			else
+			{
+				AssetNameHeights[(int32)EThumbnailSize::Huge] = RegularFontHeight * XLargeSizeMultiplier;
+			}
 		}
 
 		bInitializedHeights = true;
diff --git a/Engine/Source/Editor/ContentBrowser/Private/ContentBrowserMenuUtils.cpp b/Engine/Source/Editor/ContentBrowser/Private/ContentBrowserMenuUtils.cpp
--- a/Engine/Source/Editor/ContentBrowser/Private/ContentBrowserMenuUtils.cpp
+++ b/Engine/Source/Editor/ContentBrowser/Private/ContentBrowserMenuUtils.cpp
@@ -7,6 +7,13 @@
 
 #define LOCTEXT_NAMESPACE "ContentBrowserMenuUtils"
 
+namespace UE::Editor::ContentBrowser
+{
+	extern CONTENTBROWSER_API bool GEnableContentBrowserNewStyle;
+}
+
+using namespace UE::Editor::ContentBrowser;
+
 namespace ContentBrowserMenuUtils::Filters
 {
 	FContentBrowserInstanceConfig* GetContentBrowserConfig(FName InOwningContentBrowserName)
@@ -202,16 +209,6 @@ namespace ContentBrowserMenuUtils::Filters
 
 void ContentBrowserMenuUtils::AddFiltersToMenu(UToolMenu* InMenu, const FName& InOwningContentBrowserName, FFiltersAdditionalParams InFiltersAdditionalParams)
 {
-#if UE_CONTENTBROWSER_NEW_STYLE
-	FToolMenuSection& Section = InMenu->FindOrAddSection("Show");
-
-	// New style omits the "Show" prefix
-	const FText ShowCppLabel = LOCTEXT("ShowCppClassesOption_NewStyle", "C++ Classes");
-	const FText ShowDevelopersContentLabel = LOCTEXT("ShowDevelopersContentOption_NewStyle", "Developers Content");
-	const FText ShowEngineContentLabel = LOCTEXT("ShowEngineFolderOption_NewStyle", "Engine Content");
-	const FText ShowPluginContentLabel = LOCTEXT("ShowPluginFolderOption_NewStyle", "Plugin Content");
-	const FText ShowLocalizedContentLabel = LOCTEXT("ShowLocalizedContentOption_NewStyle", "Localized Content");
-#else
 	FToolMenuInsert ContentMenuInsert = FToolMenuInsert();
 	if (InMenu->ContainsSection("View"))
 	{
@@ -219,16 +216,16 @@ void ContentBrowserMenuUtils::AddFiltersToMenu(UToolMenu* InMenu, const FName& I
 		ContentMenuInsert.Position = EToolMenuInsertType::After;
 	}
 
-	FToolMenuSection& Section = InMenu->FindOrAddSection("Content", LOCTEXT("ContentHeading", "Content"), ContentMenuInsert);
+	FToolMenuSection& Section = GEnableContentBrowserNewStyle ? InMenu->FindOrAddSection("Show") : InMenu->FindOrAddSection("Content", LOCTEXT("ContentHeading", "Content"), ContentMenuInsert);
 
-	const FText ShowCppLabel = LOCTEXT("ShowCppClassesOption", "Show C++ Classes");
-	const FText ShowDevelopersContentLabel = LOCTEXT("ShowDevelopersContentOption", "Show Developers Content");
-	const FText ShowEngineContentLabel = LOCTEXT("ShowEngineFolderOption", "Show Engine Content");
-	const FText ShowPluginContentLabel = LOCTEXT("ShowPluginFolderOption", "Show Plugin Content");
-	const FText ShowLocalizedContentLabel = LOCTEXT("ShowLocalizedContentOption", "Show Localized Content");
-#endif
+	// New style omits the "Show" prefix
+	const FText ShowCppLabel = GEnableContentBrowserNewStyle ?  LOCTEXT("ShowCppClassesOption_NewStyle", "C++ Classes") : LOCTEXT("ShowCppClassesOption", "Show C++ Classes");
+	const FText ShowDevelopersContentLabel = GEnableContentBrowserNewStyle ? LOCTEXT("ShowDevelopersContentOption_NewStyle", "Developers Content") : LOCTEXT("ShowDevelopersContentOption", "Show Developers Content");
+	const FText ShowEngineContentLabel = GEnableContentBrowserNewStyle ? LOCTEXT("ShowEngineFolderOption_NewStyle", "Engine Content") : LOCTEXT("ShowEngineFolderOption", "Show Engine Content");
+	const FText ShowPluginContentLabel = GEnableContentBrowserNewStyle ? LOCTEXT("ShowPluginFolderOption_NewStyle", "Plugin Content") : LOCTEXT("ShowPluginFolderOption", "Show Plugin Content");
+	const FText ShowLocalizedContentLabel = GEnableContentBrowserNewStyle ? LOCTEXT("ShowLocalizedContentOption_NewStyle", "Localized Content") : LOCTEXT("ShowLocalizedContentOption", "Show Localized Content");
 
-		Section.AddMenuEntry(
+	Section.AddMenuEntry(
 		"ShowCppClasses",
 		ShowCppLabel,
 		LOCTEXT("ShowCppClassesOptionToolTip", "Show C++ classes in the view?"),
diff --git a/Engine/Source/Editor/ContentBrowser/Private/ContentBrowserModule.cpp b/Engine/Source/Editor/ContentBrowser/Private/ContentBrowserModule.cpp
--- a/Engine/Source/Editor/ContentBrowser/Private/ContentBrowserModule.cpp
+++ b/Engine/Source/Editor/ContentBrowser/Private/ContentBrowserModule.cpp
@@ -1,6 +1,5 @@
 // Copyright Epic Games, Inc. All Rights Reserved.
 
-
 #include "ContentBrowserModule.h"
 
 #include "ContentBrowserDataSubsystem.h"
@@ -20,10 +19,64 @@
 
 IMPLEMENT_MODULE( FContentBrowserModule, ContentBrowser );
 DEFINE_LOG_CATEGORY(LogContentBrowser);
+
 const FName FContentBrowserModule::NumberOfRecentAssetsName(TEXT("NumObjectsInRecentList"));
 
+// @todo: Can remove after permanent switch to new style
+namespace UE::Editor::ContentBrowser
+{
+	bool GEnableContentBrowserNewStyle = false;
+
+	namespace Private
+	{
+		static FAutoConsoleVariableRef CVarEnableContentBrowserNewStyle(
+			TEXT("ContentBrowser.EnableNewStyle"),
+			GEnableContentBrowserNewStyle,
+			TEXT("Whether or not to enable the Content Browser restyle (~5.6).")
+		);
+
+		void SetEnableNewStyleFromCmdLine()
+		{
+			static constexpr const TCHAR* CVarName = TEXT("ContentBrowser.EnableNewStyle");
+
+			const FString OriginalCmdLine = FCommandLine::GetOriginal();
+			if (const int32 FoundIdx = OriginalCmdLine.Find(CVarName);
+				FoundIdx != INDEX_NONE)
+			{
+				// If present, enable
+				GEnableContentBrowserNewStyle = true;
+
+				// Check for explicit disable =0 or \s0
+				for (int32 CharIdx = FoundIdx + FCString::Strlen(CVarName); CharIdx < OriginalCmdLine.Len(); ++CharIdx)
+				{
+					if (OriginalCmdLine[CharIdx] == TEXT(' ')
+						|| OriginalCmdLine[CharIdx] == TEXT('='))
+					{
+						continue;
+					}
+
+					if (FChar::IsDigit(OriginalCmdLine[CharIdx]);
+						const int32 Value = FChar::ConvertCharDigitToInt(OriginalCmdLine[CharIdx]))
+					{
+						// Explicitly disable
+						if (Value == 0)
+						{
+							GEnableContentBrowserNewStyle = false;
+							break;
+						}
+					}
+
+					break;
+				}
+			}
+		}
+	}
+}
+
 void FContentBrowserModule::StartupModule()
 {
+	UE::Editor::ContentBrowser::Private::SetEnableNewStyleFromCmdLine();
+
 	// Ensure the data module is loaded
 	IContentBrowserDataModule::Get();
 
diff --git a/Engine/Source/Editor/ContentBrowser/Private/ContentBrowserStyle.cpp b/Engine/Source/Editor/ContentBrowser/Private/ContentBrowserStyle.cpp
--- a/Engine/Source/Editor/ContentBrowser/Private/ContentBrowserStyle.cpp
+++ b/Engine/Source/Editor/ContentBrowser/Private/ContentBrowserStyle.cpp
@@ -9,6 +9,11 @@
 #include "Styling/StyleColors.h"
 #include "Styling/ToolBarStyle.h"
 
+namespace UE::Editor::ContentBrowser
+{
+	extern CONTENTBROWSER_API bool GEnableContentBrowserNewStyle;
+}
+
 namespace UE::ContentBrowser::Private
 {
 	FName FContentBrowserStyle::StyleName("ContentBrowser");
@@ -312,11 +317,14 @@ namespace UE::ContentBrowser::Private
 
 		ContentBrowserToolBarStyle.SetSeparatorBrush(FSlateNoResource());
 		ContentBrowserToolBarStyle.SetSeparatorPadding(FMargin(4.0f, 0.0f));
-#if UE_CONTENTBROWSER_NEW_STYLE
-		ContentBrowserToolBarStyle.SetBackgroundPadding(FMargin(4.0f, 2.0f, 4.0f, 2.0f));
-#else
-		ContentBrowserToolBarStyle.SetBackgroundPadding(FMargin(4.0f, 2.0f, 0.0f, 2.0f));
-#endif
+		if (UE::Editor::ContentBrowser::GEnableContentBrowserNewStyle)
+		{
+			ContentBrowserToolBarStyle.SetBackgroundPadding(FMargin(4.0f, 2.0f, 4.0f, 2.0f));
+		}
+		else
+		{
+			ContentBrowserToolBarStyle.SetBackgroundPadding(FMargin(4.0f, 2.0f, 0.0f, 2.0f));
+		}
 
 		Set("ContentBrowser.ToolBar", ContentBrowserToolBarStyle);
 
diff --git a/Engine/Source/Editor/ContentBrowser/Private/Menus/AssetViewMenus.cpp b/Engine/Source/Editor/ContentBrowser/Private/Menus/AssetViewMenus.cpp
--- a/Engine/Source/Editor/ContentBrowser/Private/Menus/AssetViewMenus.cpp
+++ b/Engine/Source/Editor/ContentBrowser/Private/Menus/AssetViewMenus.cpp
@@ -12,6 +12,13 @@
 
 #define LOCTEXT_NAMESPACE "ContentBrowser"
 
+namespace UE::Editor::ContentBrowser
+{
+	extern CONTENTBROWSER_API bool GEnableContentBrowserNewStyle;
+}
+
+using namespace UE::Editor::ContentBrowser;
+
 FDelayedAutoRegisterHelper SAssetView::AssetViewOptionsMenuRegistration(
 	EDelayedRegisterRunPhase::EndOfEngineInit,
 	[]
@@ -102,225 +109,224 @@ FDelayedAutoRegisterHelper SAssetView::AssetViewOptionsMenuRegistration(
 						);
 					}
 
-#if !UE_CONTENTBROWSER_NEW_STYLE
-					// Filter Bar Section
-					if (TSharedPtr<SFilterList> FilterBarPinned = AssetView->FilterBar.Pin())
+					if (!GEnableContentBrowserNewStyle)
 					{
-						FToolMenuSection& Section = InMenu->AddSection("FilterBar", LOCTEXT("FilterBarHeading", "Filter Display"));
+						// Filter Bar Section
+						if (TSharedPtr<SFilterList> FilterBarPinned = AssetView->FilterBar.Pin())
+						{
+							FToolMenuSection& Section = InMenu->AddSection("FilterBar", LOCTEXT("FilterBarHeading", "Filter Display"));
 
-						Section.AddMenuEntry(
-							"VerticalLayout",
-							LOCTEXT("FilterListVerticalLayout", "Vertical"),
-							LOCTEXT("FilterListVerticalLayoutToolTip", "Swap to a vertical layout for the filter bar"),
-							FSlateIcon(),
-							FUIAction(
-								FExecuteAction::CreateLambda([FilterBarPinned]()
-								{
-									if (FilterBarPinned->GetFilterLayout() != EFilterBarLayout::Vertical)
+							Section.AddMenuEntry(
+								"VerticalLayout",
+								LOCTEXT("FilterListVerticalLayout", "Vertical"),
+								LOCTEXT("FilterListVerticalLayoutToolTip", "Swap to a vertical layout for the filter bar"),
+								FSlateIcon(),
+								FUIAction(
+									FExecuteAction::CreateLambda([FilterBarPinned]()
 									{
-										FilterBarPinned->SetFilterLayout(EFilterBarLayout::Vertical);
-									}
-								}),
-								FCanExecuteAction(),
-								FIsActionChecked::CreateLambda([FilterBarPinned]()
-								{
-									return FilterBarPinned->GetFilterLayout() == EFilterBarLayout::Vertical;
-								})),
-							EUserInterfaceActionType::RadioButton
-						);
+										if (FilterBarPinned->GetFilterLayout() != EFilterBarLayout::Vertical)
+										{
+											FilterBarPinned->SetFilterLayout(EFilterBarLayout::Vertical);
+										}
+									}),
+									FCanExecuteAction(),
+									FIsActionChecked::CreateLambda([FilterBarPinned]()
+									{
+										return FilterBarPinned->GetFilterLayout() == EFilterBarLayout::Vertical;
+									})),
+								EUserInterfaceActionType::RadioButton
+							);
 
-						Section.AddMenuEntry(
-							"HorizontalLayout",
-							LOCTEXT("FilterListHorizontalLayout", "Horizontal"),
-							LOCTEXT("FilterListHorizontalLayoutToolTip", "Swap to a Horizontal layout for the filter bar"),
-							FSlateIcon(),
-							FUIAction(
-								FExecuteAction::CreateLambda([FilterBarPinned]()
-								{
-									if (FilterBarPinned->GetFilterLayout() != EFilterBarLayout::Horizontal)
+							Section.AddMenuEntry(
+								"HorizontalLayout",
+								LOCTEXT("FilterListHorizontalLayout", "Horizontal"),
+								LOCTEXT("FilterListHorizontalLayoutToolTip", "Swap to a Horizontal layout for the filter bar"),
+								FSlateIcon(),
+								FUIAction(
+									FExecuteAction::CreateLambda([FilterBarPinned]()
 									{
-										FilterBarPinned->SetFilterLayout(EFilterBarLayout::Horizontal);
-									}
-								}),
-								FCanExecuteAction(),
-								FIsActionChecked::CreateLambda([FilterBarPinned]()
-								{
-									return FilterBarPinned->GetFilterLayout() == EFilterBarLayout::Horizontal;
-								})),
-							EUserInterfaceActionType::RadioButton
-						);
+										if (FilterBarPinned->GetFilterLayout() != EFilterBarLayout::Horizontal)
+										{
+											FilterBarPinned->SetFilterLayout(EFilterBarLayout::Horizontal);
+										}
+									}),
+									FCanExecuteAction(),
+									FIsActionChecked::CreateLambda([FilterBarPinned]()
+									{
+										return FilterBarPinned->GetFilterLayout() == EFilterBarLayout::Horizontal;
+									})),
+								EUserInterfaceActionType::RadioButton
+							);
+						}
 					}
-#endif
 
-				// Thumbnails Section
-				{
-#if UE_CONTENTBROWSER_NEW_STYLE
-					FToolMenuInsert SectionInsert("AssetViewType", EToolMenuInsertType::After);
-#else
-					FToolMenuInsert SectionInsert("Search", EToolMenuInsertType::After);
-#endif
-					FToolMenuSection& Section = InMenu->FindOrAddSection("AssetThumbnails", LOCTEXT("ThumbnailsHeading", "Thumbnails"), SectionInsert);
-
-					auto CreateThumbnailSizeSubMenu = [AssetView](UToolMenu* SubMenu)
+					// Thumbnails Section
 					{
-						FToolMenuSection& SizeSection = SubMenu->AddSection("ThumbnailSizes");
+						const FToolMenuInsert SectionInsert(GEnableContentBrowserNewStyle ? "AssetViewType" : "Search", EToolMenuInsertType::After);
+						FToolMenuSection& Section = InMenu->FindOrAddSection("AssetThumbnails", LOCTEXT("ThumbnailsHeading", "Thumbnails"), SectionInsert);
 
-						for (int32 EnumValue = static_cast<int32>(EThumbnailSize::Tiny); EnumValue < static_cast<int32>(EThumbnailSize::MAX); ++EnumValue)
+						auto CreateThumbnailSizeSubMenu = [AssetView](UToolMenu* SubMenu)
 						{
-#if !UE_CONTENTBROWSER_NEW_STYLE
-							if ((EThumbnailSize)EnumValue == EThumbnailSize::XLarge)
+							FToolMenuSection& SizeSection = SubMenu->AddSection("ThumbnailSizes");
+
+							for (int32 EnumValue = static_cast<int32>(EThumbnailSize::Tiny); EnumValue < static_cast<int32>(EThumbnailSize::MAX); ++EnumValue)
 							{
-								continue;
+								if (!GEnableContentBrowserNewStyle)
+								{
+									if ((EThumbnailSize)EnumValue == EThumbnailSize::XLarge)
+									{
+										continue;
+									}
+								}
+
+								SizeSection.AddMenuEntry(
+									NAME_None,
+									SAssetView::ThumbnailSizeToDisplayName((EThumbnailSize)EnumValue),
+									FText::GetEmpty(),
+									FSlateIcon(),
+									FUIAction(
+										FExecuteAction::CreateSP(AssetView, &SAssetView::OnThumbnailSizeChanged, (EThumbnailSize)EnumValue),
+										FCanExecuteAction::CreateSP(AssetView, &SAssetView::IsThumbnailScalingAllowed),
+										FIsActionChecked::CreateSP(AssetView, &SAssetView::IsThumbnailSizeChecked, (EThumbnailSize)EnumValue)
+									),
+									EUserInterfaceActionType::RadioButton
+								);
 							}
-#endif
+						};
 
-							SizeSection.AddMenuEntry(
-								NAME_None,
-								SAssetView::ThumbnailSizeToDisplayName((EThumbnailSize)EnumValue),
-								FText::GetEmpty(),
+						Section.AddEntry(FToolMenuEntry::InitSubMenu(
+							"ThumbnailSize",
+							LOCTEXT("ThumbnailSize", "Thumbnail Size"),
+							LOCTEXT("ThumbnailSizeToolTip", "Adjust the size of thumbnails."),
+							FNewToolMenuDelegate::CreateLambda(CreateThumbnailSizeSubMenu)
+						));
+
+						if (!GEnableContentBrowserNewStyle)
+						{
+							Section.AddMenuEntry(
+								"ThumbnailEditMode",
+								LOCTEXT("ThumbnailEditModeOption", "Thumbnail Edit Mode"),
+								LOCTEXT("ThumbnailEditModeOptionToolTip", "Toggle thumbnail editing mode. When in this mode you can rotate the camera on 3D thumbnails by dragging them."),
 								FSlateIcon(),
 								FUIAction(
-									FExecuteAction::CreateSP(AssetView, &SAssetView::OnThumbnailSizeChanged, (EThumbnailSize)EnumValue),
-									FCanExecuteAction::CreateSP(AssetView, &SAssetView::IsThumbnailScalingAllowed),
-									FIsActionChecked::CreateSP(AssetView, &SAssetView::IsThumbnailSizeChecked, (EThumbnailSize)EnumValue)
+									FExecuteAction::CreateSP(AssetView, &SAssetView::ToggleThumbnailEditMode),
+									FCanExecuteAction::CreateSP(AssetView, &SAssetView::IsThumbnailEditModeAllowed),
+									FIsActionChecked::CreateSP(AssetView, &SAssetView::IsThumbnailEditMode)
 								),
-								EUserInterfaceActionType::RadioButton
+								EUserInterfaceActionType::ToggleButton
+							);
+
+							Section.AddMenuEntry(
+								"RealTimeThumbnails",
+								LOCTEXT("RealTimeThumbnailsOption", "Real-Time Thumbnails"),
+								LOCTEXT("RealTimeThumbnailsOptionToolTip", "Renders the assets thumbnails in real-time"),
+								FSlateIcon(),
+								FUIAction(
+									FExecuteAction::CreateSP(AssetView, &SAssetView::ToggleRealTimeThumbnails),
+									FCanExecuteAction::CreateSP(AssetView, &SAssetView::CanShowRealTimeThumbnails),
+									FIsActionChecked::CreateSP(AssetView, &SAssetView::IsShowingRealTimeThumbnails)
+								),
+								EUserInterfaceActionType::ToggleButton
 							);
 						}
-					};
+					}
 
-					Section.AddEntry(FToolMenuEntry::InitSubMenu(
-						"ThumbnailSize",
-						LOCTEXT("ThumbnailSize", "Thumbnail Size"),
-						LOCTEXT("ThumbnailSizeToolTip", "Adjust the size of thumbnails."),
-						FNewToolMenuDelegate::CreateLambda(CreateThumbnailSizeSubMenu)
-					));
+					// Show/View Section
+					{
+						FToolMenuSection& Section = InMenu->FindOrAddSection(
+							GEnableContentBrowserNewStyle ? "Show" : "View",
+							GEnableContentBrowserNewStyle ? LOCTEXT("ShowHeading", "Show") : LOCTEXT("ViewHeading", "View"));
 
-#if !UE_CONTENTBROWSER_NEW_STYLE
 						Section.AddMenuEntry(
-							"ThumbnailEditMode",
-							LOCTEXT("ThumbnailEditModeOption", "Thumbnail Edit Mode"),
-							LOCTEXT("ThumbnailEditModeOptionToolTip", "Toggle thumbnail editing mode. When in this mode you can rotate the camera on 3D thumbnails by dragging them."),
+							"ShowFolders",
+							LOCTEXT("ShowFoldersOption", "Folders"),
+							LOCTEXT("ShowFoldersOptionToolTip", "Show folders in the view as well as assets?"),
 							FSlateIcon(),
 							FUIAction(
-								FExecuteAction::CreateSP(AssetView, &SAssetView::ToggleThumbnailEditMode),
-								FCanExecuteAction::CreateSP(AssetView, &SAssetView::IsThumbnailEditModeAllowed),
-								FIsActionChecked::CreateSP(AssetView, &SAssetView::IsThumbnailEditMode)
+								FExecuteAction::CreateSP(AssetView, &SAssetView::ToggleShowFolders),
+								FCanExecuteAction::CreateSP(AssetView, &SAssetView::IsToggleShowFoldersAllowed),
+								FIsActionChecked::CreateSP(AssetView, &SAssetView::IsShowingFolders)
 							),
 							EUserInterfaceActionType::ToggleButton
 						);
 
 						Section.AddMenuEntry(
-							"RealTimeThumbnails",
-							LOCTEXT("RealTimeThumbnailsOption", "Real-Time Thumbnails"),
-							LOCTEXT("RealTimeThumbnailsOptionToolTip", "Renders the assets thumbnails in real-time"),
+							"ShowEmptyFolders",
+							LOCTEXT("ShowEmptyFoldersOption", "Empty Folders"),
+							LOCTEXT("ShowEmptyFoldersOptionToolTip", "Show empty folders in the view as well as assets?"),
 							FSlateIcon(),
 							FUIAction(
-								FExecuteAction::CreateSP(AssetView, &SAssetView::ToggleRealTimeThumbnails),
-								FCanExecuteAction::CreateSP(AssetView, &SAssetView::CanShowRealTimeThumbnails),
-								FIsActionChecked::CreateSP(AssetView, &SAssetView::IsShowingRealTimeThumbnails)
+								FExecuteAction::CreateSP(AssetView, &SAssetView::ToggleShowEmptyFolders),
+								FCanExecuteAction::CreateSP(AssetView, &SAssetView::IsToggleShowEmptyFoldersAllowed),
+								FIsActionChecked::CreateSP(AssetView, &SAssetView::IsShowingEmptyFolders)
 							),
 							EUserInterfaceActionType::ToggleButton
 						);
-#endif
-				}
 
-				// Show/View Section
-				{
-#if UE_CONTENTBROWSER_NEW_STYLE
-					FToolMenuSection& Section = InMenu->FindOrAddSection("Show", LOCTEXT("ShowHeading", "Show"));
-#else
-					FToolMenuSection& Section = InMenu->AddSection("View", LOCTEXT("ViewHeading", "View"));
-#endif
-					Section.AddMenuEntry(
-						"ShowFolders",
-						LOCTEXT("ShowFoldersOption", "Folders"),
-						LOCTEXT("ShowFoldersOptionToolTip", "Show folders in the view as well as assets?"),
-						FSlateIcon(),
-						FUIAction(
-							FExecuteAction::CreateSP(AssetView, &SAssetView::ToggleShowFolders),
-							FCanExecuteAction::CreateSP(AssetView, &SAssetView::IsToggleShowFoldersAllowed),
-							FIsActionChecked::CreateSP(AssetView, &SAssetView::IsShowingFolders)
-						),
-						EUserInterfaceActionType::ToggleButton
-					);
-
-					Section.AddMenuEntry(
-						"ShowEmptyFolders",
-						LOCTEXT("ShowEmptyFoldersOption", "Empty Folders"),
-						LOCTEXT("ShowEmptyFoldersOptionToolTip", "Show empty folders in the view as well as assets?"),
-						FSlateIcon(),
-						FUIAction(
-							FExecuteAction::CreateSP(AssetView, &SAssetView::ToggleShowEmptyFolders),
-							FCanExecuteAction::CreateSP(AssetView, &SAssetView::IsToggleShowEmptyFoldersAllowed),
-							FIsActionChecked::CreateSP(AssetView, &SAssetView::IsShowingEmptyFolders)
-						),
-						EUserInterfaceActionType::ToggleButton
-					);
-
-					Section.AddMenuEntry(
-						"ShowFavorite",
-						LOCTEXT("ShowFavoriteOptions", "Show Favorites"),
-						LOCTEXT("ShowFavoriteOptionToolTip", "Show the favorite folders in the view?"),
-						FSlateIcon(),
-						FUIAction(
-							FExecuteAction::CreateSP(AssetView, &SAssetView::ToggleShowFavorites),
-							FCanExecuteAction::CreateSP(AssetView, &SAssetView::IsToggleShowFavoritesAllowed),
-							FIsActionChecked::CreateSP(AssetView, &SAssetView::IsShowingFavorites)
-						),
-						EUserInterfaceActionType::ToggleButton
-					);
-
-#if !UE_CONTENTBROWSER_NEW_STYLE
 						Section.AddMenuEntry(
-							"FilterRecursively",
-							LOCTEXT("FilterRecursivelyOption", "Filter Recursively"),
-							LOCTEXT("FilterRecursivelyOptionToolTip", "Should filters apply recursively in the view?"),
+							"ShowFavorite",
+							LOCTEXT("ShowFavoriteOptions", "Show Favorites"),
+							LOCTEXT("ShowFavoriteOptionToolTip", "Show the favorite folders in the view?"),
 							FSlateIcon(),
 							FUIAction(
-								FExecuteAction::CreateSP(AssetView, &SAssetView::ToggleFilteringRecursively),
-								FCanExecuteAction::CreateSP(AssetView, &SAssetView::IsToggleFilteringRecursivelyAllowed),
-								FIsActionChecked::CreateSP(AssetView, &SAssetView::IsFilteringRecursively)
+								FExecuteAction::CreateSP(AssetView, &SAssetView::ToggleShowFavorites),
+								FCanExecuteAction::CreateSP(AssetView, &SAssetView::IsToggleShowFavoritesAllowed),
+								FIsActionChecked::CreateSP(AssetView, &SAssetView::IsShowingFavorites)
 							),
 							EUserInterfaceActionType::ToggleButton
 						);
 
-						Section.AddMenuEntry(
-							"ShowAllFolder",
-							LOCTEXT("ShowAllFolderOption", "Show All Folder"),
-							LOCTEXT("ShowAllFolderOptionToolTip", "Show the all folder in the view?"),
-							FSlateIcon(),
-							FUIAction(
-								FExecuteAction::CreateSP(AssetView, &SAssetView::ToggleShowAllFolder),
-								FCanExecuteAction(),
-								FIsActionChecked::CreateSP(AssetView, &SAssetView::IsShowingAllFolder)
-							),
-							EUserInterfaceActionType::ToggleButton
-						);
+						if (!GEnableContentBrowserNewStyle)
+						{
+							Section.AddMenuEntry(
+								"FilterRecursively",
+								LOCTEXT("FilterRecursivelyOption", "Filter Recursively"),
+								LOCTEXT("FilterRecursivelyOptionToolTip", "Should filters apply recursively in the view?"),
+								FSlateIcon(),
+								FUIAction(
+									FExecuteAction::CreateSP(AssetView, &SAssetView::ToggleFilteringRecursively),
+									FCanExecuteAction::CreateSP(AssetView, &SAssetView::IsToggleFilteringRecursivelyAllowed),
+									FIsActionChecked::CreateSP(AssetView, &SAssetView::IsFilteringRecursively)
+								),
+								EUserInterfaceActionType::ToggleButton
+							);
 
-						Section.AddMenuEntry(
-							"OrganizeFolders",
-							LOCTEXT("OrganizeFoldersOption", "Organize Folders"),
-							LOCTEXT("OrganizeFoldersOptionToolTip", "Organize folders in the view?"),
-							FSlateIcon(),
-							FUIAction(
-								FExecuteAction::CreateSP(AssetView, &SAssetView::ToggleOrganizeFolders),
-								FCanExecuteAction(),
-								FIsActionChecked::CreateSP(AssetView, &SAssetView::IsOrganizingFolders)
-							),
-							EUserInterfaceActionType::ToggleButton
-						);
+							Section.AddMenuEntry(
+								"ShowAllFolder",
+								LOCTEXT("ShowAllFolderOption", "Show All Folder"),
+								LOCTEXT("ShowAllFolderOptionToolTip", "Show the all folder in the view?"),
+								FSlateIcon(),
+								FUIAction(
+									FExecuteAction::CreateSP(AssetView, &SAssetView::ToggleShowAllFolder),
+									FCanExecuteAction(),
+									FIsActionChecked::CreateSP(AssetView, &SAssetView::IsShowingAllFolder)
+								),
+								EUserInterfaceActionType::ToggleButton
+							);
 
-						if (AssetView->bShowPathViewFilters)
-						{
-							Section.AddSubMenu(
-								"PathViewFilters",
-								LOCTEXT("PathViewFilters", "Path View Filters"),
-								LOCTEXT("PathViewFilters_ToolTip", "Path View Filters"),
-								FNewToolMenuDelegate());
+							Section.AddMenuEntry(
+								"OrganizeFolders",
+								LOCTEXT("OrganizeFoldersOption", "Organize Folders"),
+								LOCTEXT("OrganizeFoldersOptionToolTip", "Organize folders in the view?"),
+								FSlateIcon(),
+								FUIAction(
+									FExecuteAction::CreateSP(AssetView, &SAssetView::ToggleOrganizeFolders),
+									FCanExecuteAction(),
+									FIsActionChecked::CreateSP(AssetView, &SAssetView::IsOrganizingFolders)
+								),
+								EUserInterfaceActionType::ToggleButton
+							);
+
+							if (AssetView->bShowPathViewFilters)
+							{
+								Section.AddSubMenu(
+									"PathViewFilters",
+									LOCTEXT("PathViewFilters", "Path View Filters"),
+									LOCTEXT("PathViewFilters_ToolTip", "Path View Filters"),
+									FNewToolMenuDelegate());
+							}
 						}
-#endif
-				}
+					}
 
 					// Search Section
 					{
@@ -365,53 +371,54 @@ FDelayedAutoRegisterHelper SAssetView::AssetViewOptionsMenuRegistration(
 						);
 					}
 
-
-#if UE_CONTENTBROWSER_NEW_STYLE
-					// Manage Section
-					{
-						FToolMenuSection& Section = InMenu->FindOrAddSection("Manage", LOCTEXT("ManageHeading", "Manage"));
-						Section.AddMenuEntry(
-							"OrganizeFolders",
-							LOCTEXT("OrganizeFoldersOption", "Organize Folders"),
-							LOCTEXT("OrganizeFoldersOptionToolTip", "Organize folders in the view?"),
-							FSlateIcon(),
-							FUIAction(
-								FExecuteAction::CreateSP(AssetView, &SAssetView::ToggleOrganizeFolders),
-								FCanExecuteAction(),
-								FIsActionChecked::CreateSP(AssetView, &SAssetView::IsOrganizingFolders)
-							),
-							EUserInterfaceActionType::ToggleButton
-						);
-					}
-#endif
-
-#if !UE_CONTENTBROWSER_NEW_STYLE
-					// Asset Columns Section
-					if (AssetView->GetColumnViewVisibility() == EVisibility::Visible)
+					if (GEnableContentBrowserNewStyle)
 					{
+						// Manage Section
 						{
-							FToolMenuSection& Section = InMenu->AddSection("AssetColumns", LOCTEXT("ToggleColumnsHeading", "Columns"));
-
+							FToolMenuSection& Section = InMenu->FindOrAddSection("Manage", LOCTEXT("ManageHeading", "Manage"));
 							Section.AddMenuEntry(
-								"ResetColumns",
-								LOCTEXT("ResetColumns", "Reset Columns"),
-								LOCTEXT("ResetColumnsToolTip", "Reset all columns to be visible again."),
+								"OrganizeFolders",
+								LOCTEXT("OrganizeFoldersOption", "Organize Folders"),
+								LOCTEXT("OrganizeFoldersOptionToolTip", "Organize folders in the view?"),
 								FSlateIcon(),
-								FUIAction(FExecuteAction::CreateSP(AssetView, &SAssetView::ResetColumns)),
-								EUserInterfaceActionType::Button
+								FUIAction(
+									FExecuteAction::CreateSP(AssetView, &SAssetView::ToggleOrganizeFolders),
+									FCanExecuteAction(),
+									FIsActionChecked::CreateSP(AssetView, &SAssetView::IsOrganizingFolders)
+								),
+								EUserInterfaceActionType::ToggleButton
 							);
+						}
+					}
 
-							Section.AddMenuEntry(
-								"ExportColumns",
-								LOCTEXT("ExportColumns", "Export to CSV"),
-								LOCTEXT("ExportColumnsToolTip", "Export column data to CSV."),
-								FSlateIcon(),
-								FUIAction(FExecuteAction::CreateSP(AssetView, &SAssetView::ExportColumns)),
-								EUserInterfaceActionType::Button
-							);
+					if (!GEnableContentBrowserNewStyle)
+					{
+						// Asset Columns Section
+						if (AssetView->GetColumnViewVisibility() == EVisibility::Visible)
+						{
+							{
+								FToolMenuSection& Section = InMenu->AddSection("AssetColumns", LOCTEXT("ToggleColumnsHeading", "Columns"));
+
+								Section.AddMenuEntry(
+									"ResetColumns",
+									LOCTEXT("ResetColumns", "Reset Columns"),
+									LOCTEXT("ResetColumnsToolTip", "Reset all columns to be visible again."),
+									FSlateIcon(),
+									FUIAction(FExecuteAction::CreateSP(AssetView, &SAssetView::ResetColumns)),
+									EUserInterfaceActionType::Button
+								);
+
+								Section.AddMenuEntry(
+									"ExportColumns",
+									LOCTEXT("ExportColumns", "Export to CSV"),
+									LOCTEXT("ExportColumnsToolTip", "Export column data to CSV."),
+									FSlateIcon(),
+									FUIAction(FExecuteAction::CreateSP(AssetView, &SAssetView::ExportColumns)),
+									EUserInterfaceActionType::Button
+								);
+							}
 						}
 					}
-#endif
 
 					AssetView->PopulateFilterAdditionalParams(FilterParams);
 				}
@@ -426,8 +433,6 @@ FDelayedAutoRegisterHelper SAssetView::AssetViewOptionsMenuRegistration(
 		}));
 	});
 
-#if UE_CONTENTBROWSER_NEW_STYLE
-
 FDelayedAutoRegisterHelper SAssetView::NavigationBarMenuExtensionRegistration(
 	EDelayedRegisterRunPhase::EndOfEngineInit,
 	[]
@@ -439,28 +444,31 @@ FDelayedAutoRegisterHelper SAssetView::NavigationBarMenuExtensionRegistration(
 			AssetCountSection.InsertPosition = FToolMenuInsert(NAME_None, EToolMenuInsertType::Last);
 			AssetCountSection.Alignment = EToolMenuSectionAlign::Last;
 
-			AssetCountSection.AddDynamicEntry(
-				"AssetCount",
-				FNewToolMenuSectionDelegate::CreateLambda([](FToolMenuSection& InSection)
-				{
-					if (const UContentBrowserToolbarMenuContext* Context = InSection.FindContext<UContentBrowserToolbarMenuContext>())
+			if (GEnableContentBrowserNewStyle)
+			{
+				AssetCountSection.AddDynamicEntry(
+					"AssetCount",
+					FNewToolMenuSectionDelegate::CreateLambda([](FToolMenuSection& InSection)
 					{
-						if (Context->AssetView.IsValid())
+						if (const UContentBrowserToolbarMenuContext* Context = InSection.FindContext<UContentBrowserToolbarMenuContext>())
 						{
-							const TSharedRef<SAssetView> AssetView = Context->AssetView.Pin().ToSharedRef();
-
-							InSection.AddEntry(
-								FToolMenuEntry::InitWidget(
-									"AssetCount",
-									SNew(STextBlock)
-									.Text(AssetView, &SAssetView::GetAssetCountText),
-									FText::GetEmpty()));
+							if (Context->AssetView.IsValid())
+							{
+								const TSharedRef<SAssetView> AssetView = Context->AssetView.Pin().ToSharedRef();
+
+								FToolMenuEntry& AssetCountEntry = InSection.AddEntry(
+									FToolMenuEntry::InitWidget(
+										"AssetCount",
+										SNew(STextBlock)
+										.Text(AssetView, &SAssetView::GetAssetCountText),
+										FText::GetEmpty()));
+
+								AssetCountEntry.WidgetData.StyleParams.VerticalAlignment = VAlign_Center;
+							}
 						}
-					}
-				}));
+					}));
+			}
 		}
 	});
 
-#endif
-
 #undef LOCTEXT_NAMESPACE
diff --git a/Engine/Source/Editor/ContentBrowser/Private/Menus/ContentBrowserMenus.cpp b/Engine/Source/Editor/ContentBrowser/Private/Menus/ContentBrowserMenus.cpp
--- a/Engine/Source/Editor/ContentBrowser/Private/Menus/ContentBrowserMenus.cpp
+++ b/Engine/Source/Editor/ContentBrowser/Private/Menus/ContentBrowserMenus.cpp
@@ -2,22 +2,24 @@
 
 #include "ContentBrowserMenuContexts.h"
 #include "Misc/DelayedAutoRegister.h"
+#include "SActionButton.h"
+#include "SAssetSearchBox.h"
 #include "SContentBrowser.h"
+#include "SNavigationBar.h"
+#include "SPositiveActionButton.h"
 #include "ToolMenu.h"
 #include "ToolMenus.h"
 #include "UObject/NameTypes.h"
 
-#include "SPositiveActionButton.h"
+#define LOCTEXT_NAMESPACE "ContentBrowser"
 
-#if UE_CONTENTBROWSER_NEW_STYLE
-#include "SActionButton.h"
-#include "SAssetSearchBox.h"
-#include "SNavigationBar.h"
-#endif
+namespace UE::Editor::ContentBrowser
+{
+	extern CONTENTBROWSER_API bool GEnableContentBrowserNewStyle;
+}
 
-#define LOCTEXT_NAMESPACE "ContentBrowser"
+using namespace UE::Editor::ContentBrowser;
 
-#if UE_CONTENTBROWSER_NEW_STYLE
 FDelayedAutoRegisterHelper SContentBrowser::NavigationBarMenuRegistration(
 	EDelayedRegisterRunPhase::EndOfEngineInit,
 	[]
@@ -37,6 +39,11 @@ FDelayedAutoRegisterHelper SContentBrowser::NavigationBarMenuRegistration(
 		.AddDynamicEntry(NAME_None,
 			FNewToolMenuSectionDelegate::CreateLambda([](FToolMenuSection& InSection)
 			{
+				if (!GEnableContentBrowserNewStyle)
+				{
+					return;
+				}
+
 				if (const UContentBrowserToolbarMenuContext* Context = InSection.FindContext<UContentBrowserToolbarMenuContext>())
 				{
 					if (Context->ContentBrowser.IsValid())
@@ -82,6 +89,11 @@ FDelayedAutoRegisterHelper SContentBrowser::NavigationBarMenuRegistration(
 		.AddDynamicEntry(NAME_None, FNewToolMenuSectionDelegate::CreateLambda(
 			[](FToolMenuSection& InSection)
 			{
+				if (!GEnableContentBrowserNewStyle)
+				{
+					return;
+				}
+
 				static constexpr float NavigationPathMinWidth = 180.0f;
 				static constexpr float NavigationPathMaxWidth = 640.0f;
 
@@ -99,10 +111,10 @@ FDelayedAutoRegisterHelper SContentBrowser::NavigationBarMenuRegistration(
 
 						FToolMenuEntry& PathBoxEntry =
 							InSection.AddEntry(
-							FToolMenuEntry::InitWidget(
-								"PathBox",
-								ContentBrowser->NavigationBar.ToSharedRef(),
-								FText::GetEmpty()));
+								FToolMenuEntry::InitWidget(
+									"PathBox",
+									ContentBrowser->NavigationBar.ToSharedRef(),
+									FText::GetEmpty()));
 
 						PathBoxEntry.WidgetData.StyleParams.SizeRule = FSizeParam::SizeRule_StretchContent;
 						PathBoxEntry.WidgetData.StyleParams.MinSize = NavigationPathMinWidth;
@@ -112,7 +124,6 @@ FDelayedAutoRegisterHelper SContentBrowser::NavigationBarMenuRegistration(
 				}
 			}));
 	});
-#endif
 
 FDelayedAutoRegisterHelper SContentBrowser::AddNewContextMenuRegistration(
 	EDelayedRegisterRunPhase::EndOfEngineInit,
@@ -266,52 +277,56 @@ FDelayedAutoRegisterHelper SContentBrowser::ToolBarMenuRegistration(
 					{
 						const TSharedRef<SContentBrowser> ContentBrowser = Context->ContentBrowser.Pin().ToSharedRef();
 
-#if UE_CONTENTBROWSER_NEW_STYLE
-						TSharedRef<SWidget> SaveButton =
-							SNew(SActionButton)
-							.ToolTipText(LOCTEXT("SaveDirtyPackagesTooltip", "Save all modified assets."))
-							.OnClicked(ContentBrowser, &SContentBrowser::OnSaveClicked)
-							.AddMetaData<FTagMetaData>(FTagMetaData(TEXT("ContentBrowserSaveDirtyPackages")))
-							.Icon(FAppStyle::Get().GetBrush("MainFrame.SaveAll"))
-							.Text(LOCTEXT("SaveAll", "Save All"));
-#else
-						TSharedRef<SWidget> SaveButton =
-							SNew(SButton)
-							.ButtonStyle(FAppStyle::Get(), "SimpleButton")
-							.ToolTipText(LOCTEXT("SaveDirtyPackagesTooltip", "Save all modified assets."))
-							.ContentPadding(2.f)
-							.OnClicked(ContentBrowser, &SContentBrowser::OnSaveClicked)
-							.AddMetaData<FTagMetaData>(FTagMetaData(TEXT("ContentBrowserSaveDirtyPackages")))
-							[
-								SNew(SHorizontalBox)
-								// Save All Icon
-								+ SHorizontalBox::Slot()
-								.AutoWidth()
-								.HAlign(HAlign_Center)
-								.VAlign(VAlign_Center)
-								[
-									SNew(SImage)
-									.Image(FAppStyle::Get().GetBrush("MainFrame.SaveAll"))
-									.ColorAndOpacity(FSlateColor::UseForeground())
-								]
-
-								// Save All Text
-								+ SHorizontalBox::Slot()
-								.Padding(FMargin(3, 0, 0, 0))
-								.VAlign(VAlign_Center)
-								.AutoWidth()
+						TSharedPtr<SWidget> SaveButton = nullptr;
+						if (GEnableContentBrowserNewStyle)
+						{
+							SaveButton =
+								SNew(SActionButton)
+								.ToolTipText(LOCTEXT("SaveDirtyPackagesTooltip", "Save all modified assets."))
+								.OnClicked(ContentBrowser, &SContentBrowser::OnSaveClicked)
+								.AddMetaData<FTagMetaData>(FTagMetaData(TEXT("ContentBrowserSaveDirtyPackages")))
+								.Icon(FAppStyle::Get().GetBrush("MainFrame.SaveAll"))
+								.Text(LOCTEXT("SaveAll", "Save All"));
+						}
+						else
+						{
+							SaveButton =
+								SNew(SButton)
+								.ButtonStyle(FAppStyle::Get(), "SimpleButton")
+								.ToolTipText(LOCTEXT("SaveDirtyPackagesTooltip", "Save all modified assets."))
+								.ContentPadding(2.f)
+								.OnClicked(ContentBrowser, &SContentBrowser::OnSaveClicked)
+								.AddMetaData<FTagMetaData>(FTagMetaData(TEXT("ContentBrowserSaveDirtyPackages")))
 								[
-									SNew(STextBlock)
-									.TextStyle(FAppStyle::Get(), "NormalText")
-									.Text(LOCTEXT("SaveAll", "Save All"))
-								]
-							];
-#endif
+									SNew(SHorizontalBox)
+									// Save All Icon
+									+ SHorizontalBox::Slot()
+									.AutoWidth()
+									.HAlign(HAlign_Center)
+									.VAlign(VAlign_Center)
+									[
+										SNew(SImage)
+										.Image(FAppStyle::Get().GetBrush("MainFrame.SaveAll"))
+										.ColorAndOpacity(FSlateColor::UseForeground())
+									]
+
+									// Save All Text
+									+ SHorizontalBox::Slot()
+									.Padding(FMargin(3, 0, 0, 0))
+									.VAlign(VAlign_Center)
+									.AutoWidth()
+									[
+										SNew(STextBlock)
+										.TextStyle(FAppStyle::Get(), "NormalText")
+										.Text(LOCTEXT("SaveAll", "Save All"))
+									]
+								];
+						}
 
 						InSection.AddEntry(
 							FToolMenuEntry::InitWidget(
 								"SaveButton",
-								SaveButton,
+								SaveButton.ToSharedRef(),
 								FText::GetEmpty(),
 								true,
 								false
@@ -320,108 +335,109 @@ FDelayedAutoRegisterHelper SContentBrowser::ToolBarMenuRegistration(
 				}
 			}));
 
-#if UE_CONTENTBROWSER_NEW_STYLE
-		ToolBar->AddSection("Filter")
-		.AddDynamicEntry(NAME_None, FNewToolMenuSectionDelegate::CreateLambda(
-			[](FToolMenuSection& InSection)
-			{
-				if (const UContentBrowserToolbarMenuContext* Context = InSection.FindContext<UContentBrowserToolbarMenuContext>())
+		if (GEnableContentBrowserNewStyle)
+		{
+			ToolBar->AddSection("Filter")
+			.AddDynamicEntry(NAME_None, FNewToolMenuSectionDelegate::CreateLambda(
+				[](FToolMenuSection& InSection)
 				{
-					if (Context->ContentBrowser.IsValid())
+					if (const UContentBrowserToolbarMenuContext* Context = InSection.FindContext<UContentBrowserToolbarMenuContext>())
 					{
-						const TSharedRef<SContentBrowser> ContentBrowser = Context->ContentBrowser.Pin().ToSharedRef();
-
-						if (!ContentBrowser->FilterComboButton.IsValid())
+						if (Context->ContentBrowser.IsValid())
 						{
-							return;
+							const TSharedRef<SContentBrowser> ContentBrowser = Context->ContentBrowser.Pin().ToSharedRef();
+
+							if (!ContentBrowser->FilterComboButton.IsValid())
+							{
+								return;
+							}
+
+							FToolMenuEntry& FilterEntry = InSection.AddEntry(
+								FToolMenuEntry::InitWidget(
+									"Filter",
+									SNew(SBox)
+									.HAlign(HAlign_Left)
+									.VAlign(VAlign_Center)
+									[
+										ContentBrowser->FilterComboButton.ToSharedRef()
+									],
+									LOCTEXT("Filter", "Filter")
+								));
+
+							FilterEntry.WidgetData.StyleParams.HorizontalAlignment = HAlign_Left;
+							FilterEntry.WidgetData.StyleParams.SizeRule = FSizeParam::ESizeRule::SizeRule_Auto;
+
+							// ... no label when in the toolbar
+							FilterEntry.ToolBarData.LabelOverride = FText::GetEmpty();
 						}
+					}
+				}));
 
-						FToolMenuEntry& FilterEntry = InSection.AddEntry(
-							FToolMenuEntry::InitWidget(
-								"Filter",
-								SNew(SBox)
-								.HAlign(HAlign_Left)
-								.VAlign(VAlign_Center)
-								[
-									ContentBrowser->FilterComboButton.ToSharedRef()
-								],
-								LOCTEXT("Filter", "Filter")
-							));
+			ToolBar->AddSection("Search")
+			.AddDynamicEntry(NAME_None, FNewToolMenuSectionDelegate::CreateLambda(
+				[](FToolMenuSection& InSection)
+				{
+					static constexpr float SearchBoxMinWidth = 180.0f;
+					static constexpr float SearchBoxMaxWidth = 640.0f;
 
-						FilterEntry.WidgetData.StyleParams.HorizontalAlignment = HAlign_Left;
-						FilterEntry.WidgetData.StyleParams.SizeRule = FSizeParam::ESizeRule::SizeRule_Auto;
+					if (const UContentBrowserToolbarMenuContext* Context = InSection.FindContext<UContentBrowserToolbarMenuContext>())
+					{
+						if (Context->ContentBrowser.IsValid())
+						{
+							const TSharedRef<SContentBrowser> ContentBrowser = Context->ContentBrowser.Pin().ToSharedRef();
 
-						// ... no label when in the toolbar
-						FilterEntry.ToolBarData.LabelOverride = FText::GetEmpty();
+							// Should have been initialized in SContentBrowser::CreateAssetView
+							if (!ensure(ContentBrowser->SearchBoxPtr.IsValid()))
+							{
+								return;
+							}
+
+							FToolMenuEntry& SearchBoxEntry =
+								InSection.AddEntry(
+								FToolMenuEntry::InitWidget(
+									"SearchBox",
+									ContentBrowser->SearchBoxPtr.ToSharedRef(),
+									FText::GetEmpty()));
+
+							SearchBoxEntry.WidgetData.StyleParams.SizeRule = FSizeParam::SizeRule_StretchContent;
+							SearchBoxEntry.WidgetData.StyleParams.MinSize = SearchBoxMinWidth;
+							SearchBoxEntry.WidgetData.StyleParams.MaxSize = SearchBoxMaxWidth;
+							SearchBoxEntry.WidgetData.StyleParams.DesiredWidth = SearchBoxMaxWidth;
+						}
 					}
-				}
-			}));
+				}));
 
-		ToolBar->AddSection("Search")
-		.AddDynamicEntry(NAME_None, FNewToolMenuSectionDelegate::CreateLambda(
-			[](FToolMenuSection& InSection)
-			{
-				static constexpr float SearchBoxMinWidth = 180.0f;
-				static constexpr float SearchBoxMaxWidth = 640.0f;
+			FToolMenuSection& SettingsSection = ToolBar->AddSection("Settings");
+			SettingsSection.Alignment = EToolMenuSectionAlign::Last;
 
+			SettingsSection.AddDynamicEntry(NAME_None,
+			FNewToolMenuSectionDelegate::CreateLambda([](FToolMenuSection& InSection)
+			{
 				if (const UContentBrowserToolbarMenuContext* Context = InSection.FindContext<UContentBrowserToolbarMenuContext>())
 				{
 					if (Context->ContentBrowser.IsValid())
 					{
 						const TSharedRef<SContentBrowser> ContentBrowser = Context->ContentBrowser.Pin().ToSharedRef();
 
-						// Should have been initialized in SContentBrowser::CreateAssetView
-						if (!ensure(ContentBrowser->SearchBoxPtr.IsValid()))
-						{
-							return;
-						}
+						const TSharedRef<SWidget> LockToggleButton =
+							SNew(SActionButton)
+							.ActionButtonType(EActionButtonType::Simple)
+							.Icon(ContentBrowser, &SContentBrowser::GetLockIconBrush)
+							.IconColorAndOpacity(FSlateColor::UseStyle())
+							.ToolTipText(LOCTEXT("LockToggleTooltip", "Toggle lock. If locked, this browser will ignore Find in Content Browser requests."))
+							.OnClicked(ContentBrowser, &SContentBrowser::ToggleLockClicked)
+							.AddMetaData<FTagMetaData>(FTagMetaData(TEXT("ContentBrowserLock")));
 
-						FToolMenuEntry& SearchBoxEntry =
-							InSection.AddEntry(
+						InSection.AddEntry(
 							FToolMenuEntry::InitWidget(
-								"SearchBox",
-								ContentBrowser->SearchBoxPtr.ToSharedRef(),
-								FText::GetEmpty()));
-
-						SearchBoxEntry.WidgetData.StyleParams.SizeRule = FSizeParam::SizeRule_StretchContent;
-						SearchBoxEntry.WidgetData.StyleParams.MinSize = SearchBoxMinWidth;
-						SearchBoxEntry.WidgetData.StyleParams.MaxSize = SearchBoxMaxWidth;
-						SearchBoxEntry.WidgetData.StyleParams.DesiredWidth = SearchBoxMaxWidth;
+								"Lock",
+								LockToggleButton,
+								FText::GetEmpty()
+							));
 					}
 				}
 			}));
-
-		FToolMenuSection& SettingsSection = ToolBar->AddSection("Settings");
-		SettingsSection.Alignment = EToolMenuSectionAlign::Last;
-
-		SettingsSection.AddDynamicEntry(NAME_None,
-		FNewToolMenuSectionDelegate::CreateLambda([](FToolMenuSection& InSection)
-		{
-			if (const UContentBrowserToolbarMenuContext* Context = InSection.FindContext<UContentBrowserToolbarMenuContext>())
-			{
-				if (Context->ContentBrowser.IsValid())
-				{
-					const TSharedRef<SContentBrowser> ContentBrowser = Context->ContentBrowser.Pin().ToSharedRef();
-
-					const TSharedRef<SWidget> LockToggleButton =
-						SNew(SActionButton)
-						.ActionButtonType(EActionButtonType::Simple)
-						.Icon(ContentBrowser, &SContentBrowser::GetLockIconBrush)
-						.IconColorAndOpacity(FSlateColor::UseStyle())
-						.ToolTipText(LOCTEXT("LockToggleTooltip", "Toggle lock. If locked, this browser will ignore Find in Content Browser requests."))
-						.OnClicked(ContentBrowser, &SContentBrowser::ToggleLockClicked)
-						.AddMetaData<FTagMetaData>(FTagMetaData(TEXT("ContentBrowserLock")));
-
-					InSection.AddEntry(
-						FToolMenuEntry::InitWidget(
-							"Lock",
-							LockToggleButton,
-							FText::GetEmpty()
-						));
-				}
-			}
-		}));
-#endif
+		}
 	});
 
 #undef LOCTEXT_NAMESPACE
diff --git a/Engine/Source/Editor/ContentBrowser/Private/SAssetView.cpp b/Engine/Source/Editor/ContentBrowser/Private/SAssetView.cpp
--- a/Engine/Source/Editor/ContentBrowser/Private/SAssetView.cpp
+++ b/Engine/Source/Editor/ContentBrowser/Private/SAssetView.cpp
@@ -76,10 +76,6 @@
 #include "ISourceControlModule.h"
 #include "RevisionControlStyle/RevisionControlStyle.h"
 
-#if UE_CONTENTBROWSER_NEW_STYLE
-#include "SPathView.h"
-#endif
-
 #define LOCTEXT_NAMESPACE "ContentBrowser"
 #define MAX_THUMBNAIL_SIZE 4096
 
@@ -90,6 +86,8 @@
 	#define checkAssetList(cond)
 #endif
 
+using namespace UE::Editor::ContentBrowser;
+
 namespace UE::AssetView
 {
 	/** Time delay between recently added items being added to the filtered asset items list */
@@ -1114,11 +1112,7 @@ void SAssetView::Construct( const FArguments& InArgs )
 	TileViewThumbnailResolution = 256;
 
 	// Max Size for the thumbnail
-#if UE_CONTENTBROWSER_NEW_STYLE
-	constexpr int32 MaxTileViewThumbnailSize = 160;
-#else
-	constexpr int32 MaxTileViewThumbnailSize = 150;
-#endif
+	const int32 MaxTileViewThumbnailSize = GEnableContentBrowserNewStyle ? 160 : 150;
 
 	TileViewThumbnailSize = MaxTileViewThumbnailSize;
 
@@ -1374,56 +1368,59 @@ void SAssetView::Construct( const FArguments& InArgs )
 
 	if (InArgs._ShowBottomToolbar)
 	{
-		// Bottom panel
-		VerticalBox->AddSlot()
-		.AutoHeight()
-		[
-			SNew(SHorizontalBox)
+		const TSharedRef<SHorizontalBox> BottomToolBarBox = SNew(SHorizontalBox);
 
-#if !UE_CONTENTBROWSER_NEW_STYLE
+		if (!GEnableContentBrowserNewStyle)
+		{
 			// Asset count
-			+SHorizontalBox::Slot()
+			BottomToolBarBox->AddSlot()
 			.FillWidth(1.f)
 			.VAlign(VAlign_Center)
 			.Padding(8, 5)
 			[
 				SNew(STextBlock)
 				.Text(this, &SAssetView::GetAssetCountText)
-			]
-#endif
+			];
+		}
 
-			// View mode combo button
-			+SHorizontalBox::Slot()
-			.AutoWidth()
+		// View mode combo button
+		BottomToolBarBox->AddSlot()
+		.AutoWidth()
+		[
+			SNew(SComboButton)
+			.Visibility(InArgs._ShowViewOptions ? EVisibility::Visible : EVisibility::Collapsed)
+			.ContentPadding(0.f)
+			.ButtonStyle( FAppStyle::Get(), "ToggleButton" ) // Use the tool bar item style for this button
+			.OnGetMenuContent( this, &SAssetView::GetViewButtonContent )
+			.ButtonContent()
 			[
-				SNew(SComboButton)
-				.Visibility(InArgs._ShowViewOptions ? EVisibility::Visible : EVisibility::Collapsed)
-				.ContentPadding(0.f)
-				.ButtonStyle( FAppStyle::Get(), "ToggleButton" ) // Use the tool bar item style for this button
-				.OnGetMenuContent( this, &SAssetView::GetViewButtonContent )
-				.ButtonContent()
+				SNew(SHorizontalBox)
+
+				+SHorizontalBox::Slot()
+				.AutoWidth()
+				.VAlign(VAlign_Center)
 				[
-					SNew(SHorizontalBox)
- 
-					+SHorizontalBox::Slot()
-					.AutoWidth()
-					.VAlign(VAlign_Center)
-					[
-						SNew(SImage)
-						.Image( FAppStyle::GetBrush("GenericViewButton") )
-					]
- 
-					+SHorizontalBox::Slot()
-					.AutoWidth()
-					.Padding(2.f, 0.f, 0.f, 0.f)
-					.VAlign(VAlign_Center)
-					[
-						SNew(STextBlock)
-						.Text( LOCTEXT("ViewButton", "View Options") )
-					]
+					SNew(SImage)
+					.Image( FAppStyle::GetBrush("GenericViewButton") )
+				]
+
+				+SHorizontalBox::Slot()
+				.AutoWidth()
+				.Padding(2.f, 0.f, 0.f, 0.f)
+				.VAlign(VAlign_Center)
+				[
+					SNew(STextBlock)
+					.Text( LOCTEXT("ViewButton", "View Options") )
 				]
 			]
 		];
+
+		// Bottom panel
+		VerticalBox->AddSlot()
+		.AutoHeight()
+		[
+			BottomToolBarBox
+		];
 	}
 
 	CreateCurrentView();
@@ -1858,12 +1855,14 @@ void SAssetView::LoadSettings(const FString& IniFilename, const FString& IniSect
 		ThumbnailSizeConfig = FMath::Clamp<int32>(ThumbnailSizeConfig, 0, (int32)EThumbnailSize::MAX-1);
 
 		// TODO: Remove this afterwards, current CB should hide new size
-#if !UE_CONTENTBROWSER_NEW_STYLE
-		if ((EThumbnailSize)ThumbnailSizeConfig == EThumbnailSize::XLarge)
+		if (!GEnableContentBrowserNewStyle)
 		{
-			ThumbnailSizeConfig -= 1;
+			if ((EThumbnailSize)ThumbnailSizeConfig == EThumbnailSize::XLarge)
+			{
+				ThumbnailSizeConfig -= 1;
+			}
 		}
-#endif
+
 		ThumbnailSize = (EThumbnailSize) ThumbnailSizeConfig;
 		// Set the thumbnail value here after the Size enum is retrieved from the config
 		UpdateThumbnailSizeValue();
@@ -2487,47 +2486,50 @@ FReply SAssetView::OnMouseWheel( const FGeometry& MyGeometry, const FPointerEven
 	// Make sure to not change the thumbnail scaling when we're in Columns view since thumbnail scaling isn't applicable there.
 	if( MouseEvent.IsControlDown() && IsThumbnailScalingAllowed() )
 	{
-#if UE_CONTENTBROWSER_NEW_STYLE
-		const float NewDelta = MouseEvent.GetWheelDelta() * 0.4f;
-		if ((ZoomScale == 1.f && NewDelta > 0) || (ZoomScale == 0.f && NewDelta < 0))
+		if (GEnableContentBrowserNewStyle)
 		{
-			int32 Step = (int32)FMath::Sign(NewDelta);
-			EThumbnailSize OldSize = ThumbnailSize;
-			ThumbnailSize = (EThumbnailSize)FMath::Clamp<int32>(((int32)ThumbnailSize + Step), 0, (int32)EThumbnailSize::MAX - 1);
-			if (OldSize != ThumbnailSize)
+			const float NewDelta = MouseEvent.GetWheelDelta() * 0.4f;
+			if ((ZoomScale == 1.f && NewDelta > 0) || (ZoomScale == 0.f && NewDelta < 0))
+			{
+				int32 Step = (int32)FMath::Sign(NewDelta);
+				EThumbnailSize OldSize = ThumbnailSize;
+				ThumbnailSize = (EThumbnailSize)FMath::Clamp<int32>(((int32)ThumbnailSize + Step), 0, (int32)EThumbnailSize::MAX - 1);
+				if (OldSize != ThumbnailSize)
+				{
+					OnThumbnailSizeChanged(ThumbnailSize);
+					ZoomScale = NewDelta > 0.f ? 0.f : 1.f;
+				}
+			}
+			else
 			{
-				OnThumbnailSizeChanged(ThumbnailSize);
-				ZoomScale = NewDelta > 0.f ? 0.f : 1.f;
+				ZoomScale = FMath::Clamp(ZoomScale + NewDelta, 0.f, 1.f);
 			}
 		}
 		else
 		{
-			ZoomScale = FMath::Clamp(ZoomScale + NewDelta, 0.f, 1.f);
-		}
-#else
-		// Step up/down a level depending on the scroll wheel direction.
-		// Clamp value to enum min/max before updating.
-		const int32 Delta = MouseEvent.GetWheelDelta() > 0 ? 1 : -1;
-		EThumbnailSize DesiredThumbnailSize = (EThumbnailSize)FMath::Clamp<int32>((int32)ThumbnailSize + Delta, 0, (int32)EThumbnailSize::MAX - 1);
+			// Step up/down a level depending on the scroll wheel direction.
+			// Clamp value to enum min/max before updating.
+			const int32 Delta = MouseEvent.GetWheelDelta() > 0 ? 1 : -1;
+			EThumbnailSize DesiredThumbnailSize = (EThumbnailSize)FMath::Clamp<int32>((int32)ThumbnailSize + Delta, 0, (int32)EThumbnailSize::MAX - 1);
 
-		// TODO: Remove this afterwards, current CB should hide new size
-		if (DesiredThumbnailSize == EThumbnailSize::XLarge)
-		{
-			if (Delta > 0)
+			// TODO: Remove this afterwards, current CB should hide new size
+			if (DesiredThumbnailSize == EThumbnailSize::XLarge)
 			{
-				DesiredThumbnailSize = EThumbnailSize::Huge;
+				if (Delta > 0)
+				{
+					DesiredThumbnailSize = EThumbnailSize::Huge;
+				}
+				else
+				{
+					DesiredThumbnailSize = EThumbnailSize::Large;
+				}
 			}
-			else
+
+			if ( DesiredThumbnailSize != ThumbnailSize )
 			{
-				DesiredThumbnailSize = EThumbnailSize::Large;
+				OnThumbnailSizeChanged(DesiredThumbnailSize);
 			}
 		}
-
-		if ( DesiredThumbnailSize != ThumbnailSize )
-		{
-			OnThumbnailSizeChanged(DesiredThumbnailSize);
-		}
-#endif
 		return FReply::Handled();
 	}
 	return FReply::Unhandled();
@@ -4810,13 +4812,9 @@ TSharedPtr<FAssetThumbnail> SAssetView::AddItemToNewThumbnailRelevancyMap(const
 		}
 
 		// The thumbnail newly relevant, create a new thumbnail
-		int32 ThumbnailResolution = 0;
-
-#if UE_CONTENTBROWSER_NEW_STYLE
-		ThumbnailResolution = CurrentThumbnailSize;
-#else
-		ThumbnailResolution = FMath::TruncToInt((float)CurrentThumbnailSize * MaxThumbnailScale);
-#endif
+		int32 ThumbnailResolution = GEnableContentBrowserNewStyle
+			? CurrentThumbnailSize
+			: FMath::TruncToInt((float)CurrentThumbnailSize * MaxThumbnailScale);
 
 		Thumbnail = MakeShared<FAssetThumbnail>(FAssetData(), ThumbnailResolution, ThumbnailResolution, AssetThumbnailPool);
 		Item->GetItem().UpdateThumbnail(*Thumbnail);
@@ -5313,40 +5311,42 @@ bool SAssetView::IsThumbnailScalingAllowed() const
 
 float SAssetView::GetTileViewTypeNameHeight() const
 {
-#if UE_CONTENTBROWSER_NEW_STYLE
+	if (GEnableContentBrowserNewStyle)
+	{
 		if (ThumbnailSize == EThumbnailSize::Tiny)
 		{
 			return 0;
 		}
 		return 67.f;
-#else
-	float TypeNameHeight = 0;
-
-	if (bShowTypeInTileView)
-	{
-		TypeNameHeight = 50;
 	}
 	else
 	{
-		if (ThumbnailSize == EThumbnailSize::Small)
-		{
-			TypeNameHeight = 25;
-		}
-		else if (ThumbnailSize == EThumbnailSize::Medium)
-		{
-			TypeNameHeight = -5;
-		}
-		else if (ThumbnailSize > EThumbnailSize::Medium)
+		float TypeNameHeight = 0;
+		if (bShowTypeInTileView)
 		{
-			TypeNameHeight = -25;
+			TypeNameHeight = 50;
 		}
 		else
 		{
-			TypeNameHeight = -40;
+			if (ThumbnailSize == EThumbnailSize::Small)
+			{
+				TypeNameHeight = 25;
+			}
+			else if (ThumbnailSize == EThumbnailSize::Medium)
+			{
+				TypeNameHeight = -5;
+			}
+			else if (ThumbnailSize > EThumbnailSize::Medium)
+			{
+				TypeNameHeight = -25;
+			}
+			else
+			{
+				TypeNameHeight = -40;
+			}
 		}
+		return TypeNameHeight;
 	}
-	return TypeNameHeight;
-#endif
 }
 
 float SAssetView::GetSourceControlIconHeight() const
@@ -5361,29 +5361,23 @@ float SAssetView::GetListViewItemHeight() const
 
 float SAssetView::GetTileViewItemHeight() const
 {
-#if UE_CONTENTBROWSER_NEW_STYLE
-	return GetTileViewItemBaseWidth() + GetTileViewTypeNameHeight() + TileViewHeightPadding;
-#else
-	return (((float)TileViewNameHeight + GetTileViewTypeNameHeight()) * FMath::Lerp(MinThumbnailScale, MaxThumbnailScale, GetThumbnailScale())) + GetTileViewItemBaseHeight() * FillScale + GetSourceControlIconHeight();
-#endif
+	return GEnableContentBrowserNewStyle
+		? GetTileViewItemBaseWidth() + GetTileViewTypeNameHeight() + TileViewHeightPadding
+		: (((float)TileViewNameHeight + GetTileViewTypeNameHeight()) * FMath::Lerp(MinThumbnailScale, MaxThumbnailScale, GetThumbnailScale())) + GetTileViewItemBaseHeight() * FillScale + GetSourceControlIconHeight();
 }
 
 float SAssetView::GetTileViewItemBaseHeight() const
 {
-#if UE_CONTENTBROWSER_NEW_STYLE
-	return GetTileViewItemBaseWidth();
-#else
-	return (float)(TileViewThumbnailSize + TileViewThumbnailPadding * 2) * FMath::Lerp(MinThumbnailScale, MaxThumbnailScale, GetThumbnailScale());
-#endif
+	return GEnableContentBrowserNewStyle
+		? GetTileViewItemBaseWidth()
+		: (float)(TileViewThumbnailSize + TileViewThumbnailPadding * 2) * FMath::Lerp(MinThumbnailScale, MaxThumbnailScale, GetThumbnailScale());
 }
 
 float SAssetView::GetTileViewItemWidth() const
 {
-#if UE_CONTENTBROWSER_NEW_STYLE
-	return GetTileViewItemBaseWidth() + TileViewWidthPadding;
-#else
-	return GetTileViewItemBaseWidth() * FillScale;
-#endif
+	return GEnableContentBrowserNewStyle
+		? GetTileViewItemBaseWidth() + TileViewWidthPadding
+		: GetTileViewItemBaseWidth() * FillScale;
 }
 
 float SAssetView::GetTileViewThumbnailDimension() const
@@ -5393,11 +5387,9 @@ float SAssetView::GetTileViewThumbnailDimension() const
 
 float SAssetView::GetTileViewItemBaseWidth() const //-V524
 {
-#if UE_CONTENTBROWSER_NEW_STYLE
-	return GetTileViewThumbnailDimension();
-#else
-	return (float)( TileViewThumbnailSize + TileViewThumbnailPadding * 2 ) * FMath::Lerp( MinThumbnailScale, MaxThumbnailScale, GetThumbnailScale() );
-#endif
+	return GEnableContentBrowserNewStyle
+		? GetTileViewThumbnailDimension()
+		: (float)( TileViewThumbnailSize + TileViewThumbnailPadding * 2 ) * FMath::Lerp( MinThumbnailScale, MaxThumbnailScale, GetThumbnailScale() );
 }
 
 EColumnSortMode::Type SAssetView::GetColumnSortMode(const FName ColumnId) const
diff --git a/Engine/Source/Editor/ContentBrowser/Private/SContentBrowser.cpp b/Engine/Source/Editor/ContentBrowser/Private/SContentBrowser.cpp
--- a/Engine/Source/Editor/ContentBrowser/Private/SContentBrowser.cpp
+++ b/Engine/Source/Editor/ContentBrowser/Private/SContentBrowser.cpp
@@ -95,6 +95,7 @@
 #include "Modules/ModuleManager.h"
 #include "NewAssetOrClassContextMenu.h"
 #include "PathContextMenu.h"
+#include "SActionButton.h"
 #include "SAssetSearchBox.h"
 #include "SAssetView.h"
 #include "SCollectionView.h"
@@ -155,10 +156,6 @@
 #include "Widgets/SNullWidget.h"
 #include "Widgets/SWidget.h"
 #include "Widgets/Text/STextBlock.h"
-#if UE_CONTENTBROWSER_NEW_STYLE
-#include "SActionButton.h"
-#include "Widgets/Colors/SColorBlock.h"
-#endif
 
 class FTreeItem;
 struct FGeometry;
@@ -168,6 +165,8 @@ struct FSlateBrush;
 
 const FString SContentBrowser::SettingsIniSection = TEXT("ContentBrowser");
 
+using namespace UE::Editor::ContentBrowser;
+
 namespace ContentBrowserSourcesWidgetSwitcherIndex
 {
 	static constexpr int32 PathView = 0;
@@ -491,225 +490,231 @@ void SContentBrowser::Construct( const FArguments& InArgs, const FName& InInstan
 	// Note, for backwards compatibility ShowBottomToolbar controls the visibility of view options so we respect that here
 	if (bShowBottomToolbar)
 	{
-#if UE_CONTENTBROWSER_NEW_STYLE
-		ViewOptions =
-			SNew(SActionButton)
-			.ActionButtonType(EActionButtonType::Simple)
-			.OnGetMenuContent(AssetViewPtr.ToSharedRef(), &SAssetView::GetViewButtonContent)
-			.Icon(FAppStyle::Get().GetBrush("Icons.Settings"));
-#else
-		ViewOptions =
-			SNew(SComboButton)
-			.ComboButtonStyle(&FAppStyle::Get().GetWidgetStyle<FComboButtonStyle>("SimpleComboButton"))
-			.OnGetMenuContent(AssetViewPtr.ToSharedRef(), &SAssetView::GetViewButtonContent)
-			.HasDownArrow(false)
-			.ButtonContent()
-			[
-				SNew(SHorizontalBox)
-				+ SHorizontalBox::Slot()
-				.AutoWidth()
-				.VAlign(VAlign_Center)
-				.Padding(4.0, 0.0f)
+		if (GEnableContentBrowserNewStyle)
+		{
+			ViewOptions =
+				SNew(SActionButton)
+				.ActionButtonType(EActionButtonType::Simple)
+				.OnGetMenuContent(AssetViewPtr.ToSharedRef(), &SAssetView::GetViewButtonContent)
+				.Icon(FAppStyle::Get().GetBrush("Icons.Settings"));
+		}
+		else
+		{
+			ViewOptions =
+				SNew(SComboButton)
+				.ComboButtonStyle(&FAppStyle::Get().GetWidgetStyle<FComboButtonStyle>("SimpleComboButton"))
+				.OnGetMenuContent(AssetViewPtr.ToSharedRef(), &SAssetView::GetViewButtonContent)
+				.HasDownArrow(false)
+				.ButtonContent()
 				[
-					SNew(SImage)
-					.ColorAndOpacity(FSlateColor::UseForeground())
-					.Image(FAppStyle::Get().GetBrush("Icons.Settings"))
-				]
+					SNew(SHorizontalBox)
+					+ SHorizontalBox::Slot()
+					.AutoWidth()
+					.VAlign(VAlign_Center)
+					.Padding(4.0, 0.0f)
+					[
+						SNew(SImage)
+						.ColorAndOpacity(FSlateColor::UseForeground())
+						.Image(FAppStyle::Get().GetBrush("Icons.Settings"))
+					]
 
-				+ SHorizontalBox::Slot()
-				.VAlign(VAlign_Center)
-				.Padding(4.0, 0.0f)
-				[
-					SNew(STextBlock)
-					.Text(LOCTEXT("Settings", "Settings"))
-					.ColorAndOpacity(FSlateColor::UseForeground())
-				]
-			];
-#endif
+					+ SHorizontalBox::Slot()
+					.VAlign(VAlign_Center)
+					.Padding(4.0, 0.0f)
+					[
+						SNew(STextBlock)
+						.Text(LOCTEXT("Settings", "Settings"))
+						.ColorAndOpacity(FSlateColor::UseForeground())
+					]
+				];
+		}
 	}
 
-#if UE_CONTENTBROWSER_NEW_STYLE
-	constexpr float ToolBarVerticalPadding = 4.0f; // ToolBar total height should be 36.0f
-	constexpr float PanelInsetPadding = 2.0f;
-
-	const TSharedRef<SWidget> AssetView = CreateAssetView(Config);
+	if (GEnableContentBrowserNewStyle)
+	{
+		constexpr float ToolBarVerticalPadding = 4.0f; // ToolBar total height should be 36.0f
+		constexpr float PanelInsetPadding = 2.0f;
 
-	const TSharedRef<SWidget> PanelContentWidget =
-		SNew(SVerticalBox)
+		const TSharedRef<SWidget> AssetView = CreateAssetView(Config);
 
-		+ SVerticalBox::Slot()
-		.AutoHeight()
+		ChildSlot
 		[
-			SNew(SSeparator)
-			.Thickness(2.0f)
-		]
+			SNew(SVerticalBox)
 
-		// Source / Tree + Assets
-		+ SVerticalBox::Slot()
-		.FillHeight(1.0f)
-		.Padding(0.0f)
-		[
-			SNew(SHorizontalBox)
+			+ SVerticalBox::Slot()
+			.AutoHeight()
+			[
+				SNew(SSeparator)
+				.Thickness(2.0f)
+			]
 
-			// Tree + Assets + Navigation Bar
-			+ SHorizontalBox::Slot()
-			.FillWidth(1.0f)
-			.Padding(FMargin(PanelInsetPadding, 0.0f, 0.0f, 0.0f))
+			// Source / Tree + Assets
+			+ SVerticalBox::Slot()
+			.FillHeight(1.0f)
+			.Padding(0.0f)
 			[
-				SNew(SVerticalBox)
+				SNew(SHorizontalBox)
 
-				// Assets/tree
-				+ SVerticalBox::Slot()
-				.FillHeight(1.0f)
-				.Padding(FMargin(0.0f, 0.0f, 0.0f, PanelInsetPadding))
+				// Tree + Assets + Navigation Bar
+				+ SHorizontalBox::Slot()
+				.FillWidth(1.0f)
+				.Padding(FMargin(PanelInsetPadding, 0.0f, 0.0f, 0.0f))
 				[
-					// The tree/assets splitter
-					SAssignNew(PathAssetSplitterPtr, SSplitter)
-					.PhysicalSplitterHandleSize(2.0f)
-
-					// Sources View
-					+ SSplitter::Slot()
-					.Resizable(true)
-					.SizeRule(SSplitter::SizeToContent)
-					.OnSlotResized(this, &SContentBrowser::OnPathViewBoxColumnResized)
+					SNew(SVerticalBox)
+
+					// Assets/tree
+					+ SVerticalBox::Slot()
+					.FillHeight(1.0f)
+					.Padding(FMargin(0.0f, 0.0f, 0.0f, PanelInsetPadding))
 					[
-						SNew(SBox)
-						.Padding(0.0f)
-						.Visibility(this, &SContentBrowser::GetSourcesViewVisibility)
-						.WidthOverride(this, &SContentBrowser::GetPathViewBoxWidthOverride)
+						// The tree/assets splitter
+						SAssignNew(PathAssetSplitterPtr, SSplitter)
+						.PhysicalSplitterHandleSize(2.0f)
+
+						// Sources View
+						+ SSplitter::Slot()
+						.Resizable(true)
+						.SizeRule(SSplitter::SizeToContent)
+						.OnSlotResized(this, &SContentBrowser::OnPathViewBoxColumnResized)
 						[
-							SNew(SBorder)
-							.Padding(FMargin(0))
-							.BorderImage(FAppStyle::GetBrush("Brushes.Recessed"))
+							SNew(SBox)
+							.Padding(0.0f)
+							.Visibility(this, &SContentBrowser::GetSourcesViewVisibility)
+							.WidthOverride(this, &SContentBrowser::GetPathViewBoxWidthOverride)
 							[
-								// Note: If adding more widgets here, fix ContentBrowserSourcesWidgetSwitcherIndex and the code that uses it!
-								SAssignNew(SourcesWidgetSwitcher, SWidgetSwitcher)
-
-								// Paths View
-								+ SWidgetSwitcher::Slot()
+								SNew(SBorder)
+								.Padding(FMargin(0))
+								.BorderImage(FAppStyle::GetBrush("Brushes.Recessed"))
 								[
-									SAssignNew(PathFavoriteSplitterPtr, SSplitter)
-									.Clipping(EWidgetClipping::ClipToBounds)
-									.PhysicalSplitterHandleSize(2.0f)
-									.HitDetectionSplitterHandleSize(8.0f)
-									.Orientation(EOrientation::Orient_Vertical)
-									.MinimumSlotHeight(26.0f)
-									.Visibility(this, &SContentBrowser::GetSourcesViewVisibility)
-									+ SSplitter::Slot()
-									.SizeRule(TAttribute<SSplitter::ESizeRule>(this, &SContentBrowser::GetFavoritesAreaSizeRule))
-									.MinSize(TAttribute<float>(this, &SContentBrowser::GetFavoritesAreaMinSize))
-									.Value(0.2f)
+									// Note: If adding more widgets here, fix ContentBrowserSourcesWidgetSwitcherIndex and the code that uses it!
+									SAssignNew(SourcesWidgetSwitcher, SWidgetSwitcher)
+
+									// Paths View
+									+ SWidgetSwitcher::Slot()
 									[
-										SNew(SBorder)
-										.BorderImage(FAppStyle::Get().GetBrush("Brushes.Header"))
-										.Padding(0.0f, 2.0f, 0.0f, 0.0f)
+										SAssignNew(PathFavoriteSplitterPtr, SSplitter)
+										.Clipping(EWidgetClipping::ClipToBounds)
+										.PhysicalSplitterHandleSize(2.0f)
+										.HitDetectionSplitterHandleSize(8.0f)
+										.Orientation(EOrientation::Orient_Vertical)
+										.MinimumSlotHeight(26.0f)
+										.Visibility(this, &SContentBrowser::GetSourcesViewVisibility)
+										+ SSplitter::Slot()
+										.SizeRule(TAttribute<SSplitter::ESizeRule>(this, &SContentBrowser::GetFavoritesAreaSizeRule))
+										.MinSize(TAttribute<float>(this, &SContentBrowser::GetFavoritesAreaMinSize))
+										.Value(0.2f)
 										[
-											CreateFavoritesView(Config)
+											SNew(SBorder)
+											.BorderImage(FAppStyle::Get().GetBrush("Brushes.Header"))
+											.Padding(0.0f, 2.0f, 0.0f, 0.0f)
+											[
+												CreateFavoritesView(Config)
+											]
 										]
-									]
 
-									+ SSplitter::Slot()
-									.SizeRule(TAttribute<SSplitter::ESizeRule>(this, &SContentBrowser::GetPathAreaSizeRule))
-									.MinSize(29.0f)
-									.Value(0.8f)
-									[
-										SNew(SBorder)
-										.BorderImage(FAppStyle::Get().GetBrush("Brushes.Header"))
-										.Padding(0.0f, 2.0f, 0.0f, 0.0f)
+										+ SSplitter::Slot()
+										.SizeRule(TAttribute<SSplitter::ESizeRule>(this, &SContentBrowser::GetPathAreaSizeRule))
+										.MinSize(29.0f)
+										.Value(0.8f)
 										[
-											CreatePathView(Config)
+											SNew(SBorder)
+											.BorderImage(FAppStyle::Get().GetBrush("Brushes.Header"))
+											.Padding(0.0f, 2.0f, 0.0f, 0.0f)
+											[
+												CreatePathView(Config)
+											]
 										]
-									]
 
-									+ SSplitter::Slot()
-									.SizeRule(TAttribute<SSplitter::ESizeRule>(this, &SContentBrowser::GetCollectionsAreaSizeRule))
-									.MinSize(TAttribute<float>(this, &SContentBrowser::GetCollectionsAreaMinSize))
-									.Value(0.4f)
-									[
-										SNew(SBorder)
-										.BorderImage(FAppStyle::Get().GetBrush("Brushes.Header"))
-										.Padding(0.0f, 2.0f, 0.0f, 0.0f)
+										+ SSplitter::Slot()
+										.SizeRule(TAttribute<SSplitter::ESizeRule>(this, &SContentBrowser::GetCollectionsAreaSizeRule))
+										.MinSize(TAttribute<float>(this, &SContentBrowser::GetCollectionsAreaMinSize))
+										.Value(0.4f)
 										[
-											CreateDockedCollectionsView(Config)
+											SNew(SBorder)
+											.BorderImage(FAppStyle::Get().GetBrush("Brushes.Header"))
+											.Padding(0.0f, 2.0f, 0.0f, 0.0f)
+											[
+												CreateDockedCollectionsView(Config)
+											]
 										]
 									]
-								]
 
-								// Collections View
-								+ SWidgetSwitcher::Slot()
-								[
-									SNew(SBox)
-									.Padding(FMargin(0.0f, PanelInsetPadding, 0.0f, 0.0f))
+									// Collections View
+									+ SWidgetSwitcher::Slot()
 									[
-										CollectionViewPtr.ToSharedRef()
+										SNew(SBox)
+										.Padding(FMargin(0.0f, PanelInsetPadding, 0.0f, 0.0f))
+										[
+											CollectionViewPtr.ToSharedRef()
+										]
 									]
 								]
 							]
 						]
-					]
 
-					// Asset View
-					+ SSplitter::Slot()
-					.Value(0.75f)
-					[
-						SNew(SBox)
-						.Padding(0.0f)
+						// Asset View
+						+ SSplitter::Slot()
+						.Value(0.75f)
 						[
-							SNew(SVerticalBox)
-
-							+ SVerticalBox::Slot()
-							.AutoHeight()
+							SNew(SBox)
+							.Padding(0.0f)
 							[
-								SNew(SBorder)
-								.Padding(FMargin(3, ToolBarVerticalPadding))
-								.BorderImage(bIsDrawer ? FStyleDefaults::GetNoBrush() : FAppStyle::Get().GetBrush("Brushes.Panel"))
-								[
-									SNew(SHorizontalBox)
+								SNew(SVerticalBox)
 
-									+ SHorizontalBox::Slot()
-									.FillWidth(1.0f)
-									.HAlign(HAlign_Fill)
-									.VAlign(VAlign_Center)
-									.Padding(5, 0, 0, 0)
+								+ SVerticalBox::Slot()
+								.AutoHeight()
+								[
+									SNew(SBorder)
+									.Padding(FMargin(3, ToolBarVerticalPadding))
+									.BorderImage(bIsDrawer ? FStyleDefaults::GetNoBrush() : FAppStyle::Get().GetBrush("Brushes.Panel"))
 									[
-										CreateToolBar(Config)
-									]
+										SNew(SHorizontalBox)
 
-									+ SHorizontalBox::Slot()
-									.AutoWidth()
-									.HAlign(HAlign_Right)
-									.VAlign(VAlign_Center)
-									[
-										CreateDrawerDockButton(Config)
-									]
+										+ SHorizontalBox::Slot()
+										.FillWidth(1.0f)
+										.HAlign(HAlign_Fill)
+										.VAlign(VAlign_Center)
+										.Padding(5, 0, 0, 0)
+										[
+											CreateToolBar(Config)
+										]
 
-									+ SHorizontalBox::Slot()
-									.AutoWidth()
-									.HAlign(HAlign_Right)
-									.VAlign(VAlign_Center)
-									.Padding(5.0f, 0.0f, 0.0f, 0.0f)
-									[
-										ViewOptions
+										+ SHorizontalBox::Slot()
+										.AutoWidth()
+										.HAlign(HAlign_Right)
+										.VAlign(VAlign_Center)
+										[
+											CreateDrawerDockButton(Config)
+										]
+
+										+ SHorizontalBox::Slot()
+										.AutoWidth()
+										.HAlign(HAlign_Right)
+										.VAlign(VAlign_Center)
+										.Padding(5.0f, 0.0f, 0.0f, 0.0f)
+										[
+											ViewOptions
+										]
 									]
 								]
-							]
 
-							+ SVerticalBox::Slot()
-							[
-								AssetView
-							]
+								+ SVerticalBox::Slot()
+								[
+									AssetView
+								]
 
-							// Navigation Bar
-							+ SVerticalBox::Slot()
-							.AutoHeight()
-							.HAlign(HAlign_Fill)
-							.Padding(FMargin(0, PanelInsetPadding, 0, 0))
-							[
-								SNew(SBorder)
-								.Padding(FMargin(3))
-								.BorderImage(bIsDrawer ? FStyleDefaults::GetNoBrush() : FAppStyle::Get().GetBrush("Brushes.Panel"))
+								// Navigation Bar
+								+ SVerticalBox::Slot()
+								.AutoHeight()
+								.HAlign(HAlign_Fill)
+								.Padding(FMargin(0, PanelInsetPadding, 0, 0))
 								[
-									CreateNavigationBar(Config)
+									SNew(SBorder)
+									.Padding(FMargin(3))
+									.BorderImage(bIsDrawer ? FStyleDefaults::GetNoBrush() : FAppStyle::Get().GetBrush("Brushes.Panel"))
+									[
+										CreateNavigationBar(Config)
+									]
 								]
 							]
 						]
@@ -717,236 +722,216 @@ void SContentBrowser::Construct( const FArguments& InArgs, const FName& InInstan
 				]
 			]
 		];
-
-	ChildSlot
-	[
-		SNew(SOverlay)
-
-		+ SOverlay::Slot()
-		.HAlign(HAlign_Fill)
-		.VAlign(VAlign_Fill)
-		.Padding(0)
-		[
-			SNew(SColorBlock)
-			.Color(FStyleColors::Background.GetSpecifiedColor()) // Always apply the Background color, rather than relying that set by the Tab which is context sensitive
-		]
-
-		+ SOverlay::Slot()
-		.HAlign(HAlign_Fill)
-		.VAlign(VAlign_Fill)
-		.Padding(0)
-		[
-			PanelContentWidget
-		]
-	];
-#else
-	ChildSlot
-	[
-		SNew(SVerticalBox)
-		+SVerticalBox::Slot()
-		.AutoHeight()
-		.Padding( 0, 0, 0, 0 )
+	}
+	else
+	{
+		ChildSlot
 		[
-			SNew( SBorder )
-			.Padding( FMargin( 3 ) )
-			.BorderImage(bIsDrawer ? FStyleDefaults::GetNoBrush() : FAppStyle::Get().GetBrush("Brushes.Panel"))
+			SNew(SVerticalBox)
+			+SVerticalBox::Slot()
+			.AutoHeight()
+			.Padding( 0, 0, 0, 0 )
 			[
-				SNew(SHorizontalBox)
-				+ SHorizontalBox::Slot()
-				.AutoWidth()
-				.HAlign(HAlign_Left)
-				.Padding(5, 0, 0, 0)
-				[
-					CreateToolBar(Config)
-				]
-				// History Back Button
-				+SHorizontalBox::Slot()
-				.Padding(10, 0, 0, 0)
-				.AutoWidth()
+				SNew( SBorder )
+				.Padding( FMargin( 3 ) )
+				.BorderImage(bIsDrawer ? FStyleDefaults::GetNoBrush() : FAppStyle::Get().GetBrush("Brushes.Panel"))
 				[
-					SNew(SButton)
-					.VAlign(EVerticalAlignment::VAlign_Center)
-					.ButtonStyle(FAppStyle::Get(), "SimpleButton")
-					.ToolTipText( this, &SContentBrowser::GetHistoryBackTooltip )
-					.ContentPadding( FMargin(1, 0) )
-					.OnClicked(this, &SContentBrowser::BackClicked)
-					.IsEnabled(this, &SContentBrowser::IsBackEnabled)
-					.AddMetaData<FTagMetaData>(FTagMetaData(TEXT("ContentBrowserHistoryBack")))
+					SNew(SHorizontalBox)
+					+ SHorizontalBox::Slot()
+					.AutoWidth()
+					.HAlign(HAlign_Left)
+					.Padding(5, 0, 0, 0)
 					[
-						SNew(SImage)
-						.Image(FAppStyle::Get().GetBrush("Icons.CircleArrowLeft"))
-						.ColorAndOpacity(FSlateColor::UseForeground())
+						CreateToolBar(Config)
 					]
-				]
-				// History Forward Button
-				+ SHorizontalBox::Slot()
-				.Padding(2, 0, 0, 0)
-				.AutoWidth()
-				[
-					SNew(SButton)
-					.VAlign(EVerticalAlignment::VAlign_Center)
-					.ButtonStyle(FAppStyle::Get(), "SimpleButton")
-					.ToolTipText( this, &SContentBrowser::GetHistoryForwardTooltip )
-					.ContentPadding( FMargin(1, 0) )
-					.OnClicked(this, &SContentBrowser::ForwardClicked)
-					.IsEnabled(this, &SContentBrowser::IsForwardEnabled)
-					.AddMetaData<FTagMetaData>(FTagMetaData(TEXT("ContentBrowserHistoryForward")))
+					// History Back Button
+					+SHorizontalBox::Slot()
+					.Padding(10, 0, 0, 0)
+					.AutoWidth()
 					[
-						SNew(SImage)
-						.Image(FAppStyle::Get().GetBrush("Icons.CircleArrowRight"))
-						.ColorAndOpacity(FSlateColor::UseForeground())
+						SNew(SButton)
+						.VAlign(EVerticalAlignment::VAlign_Center)
+						.ButtonStyle(FAppStyle::Get(), "SimpleButton")
+						.ToolTipText( this, &SContentBrowser::GetHistoryBackTooltip )
+						.ContentPadding( FMargin(1, 0) )
+						.OnClicked(this, &SContentBrowser::BackClicked)
+						.IsEnabled(this, &SContentBrowser::IsBackEnabled)
+						.AddMetaData<FTagMetaData>(FTagMetaData(TEXT("ContentBrowserHistoryBack")))
+						[
+							SNew(SImage)
+							.Image(FAppStyle::Get().GetBrush("Icons.CircleArrowLeft"))
+							.ColorAndOpacity(FSlateColor::UseForeground())
+						]
+					]
+					// History Forward Button
+					+ SHorizontalBox::Slot()
+					.Padding(2, 0, 0, 0)
+					.AutoWidth()
+					[
+						SNew(SButton)
+						.VAlign(EVerticalAlignment::VAlign_Center)
+						.ButtonStyle(FAppStyle::Get(), "SimpleButton")
+						.ToolTipText( this, &SContentBrowser::GetHistoryForwardTooltip )
+						.ContentPadding( FMargin(1, 0) )
+						.OnClicked(this, &SContentBrowser::ForwardClicked)
+						.IsEnabled(this, &SContentBrowser::IsForwardEnabled)
+						.AddMetaData<FTagMetaData>(FTagMetaData(TEXT("ContentBrowserHistoryForward")))
+						[
+							SNew(SImage)
+							.Image(FAppStyle::Get().GetBrush("Icons.CircleArrowRight"))
+							.ColorAndOpacity(FSlateColor::UseForeground())
+						]
 					]
-				]
 
-				// Path
-				+ SHorizontalBox::Slot()
-				.VAlign(VAlign_Center)
-				.HAlign(HAlign_Fill)
-				.FillWidth(1.0f)
-				.Padding(2, 0, 0, 0)
-				[
-					SAssignNew(NavigationBar, SNavigationBar)
-					.OnPathClicked(this, &SContentBrowser::OnPathClicked)
-					.GetPathMenuContent(this, &SContentBrowser::OnGetCrumbDelimiterContent)
-					.GetComboOptions(this, &SContentBrowser::GetRecentPaths)
-					.OnNavigateToPath(this, &SContentBrowser::OnNavigateToPath)
-					.OnCompletePrefix(this, &SContentBrowser::OnCompletePathPrefix)
-					.OnCanEditPathAsText(this, &SContentBrowser::OnCanEditPathAsText)
-					.AddMetaData<FTagMetaData>(FTagMetaData(TEXT("ContentBrowserPath")))
-				]
+					// Path
+					+ SHorizontalBox::Slot()
+					.VAlign(VAlign_Center)
+					.HAlign(HAlign_Fill)
+					.FillWidth(1.0f)
+					.Padding(2, 0, 0, 0)
+					[
+						SAssignNew(NavigationBar, SNavigationBar)
+						.OnPathClicked(this, &SContentBrowser::OnPathClicked)
+						.GetPathMenuContent(this, &SContentBrowser::OnGetCrumbDelimiterContent)
+						.GetComboOptions(this, &SContentBrowser::GetRecentPaths)
+						.OnNavigateToPath(this, &SContentBrowser::OnNavigateToPath)
+						.OnCompletePrefix(this, &SContentBrowser::OnCompletePathPrefix)
+						.OnCanEditPathAsText(this, &SContentBrowser::OnCanEditPathAsText)
+						.AddMetaData<FTagMetaData>(FTagMetaData(TEXT("ContentBrowserPath")))
+					]
 
-				+ SHorizontalBox::Slot()
-				.AutoWidth()
-				.HAlign(HAlign_Right)
-				.VAlign(VAlign_Center)
-				[
-					CreateLockButton(Config)
-				]
-				+ SHorizontalBox::Slot()
-				.AutoWidth()
-				.HAlign(HAlign_Right)
-				.VAlign(VAlign_Center)
-				[
-					CreateDrawerDockButton(Config)
-				]
-				+ SHorizontalBox::Slot()
-				.AutoWidth()
-				.Padding(5.0f, 0.0f, 0.0f, 0.0f)
-				.HAlign(HAlign_Right)
-				.VAlign(VAlign_Center)
-				[
-					ViewOptions
+					+ SHorizontalBox::Slot()
+					.AutoWidth()
+					.HAlign(HAlign_Right)
+					.VAlign(VAlign_Center)
+					[
+						CreateLockButton(Config)
+					]
+					+ SHorizontalBox::Slot()
+					.AutoWidth()
+					.HAlign(HAlign_Right)
+					.VAlign(VAlign_Center)
+					[
+						CreateDrawerDockButton(Config)
+					]
+					+ SHorizontalBox::Slot()
+					.AutoWidth()
+					.Padding(5.0f, 0.0f, 0.0f, 0.0f)
+					.HAlign(HAlign_Right)
+					.VAlign(VAlign_Center)
+					[
+						ViewOptions
+					]
 				]
 			]
-		]
 
-		+SVerticalBox::Slot()
-		.AutoHeight()
-		[
-			SNew(SSeparator)
-			.Thickness(2.0f)
-		]
+			+SVerticalBox::Slot()
+			.AutoHeight()
+			[
+				SNew(SSeparator)
+				.Thickness(2.0f)
+			]
 
-		// Assets/tree
-		+ SVerticalBox::Slot()
-		.FillHeight(1.0f)
-		.Padding(0.0f)
-		[
-			// The tree/assets splitter
-			SAssignNew(PathAssetSplitterPtr, SSplitter)
-			.PhysicalSplitterHandleSize(2.0f)
-
-			// Sources View
-			+ SSplitter::Slot()
-			.Resizable(true)
-			.SizeRule(SSplitter::SizeToContent)
-			.OnSlotResized(this, &SContentBrowser::OnPathViewBoxColumnResized)
+			// Assets/tree
+			+ SVerticalBox::Slot()
+			.FillHeight(1.0f)
+			.Padding(0.0f)
 			[
-				SNew(SBox)
-				.Padding(FMargin(4.f))
-				.Visibility(this, &SContentBrowser::GetSourcesViewVisibility)
-				.WidthOverride(this, &SContentBrowser::GetPathViewBoxWidthOverride)
+				// The tree/assets splitter
+				SAssignNew(PathAssetSplitterPtr, SSplitter)
+				.PhysicalSplitterHandleSize(2.0f)
+
+				// Sources View
+				+ SSplitter::Slot()
+				.Resizable(true)
+				.SizeRule(SSplitter::SizeToContent)
+				.OnSlotResized(this, &SContentBrowser::OnPathViewBoxColumnResized)
 				[
-					SNew(SBorder)
-					.Padding(FMargin(0))
-					.BorderImage(FAppStyle::GetBrush("Brushes.Recessed"))
-
+					SNew(SBox)
+					.Padding(FMargin(4.f))
+					.Visibility(this, &SContentBrowser::GetSourcesViewVisibility)
+					.WidthOverride(this, &SContentBrowser::GetPathViewBoxWidthOverride)
 					[
-						// Note: If adding more widgets here, fix ContentBrowserSourcesWidgetSwitcherIndex and the code that uses it!
-						SAssignNew(SourcesWidgetSwitcher, SWidgetSwitcher)
+						SNew(SBorder)
+						.Padding(FMargin(0))
+						.BorderImage(FAppStyle::GetBrush("Brushes.Recessed"))
 
-						// Paths View
-						+SWidgetSwitcher::Slot()
 						[
-							SAssignNew(PathFavoriteSplitterPtr, SSplitter)
-							.Clipping(EWidgetClipping::ClipToBounds)
-							.PhysicalSplitterHandleSize(2.0f)
-							.HitDetectionSplitterHandleSize(8.0f)
-							.Orientation(EOrientation::Orient_Vertical)
-							.MinimumSlotHeight(26.0f)
-							.Visibility( this, &SContentBrowser::GetSourcesViewVisibility )
-							+SSplitter::Slot()
-							.SizeRule(TAttribute<SSplitter::ESizeRule>(this, &SContentBrowser::GetFavoritesAreaSizeRule))
-							.MinSize(TAttribute<float>(this, &SContentBrowser::GetFavoritesAreaMinSize))
-							.Value(0.2f)
+							// Note: If adding more widgets here, fix ContentBrowserSourcesWidgetSwitcherIndex and the code that uses it!
+							SAssignNew(SourcesWidgetSwitcher, SWidgetSwitcher)
+
+							// Paths View
+							+SWidgetSwitcher::Slot()
 							[
-								SNew(SBorder)
-								.BorderImage(FAppStyle::Get().GetBrush("Brushes.Header"))
-								.Padding(0.0f, 2.0f, 0.0f, 0.0f)
+								SAssignNew(PathFavoriteSplitterPtr, SSplitter)
+								.Clipping(EWidgetClipping::ClipToBounds)
+								.PhysicalSplitterHandleSize(2.0f)
+								.HitDetectionSplitterHandleSize(8.0f)
+								.Orientation(EOrientation::Orient_Vertical)
+								.MinimumSlotHeight(26.0f)
+								.Visibility( this, &SContentBrowser::GetSourcesViewVisibility )
+								+SSplitter::Slot()
+								.SizeRule(TAttribute<SSplitter::ESizeRule>(this, &SContentBrowser::GetFavoritesAreaSizeRule))
+								.MinSize(TAttribute<float>(this, &SContentBrowser::GetFavoritesAreaMinSize))
+								.Value(0.2f)
 								[
-									CreateFavoritesView(Config)
+									SNew(SBorder)
+									.BorderImage(FAppStyle::Get().GetBrush("Brushes.Header"))
+									.Padding(0.0f, 2.0f, 0.0f, 0.0f)
+									[
+										CreateFavoritesView(Config)
+									]
 								]
-							]
 
-							+SSplitter::Slot()
-							.SizeRule(TAttribute<SSplitter::ESizeRule>(this, &SContentBrowser::GetPathAreaSizeRule))
-							.MinSize(29.0f)
-							.Value(0.8f)
-							[
-								SNew(SBorder)
-								.BorderImage(FAppStyle::Get().GetBrush("Brushes.Header"))
-								.Padding(0.0f, 2.0f, 0.0f, 0.0f)
+								+SSplitter::Slot()
+								.SizeRule(TAttribute<SSplitter::ESizeRule>(this, &SContentBrowser::GetPathAreaSizeRule))
+								.MinSize(29.0f)
+								.Value(0.8f)
 								[
-									CreatePathView(Config)
+									SNew(SBorder)
+									.BorderImage(FAppStyle::Get().GetBrush("Brushes.Header"))
+									.Padding(0.0f, 2.0f, 0.0f, 0.0f)
+									[
+										CreatePathView(Config)
+									]
 								]
-							]
 
-							+SSplitter::Slot()
-							.SizeRule(TAttribute<SSplitter::ESizeRule>(this, &SContentBrowser::GetCollectionsAreaSizeRule))
-							.MinSize(TAttribute<float>(this, &SContentBrowser::GetCollectionsAreaMinSize))
-							.Value(0.4f)
-							[
-								SNew(SBorder)
-								.BorderImage(FAppStyle::Get().GetBrush("Brushes.Header"))
-								.Padding(0.0f, 2.0f, 0.0f, 0.0f)
+								+SSplitter::Slot()
+								.SizeRule(TAttribute<SSplitter::ESizeRule>(this, &SContentBrowser::GetCollectionsAreaSizeRule))
+								.MinSize(TAttribute<float>(this, &SContentBrowser::GetCollectionsAreaMinSize))
+								.Value(0.4f)
 								[
-									CreateDockedCollectionsView(Config)
+									SNew(SBorder)
+									.BorderImage(FAppStyle::Get().GetBrush("Brushes.Header"))
+									.Padding(0.0f, 2.0f, 0.0f, 0.0f)
+									[
+										CreateDockedCollectionsView(Config)
+									]
 								]
 							]
-						]
 
-						// Collections View
-						+SWidgetSwitcher::Slot()
-						[
-							SNew(SBox)
-							.Padding(FMargin(0.0f, 4.0f, 0.0f, 0.0f))
+							// Collections View
+							+SWidgetSwitcher::Slot()
 							[
-								CollectionViewPtr.ToSharedRef()
+								SNew(SBox)
+								.Padding(FMargin(0.0f, 4.0f, 0.0f, 0.0f))
+								[
+									CollectionViewPtr.ToSharedRef()
+								]
 							]
 						]
 					]
 				]
-			]
 
-			// Asset View
-			+ SSplitter::Slot()
-			.Value(0.75f)
-			[
-				CreateAssetView(Config)
+				// Asset View
+				+ SSplitter::Slot()
+				.Value(0.75f)
+				[
+					CreateAssetView(Config)
+				]
 			]
-		]
-	];
-#endif
+		];
+	}
 
 	ExtendViewOptionsMenu(Config);
 
@@ -1203,7 +1188,6 @@ TSharedRef<SWidget> SContentBrowser::CreateToolBar(const FContentBrowserConfig*
 	return UToolMenus::Get()->GenerateWidget("ContentBrowser.ToolBar", MenuContext);
 }
 
-#if UE_CONTENTBROWSER_NEW_STYLE
 TSharedRef<SWidget> SContentBrowser::CreateNavigationBar(const FContentBrowserConfig* Config)
 {
 	FToolMenuContext MenuContext;
@@ -1217,7 +1201,6 @@ TSharedRef<SWidget> SContentBrowser::CreateNavigationBar(const FContentBrowserCo
 
 	return UToolMenus::Get()->GenerateWidget("ContentBrowser.NavigationBar", MenuContext);
 }
-#endif
 
 TSharedRef<SWidget> SContentBrowser::CreateLockButton(const FContentBrowserConfig* Config)
 {
@@ -1242,241 +1225,244 @@ TSharedRef<SWidget> SContentBrowser::CreateLockButton(const FContentBrowserConfi
 
 void SContentBrowser::OnFilterBarLayoutChanging(EFilterBarLayout NewLayout)
 {
-#if UE_CONTENTBROWSER_NEW_STYLE
-	if (NewLayout == EFilterBarLayout::Horizontal)
+	if (GEnableContentBrowserNewStyle)
 	{
-		AssetViewBorder->SetContent(
+		if (NewLayout == EFilterBarLayout::Horizontal)
+		{
+			AssetViewBorder->SetContent(
 
-			SNew(SVerticalBox)
+				SNew(SVerticalBox)
 
-			+ SVerticalBox::Slot()
-			.Padding(0.0f)
-			.AutoHeight()
-			[
-				FilterListPtr.ToSharedRef()
-			]
+				+ SVerticalBox::Slot()
+				.Padding(0.0f)
+				.AutoHeight()
+				[
+					FilterListPtr.ToSharedRef()
+				]
 
-			+ SVerticalBox::Slot()
-			.FillHeight(1.0f)
-			.Padding(0, 0)
-			[
-				AssetViewPtr.ToSharedRef()
-			]
-		);
-	}
-	else
-	{
-		AssetViewBorder->SetContent(
+				+ SVerticalBox::Slot()
+				.FillHeight(1.0f)
+				.Padding(0, 0)
+				[
+					AssetViewPtr.ToSharedRef()
+				]
+			);
+		}
+		else
+		{
+			AssetViewBorder->SetContent(
 
-			SNew(SSplitter)
-			.PhysicalSplitterHandleSize(2.0f)
+				SNew(SSplitter)
+				.PhysicalSplitterHandleSize(2.0f)
 
-			/* Filters in an SScrollBox */
-			+ SSplitter::Slot()
-			.MinSize(95.f)
-			.Resizable(true)
-			.SizeRule(SSplitter::SizeToContent)
-			.OnSlotResized(this, &SContentBrowser::OnFilterBoxColumnResized)
-			[
-				SNew(SBox)
-				.WidthOverride(this, &SContentBrowser::GetFilterViewBoxWidthOverride)
-				// Don't take up space when there are no filters
-				.Visibility_Lambda([this]
-				{
-					return FilterListPtr->HasAnyFilters() ? EVisibility::Visible : EVisibility::Collapsed;
-				})
+				/* Filters in an SScrollBox */
+				+ SSplitter::Slot()
+				.MinSize(95.f)
+				.Resizable(true)
+				.SizeRule(SSplitter::SizeToContent)
+				.OnSlotResized(this, &SContentBrowser::OnFilterBoxColumnResized)
 				[
-					SNew(SVerticalBox)
-
-					// Header
-					+ SVerticalBox::Slot()
-					.Padding(0.0f, 2.0f)
-					.AutoHeight()
+					SNew(SBox)
+					.WidthOverride(this, &SContentBrowser::GetFilterViewBoxWidthOverride)
+					// Don't take up space when there are no filters
+					.Visibility_Lambda([this]
+					{
+						return FilterListPtr->HasAnyFilters() ? EVisibility::Visible : EVisibility::Collapsed;
+					})
 					[
+						SNew(SVerticalBox)
 
-						SNew(SBorder)
-						.BorderImage(FAppStyle::Get().GetBrush("Brushes.Header"))
-						.Padding(FMargin(8.0f, 6.0f))
-						.Content()
+						// Header
+						+ SVerticalBox::Slot()
+						.Padding(0.0f, 2.0f)
+						.AutoHeight()
 						[
-							SNew(STextBlock)
-							.Text(LOCTEXT("FilterListVerticalHeader", "Filters"))
-							.TextStyle(FAppStyle::Get(), "ButtonText")
-							.Font(FAppStyle::Get().GetFontStyle("NormalFontBold"))
+
+							SNew(SBorder)
+							.BorderImage(FAppStyle::Get().GetBrush("Brushes.Header"))
+							.Padding(FMargin(8.0f, 6.0f))
+							.Content()
+							[
+								SNew(STextBlock)
+								.Text(LOCTEXT("FilterListVerticalHeader", "Filters"))
+								.TextStyle(FAppStyle::Get(), "ButtonText")
+								.Font(FAppStyle::Get().GetFontStyle("NormalFontBold"))
+							]
+						]
+
+						// Filter List
+						+ SVerticalBox::Slot()
+						.Padding(0.0f)
+						.FillHeight(1.0f)
+						[
+							FilterListPtr.ToSharedRef()
 						]
 					]
+				]
+
+				+ SSplitter::Slot()
+				.Value(0.88f)
+				[
+					SNew(SVerticalBox)
 
-					// Filter List
 					+ SVerticalBox::Slot()
-					.Padding(0.0f)
 					.FillHeight(1.0f)
+					.Padding(0, 0)
 					[
-						FilterListPtr.ToSharedRef()
+						AssetViewPtr.ToSharedRef()
 					]
 				]
-			]
+			);
+		}
+	}
+	else
+	{
+		const FOptionalSize SearchBoxDesiredWidth = 500.0f;
+		constexpr float SearchBoxMaxWidth = 0.0f;
 
-			+ SSplitter::Slot()
-			.Value(0.88f)
-			[
-				SNew(SVerticalBox)
+		if (NewLayout == EFilterBarLayout::Horizontal)
+		{
+			TSharedPtr<SHorizontalBox> SearchBoxSlot =
+				SNew(SHorizontalBox)
 
-				+ SVerticalBox::Slot()
-				.FillHeight(1.0f)
-				.Padding(0, 0)
+				+ SHorizontalBox::Slot()
+				.AutoWidth()
+				.VAlign(VAlign_Top)
+				.Padding(6, 4, 2, 0)
 				[
-					AssetViewPtr.ToSharedRef()
+					FilterComboButton.ToSharedRef()
 				]
-			]
-		);
-	}
-#else
-	const FOptionalSize SearchBoxDesiredWidth = 500.0f;
-	constexpr float SearchBoxMaxWidth = 0.0f;
 
-	if (NewLayout == EFilterBarLayout::Horizontal)
-	{
-		TSharedPtr<SHorizontalBox> SearchBoxSlot =
-			SNew(SHorizontalBox)
+				+ SHorizontalBox::Slot()
+				.FillWidth(1.0)
+				.VAlign(VAlign_Top)
+				.Padding(0, 4, 0, 0)
+				[
+					SNew(SBox)
+					.VAlign(VAlign_Center)
+					.WidthOverride(SearchBoxDesiredWidth)
+					[
+						SearchBoxPtr.ToSharedRef()
+					]
+				];
 
-			+ SHorizontalBox::Slot()
-			.AutoWidth()
-			.VAlign(VAlign_Top)
-			.Padding(6, 4, 2, 0)
-			[
-				FilterComboButton.ToSharedRef()
-			]
+			/** We add the Combo Button and the Search Box to the FilterList itself, so that the filters wrap with them
+			 *	properly in the Horizontal Layout
+			 */
+			FilterListPtr->AddWidgetToCurrentLayout(SearchBoxSlot.ToSharedRef());
 
-			+ SHorizontalBox::Slot()
-			.FillWidth(1.0)
-			.VAlign(VAlign_Top)
-			.Padding(0, 4, 0, 0)
-			[
-				SNew(SBox)
-				.VAlign(VAlign_Center)
-				.WidthOverride(SearchBoxDesiredWidth)
+			AssetViewBorder->SetContent(
+
+				SNew(SVerticalBox)
+
+				+ SVerticalBox::Slot()
+				.Padding(0.0f)
+				.AutoHeight()
 				[
-					SearchBoxPtr.ToSharedRef()
+					FilterListPtr.ToSharedRef()
 				]
-			];
 
-		/** We add the Combo Button and the Search Box to the FilterList itself, so that the filters wrap with them
-		 *	properly in the Horizontal Layout
-		 */
-		FilterListPtr->AddWidgetToCurrentLayout(SearchBoxSlot.ToSharedRef());
+				+ SVerticalBox::Slot()
+				.FillHeight(1.0f)
+				.Padding(0, 0)
+				[
+					AssetViewPtr.ToSharedRef()
+				]
+			);
+		}
+		else
+		{
+			AssetViewBorder->SetContent(
 
-		AssetViewBorder->SetContent(
+				SNew(SSplitter)
+				.PhysicalSplitterHandleSize(2.0f)
 
-			SNew(SVerticalBox)
+				/* Filters in an SScrollBox */
+				+ SSplitter::Slot()
+				.MinSize(95.f)
+				.Resizable(true)
+				.SizeRule(SSplitter::SizeToContent)
+				.OnSlotResized(this, &SContentBrowser::OnFilterBoxColumnResized)
+				[
+					SNew(SBox)
+					.WidthOverride(this, &SContentBrowser::GetFilterViewBoxWidthOverride)
+					// Don't take up space when there are no filters
+					.Visibility_Lambda([this]
+					{
+						return FilterListPtr->HasAnyFilters() ? EVisibility::Visible : EVisibility::Collapsed;
+					})
+					[
+						SNew(SVerticalBox)
 
-			+ SVerticalBox::Slot()
-			.Padding(0.0f)
-			.AutoHeight()
-			[
-				FilterListPtr.ToSharedRef()
-			]
+						// Header
+						+ SVerticalBox::Slot()
+						.Padding(0.0f, 2.0f)
+						.AutoHeight()
+						[
 
-			+ SVerticalBox::Slot()
-			.FillHeight(1.0f)
-			.Padding(0, 0)
-			[
-				AssetViewPtr.ToSharedRef()
-			]
-		);
-	}
-	else
-	{
-		AssetViewBorder->SetContent(
+							SNew(SBorder)
+							.BorderImage(FAppStyle::Get().GetBrush("Brushes.Header"))
+							.Padding(FMargin(8.0f, 6.0f))
+							.Content()
+							[
+								SNew(STextBlock)
+								.Text(LOCTEXT("FilterListVerticalHeader", "Filters"))
+								.TextStyle(FAppStyle::Get(), "ButtonText")
+								.Font(FAppStyle::Get().GetFontStyle("NormalFontBold"))
+							]
+						]
 
-			SNew(SSplitter)
-			.PhysicalSplitterHandleSize(2.0f)
+						// Filter List
+						+ SVerticalBox::Slot()
+						.Padding(0.0f)
+						.FillHeight(1.0f)
+						[
+							FilterListPtr.ToSharedRef()
+						]
+					]
+				]
 
-			/* Filters in an SScrollBox */
-			+ SSplitter::Slot()
-			.MinSize(95.f)
-			.Resizable(true)
-			.SizeRule(SSplitter::SizeToContent)
-			.OnSlotResized(this, &SContentBrowser::OnFilterBoxColumnResized)
-			[
-				SNew(SBox)
-				.WidthOverride(this, &SContentBrowser::GetFilterViewBoxWidthOverride)
-				// Don't take up space when there are no filters
-				.Visibility_Lambda([this]
-				{
-					return FilterListPtr->HasAnyFilters() ? EVisibility::Visible : EVisibility::Collapsed;
-				})
+				+ SSplitter::Slot()
+				.Value(0.88f)
 				[
 					SNew(SVerticalBox)
 
-					// Header
 					+ SVerticalBox::Slot()
-					.Padding(0.0f, 2.0f)
+					.Padding(6, 4, 0, 0)
 					.AutoHeight()
 					[
+						SNew(SHorizontalBox)
 
-						SNew(SBorder)
-						.BorderImage(FAppStyle::Get().GetBrush("Brushes.Header"))
-						.Padding(FMargin(8.0f, 6.0f))
-						.Content()
+						+ SHorizontalBox::Slot()
+						.AutoWidth()
+						.Padding(0, 0, 2, 0)
 						[
-							SNew(STextBlock)
-							.Text(LOCTEXT("FilterListVerticalHeader", "Filters"))
-							.TextStyle(FAppStyle::Get(), "ButtonText")
-							.Font(FAppStyle::Get().GetFontStyle("NormalFontBold"))
+							FilterComboButton.ToSharedRef()
+						]
+
+						+ SHorizontalBox::Slot()
+						.AutoWidth()
+						.MaxWidth(SearchBoxMaxWidth)
+						[
+							SNew(SBox)
+							.VAlign(VAlign_Center)
+							.WidthOverride(SearchBoxDesiredWidth)
+							[
+								SearchBoxPtr.ToSharedRef()
+							]
 						]
 					]
 
-					// Filter List
 					+ SVerticalBox::Slot()
-					.Padding(0.0f)
 					.FillHeight(1.0f)
+					.Padding(0, 0)
 					[
-						FilterListPtr.ToSharedRef()
-					]
-				]
-			]
-
-			+ SSplitter::Slot()
-			.Value(0.88f)
-			[
-				SNew(SVerticalBox)
-
-				+ SVerticalBox::Slot()
-				.Padding(6, 4, 0, 0)
-				.AutoHeight()
-				[
-					SNew(SHorizontalBox)
-
-					+ SHorizontalBox::Slot()
-					.AutoWidth()
-					.Padding(0, 0, 2, 0)
-					[
-						FilterComboButton.ToSharedRef()
-					]
-
-					+ SHorizontalBox::Slot()
-					.AutoWidth()
-					.MaxWidth(SearchBoxMaxWidth)
-					[
-						SNew(SBox)
-						.VAlign(VAlign_Center)
-						.WidthOverride(SearchBoxDesiredWidth)
-						[
-							SearchBoxPtr.ToSharedRef()
-						]
+						AssetViewPtr.ToSharedRef()
 					]
 				]
-
-				+ SVerticalBox::Slot()
-				.FillHeight(1.0f)
-				.Padding(0, 0)
-				[
-					AssetViewPtr.ToSharedRef()
-				]
-			]
-		);
+			);
+		}
 	}
-#endif
 }
 
 TSharedRef<SWidget> SContentBrowser::CreateAssetView(const FContentBrowserConfig* Config)
@@ -1523,26 +1509,26 @@ TSharedRef<SWidget> SContentBrowser::CreateAssetView(const FContentBrowserConfig
 				   .Visibility((Config != nullptr ? Config->bCanShowAssetSearch : true) ? EVisibility::Visible : EVisibility::Collapsed)
 				   .AddMetaData<FTagMetaData>(FTagMetaData(TEXT("ContentBrowserSearchAssets")));
 
-#if UE_CONTENTBROWSER_NEW_STYLE
-	NavigationBar = SNew(SNavigationBar)
-					.OnPathClicked(this, &SContentBrowser::OnPathClicked)
-					.GetPathMenuContent(this, &SContentBrowser::OnGetCrumbDelimiterContent)
-					.GetComboOptions(this, &SContentBrowser::GetRecentPaths)
-					.OnNavigateToPath(this, &SContentBrowser::OnNavigateToPath)
-					.OnCompletePrefix(this, &SContentBrowser::OnCompletePathPrefix)
-					.OnCanEditPathAsText(this, &SContentBrowser::OnCanEditPathAsText)
-					.AddMetaData<FTagMetaData>(FTagMetaData(TEXT("ContentBrowserPath")));
-#endif
+	if (GEnableContentBrowserNewStyle)
+	{
+		NavigationBar = SNew(SNavigationBar)
+						.OnPathClicked(this, &SContentBrowser::OnPathClicked)
+						.GetPathMenuContent(this, &SContentBrowser::OnGetCrumbDelimiterContent)
+						.GetComboOptions(this, &SContentBrowser::GetRecentPaths)
+						.OnNavigateToPath(this, &SContentBrowser::OnNavigateToPath)
+						.OnCompletePrefix(this, &SContentBrowser::OnCompletePathPrefix)
+						.OnCanEditPathAsText(this, &SContentBrowser::OnCanEditPathAsText)
+						.AddMetaData<FTagMetaData>(FTagMetaData(TEXT("ContentBrowserPath")));
+	}
+
 
 	/* Create the Border that the Asset View will live in, the actual layout is populated in OnFilterBarLayoutChanging,
 	 * which is called initially called through SAssetFilterBar::LoadSettings through SContentBrowser::LoadSettings()
 	 */
-
-#if UE_CONTENTBROWSER_NEW_STYLE
-	static const FMargin AssetViewPadding(2.0f, 0.0f, 2.0f, 0.0f);
-#else
-	static const FMargin AssetViewPadding(2.0f, 2.0f, 2.0f, 0.0f);
-#endif
+	const FMargin AssetViewPadding =
+		GEnableContentBrowserNewStyle
+			? FMargin(2.0f, 0.0f, 2.0f, 0.0f)
+			: FMargin(2.0f, 2.0f, 2.0f, 0.0f);
 
 	AssetViewBorder = SNew(SBorder)
 			.BorderImage(FAppStyle::Get().GetBrush("Brushes.Panel"))
@@ -1847,47 +1833,52 @@ void SContentBrowser::ExtendViewOptionsMenu(const FContentBrowserConfig* Config)
 				{
 					if (bShowLockButton)
 					{
-#if UE_CONTENTBROWSER_NEW_STYLE
-						FToolMenuSection& Section = InMenu->FindOrAddSection("Manage");
-						Section.AddMenuEntry(
-							"ToggleLock",
-							TAttribute<FText>(ContentBrowser.ToSharedRef(), &SContentBrowser::GetLockMenuText),
-							LOCTEXT("LockToggleTooltip", "Toggle lock. If locked, this browser will ignore Find in Content Browser requests."),
-							TAttribute<FSlateIcon>(ContentBrowser.ToSharedRef(), &SContentBrowser::GetLockIcon),
-							FUIAction(
-								FExecuteAction::CreateLambda([ContentBrowser = Context->OwningContentBrowser]() { ContentBrowser.Pin()->ToggleLockClicked(); })
-							)
-						);
-#else
-						FToolMenuSection& Section = InMenu->AddSection("Locking", LOCTEXT("LockingMenuHeader", "Locking"), FToolMenuInsert("AssetViewType", EToolMenuInsertType::After));
-						Section.AddMenuEntry(
-							"ToggleLock",
-							TAttribute<FText>(ContentBrowser.ToSharedRef(), &SContentBrowser::GetLockMenuText),
-							LOCTEXT("LockToggleTooltip", "Toggle lock. If locked, this browser will ignore Find in Content Browser requests."),
-							TAttribute<FSlateIcon>(ContentBrowser.ToSharedRef(), &SContentBrowser::GetLockIcon),
-							FUIAction(
-								FExecuteAction::CreateLambda([ContentBrowser = Context->OwningContentBrowser]() { ContentBrowser.Pin()->ToggleLockClicked(); })
-							)
-						);
-#endif
+						if (GEnableContentBrowserNewStyle)
+						{
+							FToolMenuSection& Section = InMenu->FindOrAddSection("Manage");
+							Section.AddMenuEntry(
+								"ToggleLock",
+								TAttribute<FText>(ContentBrowser.ToSharedRef(), &SContentBrowser::GetLockMenuText),
+								LOCTEXT("LockToggleTooltip", "Toggle lock. If locked, this browser will ignore Find in Content Browser requests."),
+								TAttribute<FSlateIcon>(ContentBrowser.ToSharedRef(), &SContentBrowser::GetLockIcon),
+								FUIAction(
+									FExecuteAction::CreateLambda([ContentBrowser = Context->OwningContentBrowser]() { ContentBrowser.Pin()->ToggleLockClicked(); })
+								)
+							);
+						}
+						else
+						{
+							FToolMenuSection& Section = InMenu->AddSection("Locking", LOCTEXT("LockingMenuHeader", "Locking"), FToolMenuInsert("AssetViewType", EToolMenuInsertType::After));
+							Section.AddMenuEntry(
+								"ToggleLock",
+								TAttribute<FText>(ContentBrowser.ToSharedRef(), &SContentBrowser::GetLockMenuText),
+								LOCTEXT("LockToggleTooltip", "Toggle lock. If locked, this browser will ignore Find in Content Browser requests."),
+								TAttribute<FSlateIcon>(ContentBrowser.ToSharedRef(), &SContentBrowser::GetLockIcon),
+								FUIAction(
+									FExecuteAction::CreateLambda([ContentBrowser = Context->OwningContentBrowser]() { ContentBrowser.Pin()->ToggleLockClicked(); })
+								)
+							);
+						}
 					}
-#if !UE_CONTENTBROWSER_NEW_STYLE
-					if (bShowSourcesView)
+
+					if (!GEnableContentBrowserNewStyle)
 					{
-						FToolMenuSection& Section = InMenu->FindOrAddSection("View");
-						Section.AddMenuEntry(
-							"ToggleSources",
-							LOCTEXT("ToggleSourcesView", "Show Sources Panel"),
-							LOCTEXT("ToggleSourcesView_Tooltip", "Show or hide the sources panel"),
-							TAttribute<FSlateIcon>(),
-							FUIAction(
-								FExecuteAction::CreateLambda([ContentBrowser = Context->OwningContentBrowser]() { ContentBrowser.Pin()->SourcesViewExpandClicked(); }),
-								FCanExecuteAction(),
-								FIsActionChecked::CreateLambda([ContentBrowser = Context->OwningContentBrowser]() { return ContentBrowser.Pin()->bSourcesViewExpanded; })),
-							EUserInterfaceActionType::ToggleButton
-						);
+						if (bShowSourcesView)
+						{
+							FToolMenuSection& Section = InMenu->FindOrAddSection("View");
+							Section.AddMenuEntry(
+								"ToggleSources",
+								LOCTEXT("ToggleSourcesView", "Show Sources Panel"),
+								LOCTEXT("ToggleSourcesView_Tooltip", "Show or hide the sources panel"),
+								TAttribute<FSlateIcon>(),
+								FUIAction(
+									FExecuteAction::CreateLambda([ContentBrowser = Context->OwningContentBrowser]() { ContentBrowser.Pin()->SourcesViewExpandClicked(); }),
+									FCanExecuteAction(),
+									FIsActionChecked::CreateLambda([ContentBrowser = Context->OwningContentBrowser]() { return ContentBrowser.Pin()->bSourcesViewExpanded; })),
+								EUserInterfaceActionType::ToggleButton
+							);
+						}
 					}
-#endif
 				}
 			}
 		}
diff --git a/Engine/Source/Editor/ContentBrowser/Private/SContentBrowser.h b/Engine/Source/Editor/ContentBrowser/Private/SContentBrowser.h
--- a/Engine/Source/Editor/ContentBrowser/Private/SContentBrowser.h
+++ b/Engine/Source/Editor/ContentBrowser/Private/SContentBrowser.h
@@ -604,9 +604,9 @@ private:
 
 	/** Creates various widgets for the content browser main view */
 	TSharedRef<SWidget> CreateToolBar(const FContentBrowserConfig* Config);
-#if UE_CONTENTBROWSER_NEW_STYLE
+
 	TSharedRef<SWidget> CreateNavigationBar(const FContentBrowserConfig* Config);
-#endif
+
 	TSharedRef<SWidget> CreateLockButton(const FContentBrowserConfig* Config);
 	TSharedRef<SWidget> CreateAssetView(const FContentBrowserConfig* Config);
 	TSharedRef<SWidget> CreateFavoritesView(const FContentBrowserConfig* Config);
@@ -800,10 +800,8 @@ private:
 	/** Registers the "Add New" drop-down menu */
 	static FDelayedAutoRegisterHelper AddNewContextMenuRegistration;
 
-#if UE_CONTENTBROWSER_NEW_STYLE
 	/** Registers the bottom navigation bar */
 	static FDelayedAutoRegisterHelper NavigationBarMenuRegistration;
-#endif
 
 	/** Registers the context menu for folders */
 	static FDelayedAutoRegisterHelper FolderContextMenuRegistration;
diff --git a/Engine/Source/Editor/ContentBrowser/Public/ContentBrowserModule.h b/Engine/Source/Editor/ContentBrowser/Public/ContentBrowserModule.h
--- a/Engine/Source/Editor/ContentBrowser/Public/ContentBrowserModule.h
+++ b/Engine/Source/Editor/ContentBrowser/Public/ContentBrowserModule.h
@@ -25,6 +25,10 @@ class IContentBrowserSingleton;
 struct FARFilter;
 struct FAssetData;
 
+namespace UE::Editor::ContentBrowser
+{
+	CONTENTBROWSER_API extern bool GEnableContentBrowserNewStyle;
+}
 
 /** Extra state generator that adds an icon and a corresponding legend entry on an asset. */
 class FAssetViewExtraStateGenerator
diff --git a/Engine/Source/Editor/ContentBrowser/Public/SAssetView.h b/Engine/Source/Editor/ContentBrowser/Public/SAssetView.h
--- a/Engine/Source/Editor/ContentBrowser/Public/SAssetView.h
+++ b/Engine/Source/Editor/ContentBrowser/Public/SAssetView.h
@@ -569,10 +569,8 @@ private:
 	/** Registers the menu used by the settings button */
 	static FDelayedAutoRegisterHelper AssetViewOptionsMenuRegistration;
 
-#if UE_CONTENTBROWSER_NEW_STYLE
 	/** Extend the bottom navigation bar */
 	static FDelayedAutoRegisterHelper NavigationBarMenuExtensionRegistration;
-#endif
 
 	/** Toggle whether folders should be shown or not */
 	void ToggleShowFolders();
diff --git a/Engine/Source/Editor/UMGEditor/Private/Library/SLibraryView.cpp b/Engine/Source/Editor/UMGEditor/Private/Library/SLibraryView.cpp
--- a/Engine/Source/Editor/UMGEditor/Private/Library/SLibraryView.cpp
+++ b/Engine/Source/Editor/UMGEditor/Private/Library/SLibraryView.cpp
@@ -35,6 +35,13 @@
 
 #define LOCTEXT_NAMESPACE "UMG"
 
+namespace UE::Editor::ContentBrowser
+{
+	extern CONTENTBROWSER_API bool GEnableContentBrowserNewStyle;
+}
+
+using namespace UE::Editor::ContentBrowser;
+
 FText SLibraryViewItem::GetFavoriteToggleToolTipText() const
 {
 	if (GetFavoritedState() == ECheckBoxState::Checked)
@@ -156,12 +163,14 @@ TSharedRef<SWidget> SLibraryView::ConstructViewOptions()
 		{
 			for (int32 EnumValue = (int32)EThumbnailSize::Tiny; EnumValue < (int32)EThumbnailSize::MAX; ++EnumValue)
 			{
-#if !UE_CONTENTBROWSER_NEW_STYLE
-				if ((EThumbnailSize)EnumValue == EThumbnailSize::XLarge)
+				if (GEnableContentBrowserNewStyle)
 				{
-					continue;
+					if ((EThumbnailSize)EnumValue == EThumbnailSize::XLarge)
+					{
+						continue;
+					}
 				}
-#endif
+
 				SubMenu.AddMenuEntry(
 					SAssetView::ThumbnailSizeToDisplayName((EThumbnailSize)EnumValue),
 					FText::GetEmpty(),
diff --git a/Engine/Source/Editor/UnrealEd/Private/AssetDefinitionDefault.cpp b/Engine/Source/Editor/UnrealEd/Private/AssetDefinitionDefault.cpp
--- a/Engine/Source/Editor/UnrealEd/Private/AssetDefinitionDefault.cpp
+++ b/Engine/Source/Editor/UnrealEd/Private/AssetDefinitionDefault.cpp
@@ -14,7 +14,6 @@
 
 #define LOCTEXT_NAMESPACE "AssetDefinitionDefault"
 
-#if UE_CONTENTBROWSER_NEW_STYLE
 namespace UE::AssetDefinitionDefault::Status
 {
 	EVisibility GetDirtyStatusVisibility(const TSharedPtr<IAssetStatusInfoProvider> InAssetStatusInfoProvider)
@@ -101,7 +100,6 @@ namespace UE::AssetDefinitionDefault::Status
 		return SourceControlDescription;
 	}
 }
-#endif
 
 EAssetCommandResult UAssetDefinitionDefault::OpenAssets(const FAssetOpenArgs& OpenArgs) const
 {
@@ -133,7 +131,6 @@ EAssetCommandResult UAssetDefinitionDefault::PerformAssetDiff(const FAssetDiffAr
 	return EAssetCommandResult::Handled;
 }
 
-#if UE_CONTENTBROWSER_NEW_STYLE
 void UAssetDefinitionDefault::GetAssetStatusInfo(const TSharedPtr<IAssetStatusInfoProvider>& InAssetStatusInfoProvider, TArray<FAssetDisplayInfo>& OutStatusInfo) const
 {
 	FAssetDisplayInfo DirtyStatus;
@@ -150,7 +147,6 @@ void UAssetDefinitionDefault::GetAssetStatusInfo(const TSharedPtr<IAssetStatusIn
 	SCCStatus.StatusDescription = TAttribute<FText>::Create(TAttribute<FText>::FGetter::CreateStatic(&UE::AssetDefinitionDefault::Status::GetSourceControlStatusDescription, InAssetStatusInfoProvider));
 	OutStatusInfo.Add(SCCStatus);
 }
-#endif
 
 namespace UE::Editor
 {
diff --git a/Engine/Source/Editor/UnrealEd/Private/AssetThumbnail.cpp b/Engine/Source/Editor/UnrealEd/Private/AssetThumbnail.cpp
--- a/Engine/Source/Editor/UnrealEd/Private/AssetThumbnail.cpp
+++ b/Engine/Source/Editor/UnrealEd/Private/AssetThumbnail.cpp
@@ -91,11 +91,9 @@ public:
 		, _AssetTypeColorOverride()
 		, _Padding(0)
 		, _GenericThumbnailSize(64)
-#if UE_CONTENTBROWSER_NEW_STYLE
 		, _AllowAssetStatusThumbnailOverlay(false)
 		, _ShowAssetChip(false)
 		, _AssetChipBorderImageOverride()
-#endif
 		, _ColorStripOrientation(EThumbnailColorStripOrientation::HorizontalBottomEdge)
 		{}
 
@@ -114,13 +112,11 @@ public:
 		SLATE_ARGUMENT(TOptional<FLinearColor>, AssetTypeColorOverride)
 		SLATE_ARGUMENT(FMargin, Padding)
 		SLATE_ATTRIBUTE(int32, GenericThumbnailSize)
-#if UE_CONTENTBROWSER_NEW_STYLE
 		SLATE_ARGUMENT(TSharedPtr<IAssetSystemInfoProvider>, AssetSystemInfoProvider)
 		SLATE_ARGUMENT(bool, AllowAssetStatusThumbnailOverlay)
 		SLATE_ATTRIBUTE(TSharedPtr<IToolTip>, TooltipOverride)
 		SLATE_ARGUMENT(bool, ShowAssetChip)
 		SLATE_ARGUMENT(TAttribute<const FSlateBrush*>, AssetChipBorderImageOverride)
-#endif
 		SLATE_ARGUMENT(EThumbnailColorStripOrientation, ColorStripOrientation)
 
 	SLATE_END_ARGS()
@@ -142,11 +138,9 @@ public:
 		GenericThumbnailBorderPadding = 2.f;
 		GenericThumbnailSize = InArgs._GenericThumbnailSize;
 		ColorStripOrientation = InArgs._ColorStripOrientation;
-#if UE_CONTENTBROWSER_NEW_STYLE
 		AssetSystemInfoProvider = InArgs._AssetSystemInfoProvider;
 		bShowAssetChip = InArgs._ShowAssetChip;
 		AssetChipBorderImageOverride = InArgs._AssetChipBorderImageOverride;
-#endif
 		AssetThumbnail->OnAssetDataChanged().AddSP(this, &SAssetThumbnail::OnAssetDataChanged);
 
 		const FAssetData& AssetData = AssetThumbnail->GetAssetData();
@@ -172,17 +166,21 @@ public:
 
 		TSharedRef<SOverlay> OverlayWidget = SNew(SOverlay);
 
-#if UE_CONTENTBROWSER_NEW_STYLE
-		if (InArgs._TooltipOverride.IsSet())
-		{
-			SetToolTip(InArgs._TooltipOverride);
-		}
-		else
+		static const IConsoleVariable* CVarEnableContentBrowserNewStyle = IConsoleManager::Get().FindConsoleVariable(TEXT("ContentBrowser.EnableNewStyle"));
+		const bool bEnableContentBrowserNewStyle = CVarEnableContentBrowserNewStyle && CVarEnableContentBrowserNewStyle->GetBool();
+
+		if (bEnableContentBrowserNewStyle)
 		{
-			// Set our tooltip - this will refresh each time it's opened to make sure it's up-to-date
-			SetToolTip(SNew(SAssetThumbnailToolTip).AssetThumbnail(SharedThis(this)));
+			if (InArgs._TooltipOverride.IsSet())
+			{
+				SetToolTip(InArgs._TooltipOverride);
+			}
+			else
+			{
+				// Set our tooltip - this will refresh each time it's opened to make sure it's up-to-date
+				SetToolTip(SNew(SAssetThumbnailToolTip).AssetThumbnail(SharedThis(this)));
+			}
 		}
-#endif
 
 		UpdateThumbnailClass(AssetTypeActions.Get());
 
@@ -302,124 +300,128 @@ public:
 				];
 		}
 
-#if UE_CONTENTBROWSER_NEW_STYLE
 		TSharedRef<SWidget> ContentWidget = OverlayWidget;
 
-		if (bShowAssetChip)
+		if (bEnableContentBrowserNewStyle)
 		{
-			// The asset color chip
-			LayeredImage = SNew(SLayeredImage)
-				.Image(FAppStyle::GetBrush(Style, ".AssetTileViewWhiteChipBorder"))
-				.ColorAndOpacity(FStyleColors::Secondary);
-			LayeredImage->AddLayer(FAppStyle::GetBrush(Style, ".AssetTileViewWhiteChip"), AssetColor);
+			if (bShowAssetChip)
+			{
+				// The asset color chip
+				LayeredImage = SNew(SLayeredImage)
+					.Image(FAppStyle::GetBrush(Style, ".AssetTileViewWhiteChipBorder"))
+					.ColorAndOpacity(FStyleColors::Secondary);
+				LayeredImage->AddLayer(FAppStyle::GetBrush(Style, ".AssetTileViewWhiteChip"), AssetColor);
+
+				OverlayWidget->AddSlot()
+					.HAlign(HAlign_Right)
+					.VAlign(VAlign_Top)
+					[
+						LayeredImage.ToSharedRef()
+					];
 
-			OverlayWidget->AddSlot()
-				.HAlign(HAlign_Right)
-				.VAlign(VAlign_Top)
+				const TSharedRef<SImage> OverlayBorder = SNew(SImage)
+					.Image(FAppStyle::GetBrush(Style, ".AssetBorder"))
+					.Visibility(EVisibility::HitTestInvisible);
+
+				OverlayWidget->AddSlot()
 				[
-					LayeredImage.ToSharedRef()
+					OverlayBorder
 				];
 
-			const TSharedRef<SImage> OverlayBorder = SNew(SImage)
-				.Image(FAppStyle::GetBrush(Style, ".AssetBorder"))
-				.Visibility(EVisibility::HitTestInvisible);
-			
+				if (AssetChipBorderImageOverride.IsSet())
+				{
+					const TSharedRef<SBorder> ThumbnailWidgetBorder = SNew(SBorder)
+						.Padding(0)
+						.BorderImage(FAppStyle::GetBrush(Style, ".AssetBorder"))
+						[
+							OverlayWidget
+						];
+
+					OverlayBorder->SetImage(AssetChipBorderImageOverride);
+					ThumbnailWidgetBorder->SetBorderImage(AssetChipBorderImageOverride);
+					ContentWidget = ThumbnailWidgetBorder;
+				}
+			}
+		}
+		else
+		{
+			// The asset color strip
 			OverlayWidget->AddSlot()
+			.HAlign(ColorStripOrientation == EThumbnailColorStripOrientation::HorizontalBottomEdge ? HAlign_Fill : HAlign_Right)
+			.VAlign(ColorStripOrientation == EThumbnailColorStripOrientation::HorizontalBottomEdge ? VAlign_Bottom : VAlign_Fill)
 			[
-				OverlayBorder
+				SAssignNew(AssetColorStripWidget, SBorder)
+				.BorderImage(FAppStyle::GetBrush("WhiteBrush"))
+				.BorderBackgroundColor(AssetColor)
+				.Padding(this, &SAssetThumbnail::GetAssetColorStripPadding)
 			];
-
-			if (AssetChipBorderImageOverride.IsSet())
-			{
-				const TSharedRef<SBorder> ThumbnailWidgetBorder = SNew(SBorder)
-					.Padding(0)
-					.BorderImage(FAppStyle::GetBrush(Style, ".AssetBorder"))
-					[
-						OverlayWidget
-					];
-
-				OverlayBorder->SetImage(AssetChipBorderImageOverride);
-				ThumbnailWidgetBorder->SetBorderImage(AssetChipBorderImageOverride);
-				ContentWidget = ThumbnailWidgetBorder;
-			}
 		}
-#else
-		// The asset color strip
-		OverlayWidget->AddSlot()
-		.HAlign(ColorStripOrientation == EThumbnailColorStripOrientation::HorizontalBottomEdge ? HAlign_Fill : HAlign_Right)
-		.VAlign(ColorStripOrientation == EThumbnailColorStripOrientation::HorizontalBottomEdge ? VAlign_Bottom : VAlign_Fill)
-		[
-			SAssignNew(AssetColorStripWidget, SBorder)
-			.BorderImage(FAppStyle::GetBrush("WhiteBrush"))
-			.BorderBackgroundColor(AssetColor)
-			.Padding(this, &SAssetThumbnail::GetAssetColorStripPadding)
-		];
-#endif
 
-#if UE_CONTENTBROWSER_NEW_STYLE
 		const UAssetDefinition* AssetDefinition = UAssetDefinitionRegistry::Get()->GetAssetDefinitionForClass(Class);
 
-		if (!AssetDefinition)
+		if (bEnableContentBrowserNewStyle)
 		{
-			AssetDefinition = GetDefault<UAssetDefinitionDefault>();
-		}
+			if (!AssetDefinition)
+			{
+				AssetDefinition = GetDefault<UAssetDefinitionDefault>();
+			}
 
-		if (InArgs._AllowAssetStatusThumbnailOverlay && AssetDefinition)
-		{
-			const TSharedPtr<FAssetStatusAssetDataInfoProvider> AssetDataInfoProvider = MakeShared<FAssetStatusAssetDataInfoProvider>(AssetData);
-			AssetDefinition->GetAssetStatusInfo(AssetDataInfoProvider, OverlayInfo);
-			// Sort the list based on Status Priority
-			auto SortStatus = [] (const FAssetDisplayInfo& InFirstStatus, const FAssetDisplayInfo& InSecondStatus)
+			if (InArgs._AllowAssetStatusThumbnailOverlay && AssetDefinition)
 			{
-				if (!InFirstStatus.Priority.IsSet())
+				const TSharedPtr<FAssetStatusAssetDataInfoProvider> AssetDataInfoProvider = MakeShared<FAssetStatusAssetDataInfoProvider>(AssetData);
+				AssetDefinition->GetAssetStatusInfo(AssetDataInfoProvider, OverlayInfo);
+				// Sort the list based on Status Priority
+				auto SortStatus = [] (const FAssetDisplayInfo& InFirstStatus, const FAssetDisplayInfo& InSecondStatus)
 				{
-					return false;
-				}
+					if (!InFirstStatus.Priority.IsSet())
+					{
+						return false;
+					}
 
-				if (!InSecondStatus.Priority.IsSet())
+					if (!InSecondStatus.Priority.IsSet())
+					{
+						return true;
+					}
+
+					return InSecondStatus.Priority.Get() < InFirstStatus.Priority.Get();
+				};
+				OverlayInfo.Sort(SortStatus);
+
+				const TSharedRef<SHorizontalBox> HorizontalBox = SNew(SHorizontalBox);
+				for (int32 StatusIndex = 0; StatusIndex < OverlayInfo.Num(); StatusIndex++)
 				{
-					return true;
+					Statuses.Add(CreateStatusWidget(StatusIndex, OverlayInfo[StatusIndex]));
+					HorizontalBox->AddSlot()
+					[
+						Statuses[StatusIndex].ToSharedRef()
+					];
 				}
 
-				return InSecondStatus.Priority.Get() < InFirstStatus.Priority.Get();
-			};
-			OverlayInfo.Sort(SortStatus);
-
-			const TSharedRef<SHorizontalBox> HorizontalBox = SNew(SHorizontalBox);
-			for (int32 StatusIndex = 0; StatusIndex < OverlayInfo.Num(); StatusIndex++)
-			{
-				Statuses.Add(CreateStatusWidget(StatusIndex, OverlayInfo[StatusIndex]));
+				StatusOverflowWidget = CreateStatusOverflowWidget();
 				HorizontalBox->AddSlot()
 				[
-					Statuses[StatusIndex].ToSharedRef()
+					StatusOverflowWidget.ToSharedRef()
 				];
-			}
-
-			StatusOverflowWidget = CreateStatusOverflowWidget();
-			HorizontalBox->AddSlot()
-			[
-				StatusOverflowWidget.ToSharedRef()
-			];
 
-			OverlayWidget->AddSlot()
-			.HAlign(HAlign_Left)
-			.VAlign(VAlign_Bottom)
-			.Padding(StatusPadding)
-			[
-				SNew(SBorder)
-				.BorderImage(FAppStyle::GetBrush(Style, ".AssetThumbnailStatusBar"))
-				.Visibility(this, &SAssetThumbnail::GetStatusBorderVisibility)
+				OverlayWidget->AddSlot()
+				.HAlign(HAlign_Left)
+				.VAlign(VAlign_Bottom)
+				.Padding(StatusPadding)
 				[
-					HorizontalBox
-				]
-			];
+					SNew(SBorder)
+					.BorderImage(FAppStyle::GetBrush(Style, ".AssetThumbnailStatusBar"))
+					.Visibility(this, &SAssetThumbnail::GetStatusBorderVisibility)
+					[
+						HorizontalBox
+					]
+				];
+			}
 		}
-#endif
+
 		if( InArgs._AllowAssetSpecificThumbnailOverlay && AssetTypeActions.IsValid() )
 		{
-#if UE_CONTENTBROWSER_NEW_STYLE
 			FAssetActionThumbnailOverlayInfo OutThumbnailInfo;
-			if (AssetDefinition && AssetDefinition->GetThumbnailActionOverlay(AssetData, OutThumbnailInfo))
+			if (bEnableContentBrowserNewStyle && AssetDefinition && AssetDefinition->GetThumbnailActionOverlay(AssetData, OutThumbnailInfo))
 			{
 				constexpr int32 OverlayZOrder = 1;
 				if (OutThumbnailInfo.ActionImageWidget.IsValid())
@@ -471,7 +473,6 @@ public:
 				}
 			}
 			else
-#endif
 			{
 				// Does the asset provide an additional thumbnail overlay?
 				PRAGMA_DISABLE_DEPRECATION_WARNINGS
@@ -487,14 +488,20 @@ public:
 			}
 		}
 
-		ChildSlot
-		[
-#if UE_CONTENTBROWSER_NEW_STYLE
-			ContentWidget
-#else
-			OverlayWidget
-#endif
-		];
+		if (bEnableContentBrowserNewStyle)
+		{
+			ChildSlot
+			[
+				ContentWidget
+			];
+		}
+		else
+		{
+			ChildSlot
+			[
+				OverlayWidget
+			];
+		}
 
 		UpdateThumbnailVisibilities();
 
@@ -581,7 +588,6 @@ public:
 		}
 	}
 
-#if UE_CONTENTBROWSER_NEW_STYLE
 	TSharedRef<SWidget> GetDefaultTooltip() const
 	{
 		const FAssetData& AssetData = AssetThumbnail->GetAssetData();
@@ -767,10 +773,8 @@ public:
 					OverallTooltipVBox
 				];
 	}
-#endif
 
 private:
-#if UE_CONTENTBROWSER_NEW_STYLE
 	FSlateWidgetRun::FWidgetRunInfo OnCreateWidgetDecoratorWidget(const FTextRunInfo& InRunInfo, const ISlateStyle* InStyle) const
 	{
 		TSharedRef<SWidget> CtrlAltWidget = SNew(SBorder)
@@ -934,7 +938,6 @@ private:
 		const int32 StatusIndexConsideringClippedStatusIfVisible = StatusOverflowWidget.IsValid() && StatusOverflowWidget->GetVisibility().IsVisible() ? StatusIndexConsideringHiddenOnes + 1 : StatusIndexConsideringHiddenOnes;
 		return StatusIndexConsideringClippedStatusIfVisible < ShownStatus ? EVisibility::Visible : EVisibility::Collapsed;
 	}
-#endif
 
 	void OnAssetDataChanged()
 	{
@@ -990,16 +993,22 @@ private:
 			AssetColor = AssetTypeActions->GetTypeColor();
 		}
 
+		static const IConsoleVariable* CVarEnableContentBrowserNewStyle = IConsoleManager::Get().FindConsoleVariable(TEXT("ContentBrowser.EnableNewStyle"));
+		const bool bEnableContentBrowserNewStyle = CVarEnableContentBrowserNewStyle && CVarEnableContentBrowserNewStyle->GetBool();
+
 		//AssetBackgroundWidget->SetBorderBackgroundColor(AssetColor.CopyWithNewOpacity(0.3f));
-#if UE_CONTENTBROWSER_NEW_STYLE
-		if (LayeredImage.IsValid())
+		if (bEnableContentBrowserNewStyle)
+		{
+			if (LayeredImage.IsValid())
+			{
+				constexpr int32 ColoredChipLayerIndex = 1;
+				LayeredImage->SetLayerColor(ColoredChipLayerIndex, AssetColor);
+			}
+		}
+		else
 		{
-			constexpr int32 ColoredChipLayerIndex = 1;
-			LayeredImage->SetLayerColor(ColoredChipLayerIndex, AssetColor);
+			AssetColorStripWidget->SetBorderBackgroundColor(AssetColor);
 		}
-#else
-		AssetColorStripWidget->SetBorderBackgroundColor(AssetColor);
-#endif
 
 		UpdateThumbnailVisibilities();
 	}
@@ -1351,7 +1360,6 @@ private:
 	TAttribute< FLinearColor > HintColorAndOpacity;
 	TAttribute<int32> GenericThumbnailSize;
 	EThumbnailColorStripOrientation ColorStripOrientation;
-#if UE_CONTENTBROWSER_NEW_STYLE
 	TSharedPtr<SWidget> StatusOverflowWidget;
 	const float StatusSize = 16.f;
 	const float StatusPadding = 4.f;
@@ -1360,7 +1368,7 @@ private:
 	TSharedPtr<IAssetSystemInfoProvider> AssetSystemInfoProvider;
 	bool bShowAssetChip;
 	TAttribute<const FSlateBrush*> AssetChipBorderImageOverride;
-#endif
+
 
 	bool bAllowHintText;
 	bool bAllowRealTimeOnHovered;
@@ -1399,13 +1407,17 @@ bool SAssetThumbnailToolTip::IsEmpty() const
 
 void SAssetThumbnailToolTip::OnOpening()
 {
-#if UE_CONTENTBROWSER_NEW_STYLE
-	TSharedPtr<SAssetThumbnail> AssetViewItemPin = AssetThumbnail.Pin();
-	if (AssetViewItemPin.IsValid())
+	static const IConsoleVariable* CVarEnableContentBrowserNewStyle = IConsoleManager::Get().FindConsoleVariable(TEXT("ContentBrowser.EnableNewStyle"));
+	const bool bEnableContentBrowserNewStyle = CVarEnableContentBrowserNewStyle && CVarEnableContentBrowserNewStyle->GetBool();
+
+	if (bEnableContentBrowserNewStyle)
 	{
-		SetContentWidget(AssetViewItemPin->GetDefaultTooltip());
+		TSharedPtr<SAssetThumbnail> AssetViewItemPin = AssetThumbnail.Pin();
+		if (AssetViewItemPin.IsValid())
+		{
+			SetContentWidget(AssetViewItemPin->GetDefaultTooltip());
+		}
 	}
-#endif
 }
 
 void SAssetThumbnailToolTip::OnClosed()
@@ -1528,13 +1540,11 @@ TSharedRef<SWidget> FAssetThumbnail::MakeThumbnailWidget( const FAssetThumbnailC
 		.AssetTypeColorOverride(InConfig.AssetTypeColorOverride)
 		.Padding(InConfig.Padding)
 		.GenericThumbnailSize(InConfig.GenericThumbnailSize)
-#if UE_CONTENTBROWSER_NEW_STYLE
 		.AssetSystemInfoProvider(InConfig.AssetSystemInfoProvider)
 		.TooltipOverride(InConfig.TooltipOverride)
 		.AllowAssetStatusThumbnailOverlay(InConfig.bAllowAssetStatusThumbnailOverlay)
 		.ShowAssetChip(InConfig.bShowAssetChip)
 		.AssetChipBorderImageOverride(InConfig.AssetChipBorderImageOverride)
-#endif
 		.ColorStripOrientation(InConfig.ColorStripOrientation);
 }
 
diff --git a/Engine/Source/Editor/UnrealEd/Private/DragAndDrop/AssetDragDropOp.cpp b/Engine/Source/Editor/UnrealEd/Private/DragAndDrop/AssetDragDropOp.cpp
--- a/Engine/Source/Editor/UnrealEd/Private/DragAndDrop/AssetDragDropOp.cpp
+++ b/Engine/Source/Editor/UnrealEd/Private/DragAndDrop/AssetDragDropOp.cpp
@@ -11,6 +11,11 @@
 #include "ThumbnailRendering/ThumbnailManager.h"
 #include "Engine/Level.h"
 
+namespace UE::Editor::ContentBrowser
+{
+	CONTENTBROWSER_API extern bool GEnableContentBrowserNewStyle;
+}
+
 TSharedRef<FAssetDragDropOp> FAssetDragDropOp::New(const FAssetData& InAssetData, TScriptInterface<IAssetFactoryInterface> AssetFactory)
 {
 	TArray<FAssetData> AssetDataArray;
@@ -85,13 +90,16 @@ TSharedPtr<SWidget> FAssetDragDropOp::GetDefaultDecorator() const
 	}
 	else if (AssetThumbnail.IsValid())
 	{
-#if UE_CONTENTBROWSER_NEW_STYLE
-		FAssetThumbnailConfig AssetThumbnailConfig;
-		AssetThumbnailConfig.bShowAssetChip = true;
-		ThumbnailWidget = AssetThumbnail->MakeThumbnailWidget(AssetThumbnailConfig);
-#else
-		ThumbnailWidget = AssetThumbnail->MakeThumbnailWidget();
-#endif
+		if (UE::Editor::ContentBrowser::GEnableContentBrowserNewStyle)
+		{
+			FAssetThumbnailConfig AssetThumbnailConfig;
+			AssetThumbnailConfig.bShowAssetChip = true;
+			ThumbnailWidget = AssetThumbnail->MakeThumbnailWidget(AssetThumbnailConfig);
+		}
+		else
+		{
+			ThumbnailWidget = AssetThumbnail->MakeThumbnailWidget();
+		}
 	}
 	else if (HasFolders())
 	{
diff --git a/Engine/Source/Editor/UnrealEd/Public/AssetDefinitionDefault.h b/Engine/Source/Editor/UnrealEd/Public/AssetDefinitionDefault.h
--- a/Engine/Source/Editor/UnrealEd/Public/AssetDefinitionDefault.h
+++ b/Engine/Source/Editor/UnrealEd/Public/AssetDefinitionDefault.h
@@ -16,9 +16,7 @@ public:
 	// UAssetDefinition Begin
 	UNREALED_API virtual EAssetCommandResult OpenAssets(const FAssetOpenArgs& OpenArgs) const override;
 	UNREALED_API virtual EAssetCommandResult PerformAssetDiff(const FAssetDiffArgs& DiffArgs) const override;
-#if UE_CONTENTBROWSER_NEW_STYLE
 	UNREALED_API virtual void GetAssetStatusInfo(const TSharedPtr<IAssetStatusInfoProvider>& InAssetStatusInfoProvider, TArray<FAssetDisplayInfo>& OutStatusInfo) const override;
-#endif
 	// UAssetDefinition End
 };
 
diff --git a/Engine/Source/Editor/UnrealEd/Public/AssetThumbnail.h b/Engine/Source/Editor/UnrealEd/Public/AssetThumbnail.h
--- a/Engine/Source/Editor/UnrealEd/Public/AssetThumbnail.h
+++ b/Engine/Source/Editor/UnrealEd/Public/AssetThumbnail.h
@@ -66,12 +66,10 @@ struct FAssetThumbnailConfig
 		, HintColorAndOpacity( FLinearColor( 0.0f, 0.0f, 0.0f, 0.0f ) )
 		, AssetTypeColorOverride()
 		, Padding(0)
-#if UE_CONTENTBROWSER_NEW_STYLE
 		, AssetSystemInfoProvider(nullptr)
 		, bAllowAssetStatusThumbnailOverlay( false )
 		, bShowAssetChip(false)
 		, AssetChipBorderImageOverride()
-#endif
 	{
 	}
 
@@ -88,13 +86,11 @@ struct FAssetThumbnailConfig
 	FMargin Padding;
 	TAttribute<int32> GenericThumbnailSize = 64;
 	EThumbnailColorStripOrientation ColorStripOrientation = EThumbnailColorStripOrientation::HorizontalBottomEdge;
-#if UE_CONTENTBROWSER_NEW_STYLE
 	TSharedPtr<IAssetSystemInfoProvider> AssetSystemInfoProvider;
 	TAttribute<TSharedPtr<IToolTip>> TooltipOverride;
 	bool bAllowAssetStatusThumbnailOverlay;
 	bool bShowAssetChip;
 	TAttribute<const FSlateBrush*> AssetChipBorderImageOverride;
-#endif
 };
 
 
