[ContentBrowser] Fixed parsing for CVar (when explicitly disabling)

diff --git a/Engine/Source/Editor/ContentBrowser/Private/ContentBrowserModule.cpp b/Engine/Source/Editor/ContentBrowser/Private/ContentBrowserModule.cpp
--- a/Engine/Source/Editor/ContentBrowser/Private/ContentBrowserModule.cpp
+++ b/Engine/Source/Editor/ContentBrowser/Private/ContentBrowserModule.cpp
@@ -55,9 +55,10 @@ namespace UE::Editor::ContentBrowser
 						continue;
 					}
 
-					if (FChar::IsDigit(OriginalCmdLine[CharIdx]);
-						const int32 Value = FChar::ConvertCharDigitToInt(OriginalCmdLine[CharIdx]))
+					if (FChar::IsDigit(OriginalCmdLine[CharIdx]))
 					{
+						const int32 Value = FChar::ConvertCharDigitToInt(OriginalCmdLine[CharIdx]);
+
 						// Explicitly disable
 						if (Value == 0)
 						{
