Add tests for filtering by type, filters with no results, and search bar add filter button

diff --git a/Engine/Source/Developer/ToolWidgets/Private/Filters/SFilterSearchBox.cpp b/Engine/Source/Developer/ToolWidgets/Private/Filters/SFilterSearchBox.cpp
--- a/Engine/Source/Developer/ToolWidgets/Private/Filters/SFilterSearchBox.cpp
+++ b/Engine/Source/Developer/ToolWidgets/Private/Filters/SFilterSearchBox.cpp
@@ -2,6 +2,7 @@
 
 #include "Filters/SFilterSearchBox.h"
 #include "Framework/Text/TextLayout.h"
+#include "Framework/MetaData/DriverMetaData.h"
 #include "Widgets/Input/SSearchBox.h"
 #include "Widgets/Input/SMenuAnchor.h"
 #include "Widgets/Views/SListView.h"
@@ -97,6 +98,7 @@ void SFilterSearchBoxImpl::Construct(const FArguments& InArgs)
 		[
 			// Button to save the currently occuring search 
 			SNew(SButton)
+			.AddMetaData(FDriverMetaData::Id(TEXT("CreateFilterFromCurrentSearchButton")))
 			.ContentPadding(0)
 			.Visibility_Lambda([this]()
 			{
diff --git a/Engine/Source/Editor/ContentBrowser/Private/AssetViewWidgets.cpp b/Engine/Source/Editor/ContentBrowser/Private/AssetViewWidgets.cpp
--- a/Engine/Source/Editor/ContentBrowser/Private/AssetViewWidgets.cpp
+++ b/Engine/Source/Editor/ContentBrowser/Private/AssetViewWidgets.cpp
@@ -30,6 +30,7 @@
 #include "EditorFramework/AssetImportData.h"
 #include "Fonts/FontMeasure.h"
 #include "Framework/Application/SlateApplication.h"
+#include "Framework/MetaData/DriverMetaData.h"
 #include "Framework/Text/TextLayout.h"
 #include "HAL/IConsoleManager.h"
 #include "HAL/PlatformApplicationMisc.h"
@@ -2279,6 +2280,7 @@ void SAssetTileItem::Construct( const FArguments& InArgs )
 										.OverflowPolicy(ETextOverflowPolicy::Ellipsis)
 										.Text(this, &SAssetTileItem::GetAssetClassText)
 										.ColorAndOpacity(this, &SAssetTileItem::GetAssetClassLabelTextColor)
+										.AddMetaData(FDriverMetaData::Id(TEXT("AssetTileClassTextBlock")))
 									]
 
 									+SHorizontalBox::Slot()
diff --git a/Engine/Source/Editor/ContentBrowser/Private/ContentBrowserSingleton.cpp b/Engine/Source/Editor/ContentBrowser/Private/ContentBrowserSingleton.cpp
--- a/Engine/Source/Editor/ContentBrowser/Private/ContentBrowserSingleton.cpp
+++ b/Engine/Source/Editor/ContentBrowser/Private/ContentBrowserSingleton.cpp
@@ -1400,6 +1400,11 @@ void FContentBrowserSingleton::ModifyCustomTextFilterByLabel(const FCustomTextFi
 	}
 }
 
+/**
+ * Checks if the Primary Content Browser's AssetView finished filtering
+ * 
+ * @return true if the AssetView is done filtering
+ */
 bool FContentBrowserSingleton::IsAssetViewDoneFiltering()
 {
 	bool isDoneFiltering = false;
@@ -1417,4 +1422,26 @@ bool FContentBrowserSingleton::IsAssetViewDoneFiltering()
 	return isDoneFiltering;
 }
 
+/**
+ * Sets the search text for the Primary Content Browser
+ * 
+ * @params NewText - the text to put in the search box
+ */
+void FContentBrowserSingleton::SetSearchText(const FText& NewText)
+{
+	if (!HasPrimaryContentBrowser())
+	{
+		ChooseNewPrimaryBrowser();
+	}
+
+	if (PrimaryContentBrowser.IsValid())
+	{
+		if (TSharedPtr<SContentBrowser> PrimaryContentBrowserPinned = PrimaryContentBrowser.Pin())
+		{
+			PrimaryContentBrowserPinned->SetSearchText(NewText);
+			PrimaryContentBrowserPinned->OnSearchBoxCommitted(NewText, ETextCommit::Type::Default);
+		}
+	}
+}
+
 #undef LOCTEXT_NAMESPACE
diff --git a/Engine/Source/Editor/ContentBrowser/Private/ContentBrowserSingleton.h b/Engine/Source/Editor/ContentBrowser/Private/ContentBrowserSingleton.h
--- a/Engine/Source/Editor/ContentBrowser/Private/ContentBrowserSingleton.h
+++ b/Engine/Source/Editor/ContentBrowser/Private/ContentBrowserSingleton.h
@@ -122,6 +122,7 @@ PRAGMA_ENABLE_DEPRECATION_WARNINGS
 	virtual void DeleteCustomTextFilterByLabel(const FText& FilterLabel) override;
 	virtual void ModifyCustomTextFilterByLabel(const FCustomTextFilterData& NewFilterData, const FText& FilterLabel) override;
 	virtual bool IsAssetViewDoneFiltering() override;
+	virtual void SetSearchText(const FText& NewText) override;
 
 	/** Broadcast that the favorites have changed. */
 	void BroadcastFavoritesChanged() const;
diff --git a/Engine/Source/Editor/ContentBrowser/Private/SAssetView.cpp b/Engine/Source/Editor/ContentBrowser/Private/SAssetView.cpp
--- a/Engine/Source/Editor/ContentBrowser/Private/SAssetView.cpp
+++ b/Engine/Source/Editor/ContentBrowser/Private/SAssetView.cpp
@@ -2251,6 +2251,11 @@ bool SAssetView::HasItemsPendingFilter() const
 	return Items->HasItemsPendingFilter();
 }
 
+bool SAssetView::HasThumbnailsPendingUpdate() const
+{
+	return bPendingUpdateThumbnails;
+}
+
 void SAssetView::ProcessItemsPendingFilter(const double TickStartTime)
 {
 	const double ProcessItemsPendingFilterStartTime = FPlatformTime::Seconds();
diff --git a/Engine/Source/Editor/ContentBrowser/Private/SContentBrowser.cpp b/Engine/Source/Editor/ContentBrowser/Private/SContentBrowser.cpp
--- a/Engine/Source/Editor/ContentBrowser/Private/SContentBrowser.cpp
+++ b/Engine/Source/Editor/ContentBrowser/Private/SContentBrowser.cpp
@@ -2816,6 +2816,14 @@ void SContentBrowser::NewFileItemRequested(const FContentBrowserItemDataTemporar
 	}
 }
 
+void SContentBrowser::SetSearchText(const FText& InSearchText)
+{
+	if (SearchBoxPtr.IsValid())
+	{
+		SearchBoxPtr->SetText(InSearchText);
+	}
+}
+
 void SContentBrowser::SetSearchBoxText(const FText& InSearchText)
 {
 	// Has anything changed? (need to test case as the operators are case-sensitive)
@@ -3372,7 +3380,8 @@ bool SContentBrowser::IsAssetViewDoneFiltering()
 
 	if (AssetViewPtr.IsValid())
 	{
-		isDoneFiltering = !AssetViewPtr->HasItemsPendingFilter();
+		isDoneFiltering = !AssetViewPtr->HasItemsPendingFilter()
+			&& !AssetViewPtr->HasThumbnailsPendingUpdate();
 	}
 
 	return isDoneFiltering;
diff --git a/Engine/Source/Editor/ContentBrowser/Private/SContentBrowser.h b/Engine/Source/Editor/ContentBrowser/Private/SContentBrowser.h
--- a/Engine/Source/Editor/ContentBrowser/Private/SContentBrowser.h
+++ b/Engine/Source/Editor/ContentBrowser/Private/SContentBrowser.h
@@ -241,6 +241,12 @@ public:
 	/** Gets whether the AssetView in the Content Browser is done filtering */
 	bool IsAssetViewDoneFiltering();
 
+	/** Sets the actual text in the Search Box */
+	void SetSearchText(const FText& InSearchText);
+
+	/** Called by the editable text control when the user commits a text change */
+	void OnSearchBoxCommitted(const FText& InSearchText, ETextCommit::Type CommitInfo);
+
 private:
 
 	/** Called prior to syncing the selection in this Content Browser */
@@ -309,9 +315,6 @@ private:
 	/** Called by the editable text control when the search text is changed by the user */
 	void OnSearchBoxChanged(const FText& InSearchText);
 
-	/** Called by the editable text control when the user commits a text change */
-	void OnSearchBoxCommitted(const FText& InSearchText, ETextCommit::Type CommitInfo);
-
 	/** Called by the editable text control to allow the content browser to handle specific key presses without it resulting in typing into the search box */
 	FReply OnSearchKeyDown(const FGeometry& Geometry, const FKeyEvent& InKeyEvent);
 
diff --git a/Engine/Source/Editor/ContentBrowser/Public/IContentBrowserSingleton.h b/Engine/Source/Editor/ContentBrowser/Public/IContentBrowserSingleton.h
--- a/Engine/Source/Editor/ContentBrowser/Public/IContentBrowserSingleton.h
+++ b/Engine/Source/Editor/ContentBrowser/Public/IContentBrowserSingleton.h
@@ -782,4 +782,7 @@ public:
 
 	/** Gets whether the AssetView in the Content Browser is done filtering */
 	virtual bool IsAssetViewDoneFiltering() = 0;
+
+	/** Sets the AssetSearchBox text */
+	virtual void SetSearchText(const FText& NewText) = 0;
 };
diff --git a/Engine/Source/Editor/ContentBrowser/Public/SAssetView.h b/Engine/Source/Editor/ContentBrowser/Public/SAssetView.h
--- a/Engine/Source/Editor/ContentBrowser/Public/SAssetView.h
+++ b/Engine/Source/Editor/ContentBrowser/Public/SAssetView.h
@@ -487,6 +487,9 @@ public:
 
 	/** True if we have items pending that ProcessItemsPendingFilter still needs to process */
 	bool HasItemsPendingFilter() const;
+
+	/** True if we have thumbnails that are pending update */
+	bool HasThumbnailsPendingUpdate() const;
 private:
 
 	/** Sets the pending selection to the current selection (used when changing views or refreshing the view). */
