Made the behavior of the recursive search setting more predictable and less complex.

diff --git a/Engine/Source/Editor/ContentBrowser/Private/SAssetView.cpp b/Engine/Source/Editor/ContentBrowser/Private/SAssetView.cpp
--- a/Engine/Source/Editor/ContentBrowser/Private/SAssetView.cpp
+++ b/Engine/Source/Editor/ContentBrowser/Private/SAssetView.cpp
@@ -3140,28 +3140,31 @@ void SAssetView::ToggleFilteringRecursively()
 
 bool SAssetView::ShouldFilterRecursively() const
 {
-	// Quick check for conditions which force recursive filtering
-	if (bUserSearching)
+	if (IsFilteringRecursively())
 	{
-		return true;
-	}
+		// Quick check for conditions that activate the recursive filtering
+		if (bUserSearching)
+		{
+			return true;
+		}
 
-	if (IsFilteringRecursively() && !BackendFilter.IsEmpty() )
-	{
-		return true;
-	}
+		if (!BackendFilter.IsEmpty() )
+		{
+			return true;
+		}
 
-	// Otherwise, check if there are any non-inverse frontend filters selected
-	if (FrontendFilters.IsValid())
-	{
-		for (int32 FilterIndex = 0; FilterIndex < FrontendFilters->Num(); ++FilterIndex)
+		// Otherwise, check if there are any non-inverse frontend filters selected
+		if (FrontendFilters.IsValid())
 		{
-			const auto* Filter = static_cast<FFrontendFilter*>(FrontendFilters->GetFilterAtIndex(FilterIndex).Get());
-			if (Filter)
+			for (int32 FilterIndex = 0; FilterIndex < FrontendFilters->Num(); ++FilterIndex)
 			{
-				if (!Filter->IsInverseFilter())
+				const auto* Filter = static_cast<FFrontendFilter*>(FrontendFilters->GetFilterAtIndex(FilterIndex).Get());
+				if (Filter)
 				{
-					return true;
+					if (!Filter->IsInverseFilter())
+					{
+						return true;
+					}
 				}
 			}
 		}
