[ContentBrowser] Added missing Source Panel entry to settings menu in new style

diff --git a/Engine/Source/Editor/ContentBrowser/Private/SContentBrowser.cpp b/Engine/Source/Editor/ContentBrowser/Private/SContentBrowser.cpp
--- a/Engine/Source/Editor/ContentBrowser/Private/SContentBrowser.cpp
+++ b/Engine/Source/Editor/ContentBrowser/Private/SContentBrowser.cpp
@@ -1861,23 +1861,22 @@ void SContentBrowser::ExtendViewOptionsMenu(const FContentBrowserConfig* Config)
 						}
 					}
 
-					if (!GEnableContentBrowserNewStyle)
+					if (bShowSourcesView)
 					{
-						if (bShowSourcesView)
-						{
-							FToolMenuSection& Section = InMenu->FindOrAddSection("View");
-							Section.AddMenuEntry(
-								"ToggleSources",
-								LOCTEXT("ToggleSourcesView", "Show Sources Panel"),
-								LOCTEXT("ToggleSourcesView_Tooltip", "Show or hide the sources panel"),
-								TAttribute<FSlateIcon>(),
-								FUIAction(
-									FExecuteAction::CreateLambda([ContentBrowser = Context->OwningContentBrowser]() { ContentBrowser.Pin()->SourcesViewExpandClicked(); }),
-									FCanExecuteAction(),
-									FIsActionChecked::CreateLambda([ContentBrowser = Context->OwningContentBrowser]() { return ContentBrowser.Pin()->bSourcesViewExpanded; })),
-								EUserInterfaceActionType::ToggleButton
-							);
-						}
+						FToolMenuSection& Section = InMenu->FindOrAddSection(GEnableContentBrowserNewStyle ? "Show" : "View");
+						Section.AddMenuEntry(
+							"ToggleSources",
+							GEnableContentBrowserNewStyle
+							? LOCTEXT("ToggleSourcesView_NewStyle", "Sources Panel")
+							: LOCTEXT("ToggleSourcesView", "Show Sources Panel"),
+							LOCTEXT("ToggleSourcesView_Tooltip", "Show or hide the sources panel"),
+							TAttribute<FSlateIcon>(),
+							FUIAction(
+								FExecuteAction::CreateLambda([ContentBrowser = Context->OwningContentBrowser]() { ContentBrowser.Pin()->SourcesViewExpandClicked(); }),
+								FCanExecuteAction(),
+								FIsActionChecked::CreateLambda([ContentBrowser = Context->OwningContentBrowser]() { return ContentBrowser.Pin()->bSourcesViewExpanded; })),
+							EUserInterfaceActionType::ToggleButton
+						);
 					}
 				}
 			}
