[ContentBrowser] Enable ThumbnailEditMode with new CB
Re-Added EditThumbnailMode to new CB

diff --git a/Engine/Source/Editor/ContentBrowser/Private/AssetViewWidgets.cpp b/Engine/Source/Editor/ContentBrowser/Private/AssetViewWidgets.cpp
--- a/Engine/Source/Editor/ContentBrowser/Private/AssetViewWidgets.cpp
+++ b/Engine/Source/Editor/ContentBrowser/Private/AssetViewWidgets.cpp
@@ -267,15 +267,15 @@ TSharedRef<SWidget> FAssetViewItemHelper::CreateListTileItemContents(T* const In
 					.Image(InTileOrListItem, &T::GetDirtyImage)
 				]
 			];
-		}
 
-		// Tools for thumbnail edit mode
-		ItemContentsOverlay->AddSlot()
-		[
-			SNew(SThumbnailEditModeTools, InTileOrListItem->AssetThumbnail)
-			.SmallView(!InTileOrListItem->CanDisplayPrimitiveTools())
-			.Visibility(InTileOrListItem, &T::GetThumbnailEditModeUIVisibility)
-		];
+			// Tools for thumbnail edit mode
+			ItemContentsOverlay->AddSlot()
+			[
+				SNew(SThumbnailEditModeTools, InTileOrListItem->AssetThumbnail)
+				.SmallView(!InTileOrListItem->CanDisplayPrimitiveTools())
+				.Visibility(InTileOrListItem, &T::GetThumbnailEditModeUIVisibility)
+			];
+		}
 	}
 	
 	return ItemContentsOverlay;
diff --git a/Engine/Source/Editor/ContentBrowser/Private/Menus/AssetViewMenus.cpp b/Engine/Source/Editor/ContentBrowser/Private/Menus/AssetViewMenus.cpp
--- a/Engine/Source/Editor/ContentBrowser/Private/Menus/AssetViewMenus.cpp
+++ b/Engine/Source/Editor/ContentBrowser/Private/Menus/AssetViewMenus.cpp
@@ -257,7 +257,7 @@ FDelayedAutoRegisterHelper SAssetView::AssetViewOptionsMenuRegistration(
 							));
 						}
 
-						if (!GEnableContentBrowserNewStyle)
+						// Thumbnail EditMode Section
 						{
 							Section.AddMenuEntry(
 								"ThumbnailEditMode",
