[ContentBrowser] Moved style toggle cvar to UnrealEd

diff --git a/Engine/Plugins/Editor/ContentBrowser/ContentBrowserAssetDataSource/Source/ContentBrowserAssetDataSource/Private/ContentBrowserAssetDataSource.cpp b/Engine/Plugins/Editor/ContentBrowser/ContentBrowserAssetDataSource/Source/ContentBrowserAssetDataSource/Private/ContentBrowserAssetDataSource.cpp
--- a/Engine/Plugins/Editor/ContentBrowser/ContentBrowserAssetDataSource/Source/ContentBrowserAssetDataSource/Private/ContentBrowserAssetDataSource.cpp
+++ b/Engine/Plugins/Editor/ContentBrowser/ContentBrowserAssetDataSource/Source/ContentBrowserAssetDataSource/Private/ContentBrowserAssetDataSource.cpp
@@ -50,7 +50,20 @@
 
 namespace UE::Editor::ContentBrowser
 {
-	extern CONTENTBROWSER_API bool GEnableContentBrowserNewStyle;
+	inline bool IsNewStyleEnabled()
+	{
+		static bool bIsNewStyleEnabled = false;
+		UE_CALL_ONCE([&]()
+		{
+			if (const IConsoleVariable* CVar = IConsoleManager::Get().FindConsoleVariable(TEXT("ContentBrowser.EnableNewStyle")))
+			{
+				ensureAlwaysMsgf(!EnumHasAnyFlags(CVar->GetFlags(), ECVF_Default), TEXT("The CVar should have already been set from commandline, @see: UnrealEdGlobals.cpp, UE::Editor::ContentBrowser::EnableContentBrowserNewStyleCVarRegistration."));
+				bIsNewStyleEnabled = CVar->GetBool();
+			}
+		});
+
+		return bIsNewStyleEnabled;
+	}
 }
 
 namespace AssetDataSource
@@ -3596,7 +3609,7 @@ void UContentBrowserAssetDataSource::PopulateContentBrowserToolBar(UToolMenu* In
 
 	TSharedPtr<SWidget> ImportButton = nullptr;
 
-	if (UE::Editor::ContentBrowser::GEnableContentBrowserNewStyle)
+	if (UE::Editor::ContentBrowser::IsNewStyleEnabled())
 	{
 		SAssignNew(ImportButton, SActionButton)
 			.ToolTipText(LOCTEXT("ImportTooltip", "Import assets from files to the currently selected folder"))
diff --git a/Engine/Source/Editor/ContentBrowser/Private/AssetViewWidgets.cpp b/Engine/Source/Editor/ContentBrowser/Private/AssetViewWidgets.cpp
--- a/Engine/Source/Editor/ContentBrowser/Private/AssetViewWidgets.cpp
+++ b/Engine/Source/Editor/ContentBrowser/Private/AssetViewWidgets.cpp
@@ -92,13 +92,6 @@ class FDragDropEvent;
 
 #define LOCTEXT_NAMESPACE "ContentBrowser"
 
-namespace UE::Editor::ContentBrowser
-{
-	extern CONTENTBROWSER_API bool GEnableContentBrowserNewStyle;
-}
-
-using namespace UE::Editor::ContentBrowser;
-
 ///////////////////////////////
 // FAssetViewModeUtils
 ///////////////////////////////
@@ -252,7 +245,7 @@ TSharedRef<SWidget> FAssetViewItemHelper::CreateListTileItemContents(T* const In
 			]
 		];
 
-		if (!GEnableContentBrowserNewStyle)
+		if (!UE::Editor::ContentBrowser::IsNewStyleEnabled())
 		{
 			// Dirty state
 			ItemContentsOverlay->AddSlot()
@@ -304,11 +297,11 @@ public:
 		SToolTip::Construct(
 			SToolTip::FArguments()
 			.TextMargin(
-				GEnableContentBrowserNewStyle
+				UE::Editor::ContentBrowser::IsNewStyleEnabled()
 				? FMargin(12.f, 8.f, 12.f, 2.f)
 				: FMargin(1.0f))
 			.BorderImage(
-				GEnableContentBrowserNewStyle
+				UE::Editor::ContentBrowser::IsNewStyleEnabled()
 				? FAppStyle::GetBrush("AssetThumbnail.Tooltip.Border")
 				: UE::ContentBrowser::Private::FContentBrowserStyle::Get().GetBrush("ContentBrowser.TileViewTooltip.ToolTipBorder")));
 	}
@@ -1143,7 +1136,7 @@ TSharedRef<SWidget> SAssetViewItem::CreateToolTipWidget() const
 				}
 			}
 
-			if (GEnableContentBrowserNewStyle)
+			if (UE::Editor::ContentBrowser::IsNewStyleEnabled())
 			{
 				// Default values
 				FName FolderBrushName = TEXT("ContentBrowser.ListViewFolderIcon");
@@ -1336,7 +1329,7 @@ FText SAssetViewItem::GetAssetUserDescription() const
 
 void SAssetViewItem::AddToToolTipInfoBox(const TSharedRef<SVerticalBox>& InfoBox, const FText& Key, const FText& Value, bool bImportant) const
 {
-	if (GEnableContentBrowserNewStyle)
+	if (UE::Editor::ContentBrowser::IsNewStyleEnabled())
 	{
 		InfoBox->AddSlot()
 		.Padding(0.f, 0.f, 0.f, 6.f)
@@ -1690,7 +1683,7 @@ void SAssetViewItem::CacheDisplayTags()
 
 const FSlateBrush* SAssetViewItem::GetBorderImage() const
 {
-	if (GEnableContentBrowserNewStyle)
+	if (UE::Editor::ContentBrowser::IsNewStyleEnabled())
 	{
 		return bDraggedOver ? FAppStyle::GetBrush("AssetThumbnail.AssetBorder") : FAppStyle::GetBrush("NoBorder");
 	}
@@ -2080,7 +2073,7 @@ void SAssetTileItem::Construct( const FArguments& InArgs )
 		ThumbnailConfig.Padding= FMargin(2.0f);
 		ThumbnailConfig.GenericThumbnailSize = MakeAttributeSP(this, &SAssetTileItem::GetGenericThumbnailSize);
 
-		if (GEnableContentBrowserNewStyle)
+		if (UE::Editor::ContentBrowser::IsNewStyleEnabled())
 		{
 			ThumbnailConfig.AssetSystemInfoProvider = MakeShared<FAssetSystemContentBrowserInfoProvider>(AssetItem);
 			ThumbnailConfig.bAllowAssetStatusThumbnailOverlay = true;
@@ -2111,7 +2104,7 @@ void SAssetTileItem::Construct( const FArguments& InArgs )
 
 		Thumbnail = AssetThumbnail->MakeThumbnailWidget(ThumbnailConfig);
 
-		if (GEnableContentBrowserNewStyle)
+		if (UE::Editor::ContentBrowser::IsNewStyleEnabled())
 		{
 			// Use the same tooltip as the Thumbnail
 			if (const TSharedPtr<IToolTip>& ThumbnailTooltip = Thumbnail->GetToolTip())
@@ -2158,7 +2151,7 @@ void SAssetTileItem::Construct( const FArguments& InArgs )
 			.ColorAndOpacity(this, &SAssetTileItem::GetNameAreaTextColor);
 	}
 
-	if (GEnableContentBrowserNewStyle)
+	if (UE::Editor::ContentBrowser::IsNewStyleEnabled())
 	{
 		constexpr float BorderPadding = 1.f;
 		constexpr float ShadowLeftTopPadding = 3.f;
@@ -2345,7 +2338,7 @@ void SAssetTileItem::Construct( const FArguments& InArgs )
 
 	if(AssetItem.IsValid())
 	{
-		if (GEnableContentBrowserNewStyle)
+		if (UE::Editor::ContentBrowser::IsNewStyleEnabled())
 		{
 			AssetItem->OnRenameRequested().BindSP(this, &SAssetTileItem::OnRequestingRename);
 		}
@@ -2465,7 +2458,7 @@ FOptionalSize SAssetTileItem::GetStateIconImageSize() const
 
 FOptionalSize SAssetTileItem::GetThumbnailBoxSize() const
 {
-	return GEnableContentBrowserNewStyle ? FOptionalSize(ThumbnailDimension.Get()) : FOptionalSize(ItemWidth.Get()- ThumbnailPadding);
+	return UE::Editor::ContentBrowser::IsNewStyleEnabled() ? FOptionalSize(ThumbnailDimension.Get()) : FOptionalSize(ItemWidth.Get()- ThumbnailPadding);
 }
 
 EVisibility SAssetTileItem::GetAssetClassLabelVisibility() const
@@ -2484,12 +2477,12 @@ FSlateColor SAssetTileItem::GetAssetClassLabelTextColor() const
 		return FStyleColors::White;
 	}
 
-	return GEnableContentBrowserNewStyle ? FStyleColors::Hover2 : FSlateColor::UseSubduedForeground();
+	return UE::Editor::ContentBrowser::IsNewStyleEnabled() ? FStyleColors::Hover2 : FSlateColor::UseSubduedForeground();
 }
 
 FSlateFontInfo SAssetTileItem::GetThumbnailFont() const
 {
-	if (GEnableContentBrowserNewStyle)
+	if (UE::Editor::ContentBrowser::IsNewStyleEnabled())
 	{
 		const static FLazyName RegularFont("ContentBrowser.AssetTileViewNameFont");
 		return FAppStyle::GetFontStyle(RegularFont);
@@ -2562,10 +2555,10 @@ const FSlateBrush* SAssetTileItem::GetNameAreaBackgroundImage() const
 		return FStyleDefaults::GetNoBrush();
 	}
 
-	const FName SelectedHover = GEnableContentBrowserNewStyle ? "ContentBrowser.AssetTileItem.AssetContentSelectedHoverBackground" : "ContentBrowser.AssetTileItem.NameAreaSelectedHoverBackground";
-	const FName Selected = GEnableContentBrowserNewStyle ? "ContentBrowser.AssetTileItem.AssetContentSelectedBackground" : "ContentBrowser.AssetTileItem.NameAreaSelectedBackground";
-	const FName Hovered = GEnableContentBrowserNewStyle ? "ContentBrowser.AssetTileItem.AssetContentHoverBackground" : "ContentBrowser.AssetTileItem.NameAreaHoverBackground";
-	const FName Normal = GEnableContentBrowserNewStyle ? "ContentBrowser.AssetTileItem.AssetContent" : "ContentBrowser.AssetTileItem.NameAreaBackground";
+	const FName SelectedHover = UE::Editor::ContentBrowser::IsNewStyleEnabled() ? "ContentBrowser.AssetTileItem.AssetContentSelectedHoverBackground" : "ContentBrowser.AssetTileItem.NameAreaSelectedHoverBackground";
+	const FName Selected = UE::Editor::ContentBrowser::IsNewStyleEnabled() ? "ContentBrowser.AssetTileItem.AssetContentSelectedBackground" : "ContentBrowser.AssetTileItem.NameAreaSelectedBackground";
+	const FName Hovered = UE::Editor::ContentBrowser::IsNewStyleEnabled() ? "ContentBrowser.AssetTileItem.AssetContentHoverBackground" : "ContentBrowser.AssetTileItem.NameAreaHoverBackground";
+	const FName Normal = UE::Editor::ContentBrowser::IsNewStyleEnabled() ? "ContentBrowser.AssetTileItem.AssetContent" : "ContentBrowser.AssetTileItem.NameAreaBackground";
 
 	const bool bIsSelected = IsSelected.IsBound() ? IsSelected.Execute() : false;
 	const bool bIsHoveredOrDraggedOver = IsHovered() || bDraggedOver;
@@ -2648,12 +2641,12 @@ FSlateColor SAssetTileItem::GetNameAreaTextColor() const
 		return FStyleColors::White;
 	}
 
-	return GEnableContentBrowserNewStyle ? FStyleColors::Foreground : FSlateColor::UseForeground();
+	return UE::Editor::ContentBrowser::IsNewStyleEnabled() ? FStyleColors::Foreground : FSlateColor::UseForeground();
 }
 
 FOptionalSize SAssetTileItem::GetNameAreaMaxDesiredHeight() const
 {
-	if (GEnableContentBrowserNewStyle)
+	if (UE::Editor::ContentBrowser::IsNewStyleEnabled())
 	{
 		constexpr int32 MaxHeightNameArea = 42;
 		return MaxHeightNameArea;
@@ -2667,7 +2660,7 @@ FOptionalSize SAssetTileItem::GetNameAreaMaxDesiredHeight() const
 int32 SAssetTileItem::GetGenericThumbnailSize() const
 {
 	const int32 ThumbnailSizeIdx = (int32)CurrentThumbnailSize.Get();
-	if (GEnableContentBrowserNewStyle)
+	if (UE::Editor::ContentBrowser::IsNewStyleEnabled())
 	{
 		static int32 GenericThumbnailSizes[(int32)EThumbnailSize::MAX] = { 80, 96, 112, 128, 136, 160 };
 		return GenericThumbnailSizes[ThumbnailSizeIdx];
@@ -2682,7 +2675,7 @@ int32 SAssetTileItem::GetGenericThumbnailSize() const
 EVisibility SAssetTileItem::GetSCCIconVisibility() const
 {
 	// Hide the scc state icon when there is no brush or in tiny size since there isn't enough space
-	if (GEnableContentBrowserNewStyle)
+	if (UE::Editor::ContentBrowser::IsNewStyleEnabled())
 	{
 		return bHasCCStateBrush && ISourceControlModule::Get().IsEnabled() && ISourceControlModule::Get().GetProvider().IsAvailable() ? EVisibility::Visible : EVisibility::Collapsed;
 	}
@@ -2736,7 +2729,7 @@ void SAssetTileItem::InitializeAssetNameHeights()
 			AssetNameHeights[(int32)EThumbnailSize::Medium] = RegularFontHeight * MediumSizeMultiplier;
 			AssetNameHeights[(int32)EThumbnailSize::Large] = RegularFontHeight * LargeSizeMultiplier;
 
-			if (GEnableContentBrowserNewStyle)
+			if (UE::Editor::ContentBrowser::IsNewStyleEnabled())
 			{
 				AssetNameHeights[(int32)EThumbnailSize::XLarge] = RegularFontHeight * XLargeSizeMultiplier;
 				AssetNameHeights[(int32)EThumbnailSize::Huge] = RegularFontHeight * HugeSizeMultiplier;
diff --git a/Engine/Source/Editor/ContentBrowser/Private/ContentBrowserMenuUtils.cpp b/Engine/Source/Editor/ContentBrowser/Private/ContentBrowserMenuUtils.cpp
--- a/Engine/Source/Editor/ContentBrowser/Private/ContentBrowserMenuUtils.cpp
+++ b/Engine/Source/Editor/ContentBrowser/Private/ContentBrowserMenuUtils.cpp
@@ -2,18 +2,13 @@
 
 #include "ContentBrowserMenuUtils.h"
 #include "ContentBrowserConfig.h"
+#include "ContentBrowserMenuContexts.h"
+#include "ContentBrowserStyle.h"
 #include "Settings/ContentBrowserSettings.h"
 #include "ToolMenu.h"
 
 #define LOCTEXT_NAMESPACE "ContentBrowserMenuUtils"
 
-namespace UE::Editor::ContentBrowser
-{
-	extern CONTENTBROWSER_API bool GEnableContentBrowserNewStyle;
-}
-
-using namespace UE::Editor::ContentBrowser;
-
 namespace ContentBrowserMenuUtils::Filters
 {
 	FContentBrowserInstanceConfig* GetContentBrowserConfig(FName InOwningContentBrowserName)
@@ -216,14 +211,14 @@ void ContentBrowserMenuUtils::AddFiltersToMenu(UToolMenu* InMenu, const FName& I
 		ContentMenuInsert.Position = EToolMenuInsertType::After;
 	}
 
-	FToolMenuSection& Section = GEnableContentBrowserNewStyle ? InMenu->FindOrAddSection("Show") : InMenu->FindOrAddSection("Content", LOCTEXT("ContentHeading", "Content"), ContentMenuInsert);
+	FToolMenuSection& Section = UE::Editor::ContentBrowser::IsNewStyleEnabled() ? InMenu->FindOrAddSection("Show") : InMenu->FindOrAddSection("Content", LOCTEXT("ContentHeading", "Content"), ContentMenuInsert);
 
 	// New style omits the "Show" prefix
-	const FText ShowCppLabel = GEnableContentBrowserNewStyle ?  LOCTEXT("ShowCppClassesOption_NewStyle", "C++ Classes") : LOCTEXT("ShowCppClassesOption", "Show C++ Classes");
-	const FText ShowDevelopersContentLabel = GEnableContentBrowserNewStyle ? LOCTEXT("ShowDevelopersContentOption_NewStyle", "Developers Content") : LOCTEXT("ShowDevelopersContentOption", "Show Developers Content");
-	const FText ShowEngineContentLabel = GEnableContentBrowserNewStyle ? LOCTEXT("ShowEngineFolderOption_NewStyle", "Engine Content") : LOCTEXT("ShowEngineFolderOption", "Show Engine Content");
-	const FText ShowPluginContentLabel = GEnableContentBrowserNewStyle ? LOCTEXT("ShowPluginFolderOption_NewStyle", "Plugin Content") : LOCTEXT("ShowPluginFolderOption", "Show Plugin Content");
-	const FText ShowLocalizedContentLabel = GEnableContentBrowserNewStyle ? LOCTEXT("ShowLocalizedContentOption_NewStyle", "Localized Content") : LOCTEXT("ShowLocalizedContentOption", "Show Localized Content");
+	const FText ShowCppLabel = UE::Editor::ContentBrowser::IsNewStyleEnabled() ?  LOCTEXT("ShowCppClassesOption_NewStyle", "C++ Classes") : LOCTEXT("ShowCppClassesOption", "Show C++ Classes");
+	const FText ShowDevelopersContentLabel = UE::Editor::ContentBrowser::IsNewStyleEnabled() ? LOCTEXT("ShowDevelopersContentOption_NewStyle", "Developers Content") : LOCTEXT("ShowDevelopersContentOption", "Show Developers Content");
+	const FText ShowEngineContentLabel = UE::Editor::ContentBrowser::IsNewStyleEnabled() ? LOCTEXT("ShowEngineFolderOption_NewStyle", "Engine Content") : LOCTEXT("ShowEngineFolderOption", "Show Engine Content");
+	const FText ShowPluginContentLabel = UE::Editor::ContentBrowser::IsNewStyleEnabled() ? LOCTEXT("ShowPluginFolderOption_NewStyle", "Plugin Content") : LOCTEXT("ShowPluginFolderOption", "Show Plugin Content");
+	const FText ShowLocalizedContentLabel = UE::Editor::ContentBrowser::IsNewStyleEnabled() ? LOCTEXT("ShowLocalizedContentOption_NewStyle", "Localized Content") : LOCTEXT("ShowLocalizedContentOption", "Show Localized Content");
 
 	Section.AddMenuEntry(
 		"ShowCppClasses",
diff --git a/Engine/Source/Editor/ContentBrowser/Private/ContentBrowserModule.cpp b/Engine/Source/Editor/ContentBrowser/Private/ContentBrowserModule.cpp
--- a/Engine/Source/Editor/ContentBrowser/Private/ContentBrowserModule.cpp
+++ b/Engine/Source/Editor/ContentBrowser/Private/ContentBrowserModule.cpp
@@ -22,62 +22,8 @@ DEFINE_LOG_CATEGORY(LogContentBrowser);
 
 const FName FContentBrowserModule::NumberOfRecentAssetsName(TEXT("NumObjectsInRecentList"));
 
-// @todo: Can remove after permanent switch to new style
-namespace UE::Editor::ContentBrowser
-{
-	bool GEnableContentBrowserNewStyle = false;
-
-	namespace Private
-	{
-		static FAutoConsoleVariableRef CVarEnableContentBrowserNewStyle(
-			TEXT("ContentBrowser.EnableNewStyle"),
-			GEnableContentBrowserNewStyle,
-			TEXT("Whether or not to enable the Content Browser restyle (~5.6).")
-		);
-
-		void SetEnableNewStyleFromCmdLine()
-		{
-			static constexpr const TCHAR* CVarName = TEXT("ContentBrowser.EnableNewStyle");
-
-			const FString OriginalCmdLine = FCommandLine::GetOriginal();
-			if (const int32 FoundIdx = OriginalCmdLine.Find(CVarName);
-				FoundIdx != INDEX_NONE)
-			{
-				// If present, enable
-				GEnableContentBrowserNewStyle = true;
-
-				// Check for explicit disable =0 or \s0
-				for (int32 CharIdx = FoundIdx + FCString::Strlen(CVarName); CharIdx < OriginalCmdLine.Len(); ++CharIdx)
-				{
-					if (OriginalCmdLine[CharIdx] == TEXT(' ')
-						|| OriginalCmdLine[CharIdx] == TEXT('='))
-					{
-						continue;
-					}
-
-					if (FChar::IsDigit(OriginalCmdLine[CharIdx]))
-					{
-						const int32 Value = FChar::ConvertCharDigitToInt(OriginalCmdLine[CharIdx]);
-
-						// Explicitly disable
-						if (Value == 0)
-						{
-							GEnableContentBrowserNewStyle = false;
-							break;
-						}
-					}
-
-					break;
-				}
-			}
-		}
-	}
-}
-
 void FContentBrowserModule::StartupModule()
 {
-	UE::Editor::ContentBrowser::Private::SetEnableNewStyleFromCmdLine();
-
 	// Ensure the data module is loaded
 	IContentBrowserDataModule::Get();
 
diff --git a/Engine/Source/Editor/ContentBrowser/Private/ContentBrowserStyle.cpp b/Engine/Source/Editor/ContentBrowser/Private/ContentBrowserStyle.cpp
--- a/Engine/Source/Editor/ContentBrowser/Private/ContentBrowserStyle.cpp
+++ b/Engine/Source/Editor/ContentBrowser/Private/ContentBrowserStyle.cpp
@@ -9,11 +9,6 @@
 #include "Styling/StyleColors.h"
 #include "Styling/ToolBarStyle.h"
 
-namespace UE::Editor::ContentBrowser
-{
-	extern CONTENTBROWSER_API bool GEnableContentBrowserNewStyle;
-}
-
 namespace UE::ContentBrowser::Private
 {
 	FName FContentBrowserStyle::StyleName("ContentBrowser");
@@ -317,7 +312,7 @@ namespace UE::ContentBrowser::Private
 
 		ContentBrowserToolBarStyle.SetSeparatorBrush(FSlateNoResource());
 		ContentBrowserToolBarStyle.SetSeparatorPadding(FMargin(4.0f, 0.0f));
-		if (UE::Editor::ContentBrowser::GEnableContentBrowserNewStyle)
+		if (UE::Editor::ContentBrowser::IsNewStyleEnabled())
 		{
 			ContentBrowserToolBarStyle.SetBackgroundPadding(FMargin(4.0f, 2.0f, 4.0f, 2.0f));
 		}
diff --git a/Engine/Source/Editor/ContentBrowser/Private/ContentBrowserStyle.h b/Engine/Source/Editor/ContentBrowser/Private/ContentBrowserStyle.h
--- a/Engine/Source/Editor/ContentBrowser/Private/ContentBrowserStyle.h
+++ b/Engine/Source/Editor/ContentBrowser/Private/ContentBrowserStyle.h
@@ -6,6 +6,24 @@
 #include "Styling/SlateStyle.h"
 #include "Styling/SlateTypes.h"
 
+namespace UE::Editor::ContentBrowser
+{
+	inline bool IsNewStyleEnabled()
+	{
+		static bool bIsNewStyleEnabled = false;
+		UE_CALL_ONCE([&]()
+		{
+			if (const IConsoleVariable* CVar = IConsoleManager::Get().FindConsoleVariable(TEXT("ContentBrowser.EnableNewStyle")))
+			{
+				ensureAlwaysMsgf(!EnumHasAnyFlags(CVar->GetFlags(), ECVF_Default), TEXT("The CVar should have already been set from commandline, @see: UnrealEdGlobals.cpp, UE::Editor::ContentBrowser::EnableContentBrowserNewStyleCVarRegistration."));
+				bIsNewStyleEnabled = CVar->GetBool();
+			}
+		});
+
+		return bIsNewStyleEnabled;
+	}
+}
+
 namespace UE::ContentBrowser::Private
 {
 	class FContentBrowserStyle
diff --git a/Engine/Source/Editor/ContentBrowser/Private/Menus/AssetViewMenus.cpp b/Engine/Source/Editor/ContentBrowser/Private/Menus/AssetViewMenus.cpp
--- a/Engine/Source/Editor/ContentBrowser/Private/Menus/AssetViewMenus.cpp
+++ b/Engine/Source/Editor/ContentBrowser/Private/Menus/AssetViewMenus.cpp
@@ -2,6 +2,7 @@
 
 #include "ContentBrowserMenuContexts.h"
 #include "ContentBrowserMenuUtils.h"
+#include "ContentBrowserStyle.h"
 #include "Misc/DelayedAutoRegister.h"
 #include "SAssetView.h"
 #include "SContentBrowser.h"
@@ -16,8 +17,6 @@
 
 namespace UE::Editor::ContentBrowser
 {
-	extern CONTENTBROWSER_API bool GEnableContentBrowserNewStyle;
-
 	namespace Private
 	{
 		TSharedRef<SWidget> GetExtendedToolTipMouseWheel()
@@ -78,6 +77,8 @@ FDelayedAutoRegisterHelper SAssetView::AssetViewOptionsMenuRegistration(
 		UToolMenu* Menu = UToolMenus::Get()->RegisterMenu(AssetViewOptionsMenuName);
 		Menu->bCloseSelfOnly = true;
 
+		static bool bIsNewStyleEnabled = UE::Editor::ContentBrowser::IsNewStyleEnabled();
+
 		Menu->AddDynamicSection(NAME_None, FNewToolMenuDelegate::CreateLambda([](UToolMenu* InMenu)
 		{
 			FName OwningContentBrowserName = NAME_None;
@@ -153,7 +154,7 @@ FDelayedAutoRegisterHelper SAssetView::AssetViewOptionsMenuRegistration(
 						);
 					}
 
-					if (!GEnableContentBrowserNewStyle)
+					if (!bIsNewStyleEnabled)
 					{
 						// Filter Bar Section
 						if (TSharedPtr<SFilterList> FilterBarPinned = AssetView->FilterBar.Pin())
@@ -206,7 +207,7 @@ FDelayedAutoRegisterHelper SAssetView::AssetViewOptionsMenuRegistration(
 
 					// Thumbnails Section
 					{
-						const FToolMenuInsert SectionInsert(GEnableContentBrowserNewStyle ? "AssetViewType" : "Search", EToolMenuInsertType::After);
+						const FToolMenuInsert SectionInsert(bIsNewStyleEnabled ? "AssetViewType" : "Search", EToolMenuInsertType::After);
 						FToolMenuSection& Section = InMenu->FindOrAddSection("AssetThumbnails", LOCTEXT("ThumbnailsHeading", "Thumbnails"), SectionInsert);
 
 						auto CreateThumbnailSizeSubMenu = [AssetView](UToolMenu* SubMenu)
@@ -215,7 +216,7 @@ FDelayedAutoRegisterHelper SAssetView::AssetViewOptionsMenuRegistration(
 
 							for (int32 EnumValue = static_cast<int32>(EThumbnailSize::Tiny); EnumValue < static_cast<int32>(EThumbnailSize::MAX); ++EnumValue)
 							{
-								if (!GEnableContentBrowserNewStyle)
+								if (!bIsNewStyleEnabled)
 								{
 									if ((EThumbnailSize)EnumValue == EThumbnailSize::XLarge)
 									{
@@ -238,7 +239,7 @@ FDelayedAutoRegisterHelper SAssetView::AssetViewOptionsMenuRegistration(
 							}
 						};
 
-						if (!GEnableContentBrowserNewStyle)
+						if (!bIsNewStyleEnabled)
 						{
 							Section.AddEntry(FToolMenuEntry::InitSubMenu(
 								"ThumbnailSize",
@@ -290,8 +291,8 @@ FDelayedAutoRegisterHelper SAssetView::AssetViewOptionsMenuRegistration(
 					// Show/View Section
 					{
 						FToolMenuSection& Section = InMenu->FindOrAddSection(
-							GEnableContentBrowserNewStyle ? "Show" : "View",
-							GEnableContentBrowserNewStyle ? LOCTEXT("ShowHeading", "Show") : LOCTEXT("ViewHeading", "View"));
+							bIsNewStyleEnabled ? "Show" : "View",
+							bIsNewStyleEnabled ? LOCTEXT("ShowHeading", "Show") : LOCTEXT("ViewHeading", "View"));
 
 						Section.AddMenuEntry(
 							"ShowFolders",
@@ -332,7 +333,7 @@ FDelayedAutoRegisterHelper SAssetView::AssetViewOptionsMenuRegistration(
 							EUserInterfaceActionType::ToggleButton
 						);
 
-						if (!GEnableContentBrowserNewStyle)
+						if (!bIsNewStyleEnabled)
 						{
 							Section.AddMenuEntry(
 								"FilterRecursively",
@@ -427,7 +428,7 @@ FDelayedAutoRegisterHelper SAssetView::AssetViewOptionsMenuRegistration(
 						);
 					}
 
-					if (GEnableContentBrowserNewStyle)
+					if (bIsNewStyleEnabled)
 					{
 						// Manage Section
 						{
@@ -447,7 +448,7 @@ FDelayedAutoRegisterHelper SAssetView::AssetViewOptionsMenuRegistration(
 						}
 					}
 
-					if (!GEnableContentBrowserNewStyle)
+					if (!bIsNewStyleEnabled)
 					{
 						// Asset Columns Section
 						if (AssetView->GetColumnViewVisibility() == EVisibility::Visible)
@@ -500,7 +501,9 @@ FDelayedAutoRegisterHelper SAssetView::NavigationBarMenuExtensionRegistration(
 			AssetCountSection.InsertPosition = FToolMenuInsert(NAME_None, EToolMenuInsertType::Last);
 			AssetCountSection.Alignment = EToolMenuSectionAlign::Last;
 
-			if (GEnableContentBrowserNewStyle)
+			static bool bIsNewStyleEnabled = UE::Editor::ContentBrowser::IsNewStyleEnabled();
+
+			if (bIsNewStyleEnabled)
 			{
 				AssetCountSection.AddDynamicEntry(
 					"AssetCount",
diff --git a/Engine/Source/Editor/ContentBrowser/Private/Menus/ContentBrowserMenus.cpp b/Engine/Source/Editor/ContentBrowser/Private/Menus/ContentBrowserMenus.cpp
--- a/Engine/Source/Editor/ContentBrowser/Private/Menus/ContentBrowserMenus.cpp
+++ b/Engine/Source/Editor/ContentBrowser/Private/Menus/ContentBrowserMenus.cpp
@@ -1,6 +1,7 @@
 // Copyright Epic Games, Inc. All Rights Reserved.
 
 #include "ContentBrowserMenuContexts.h"
+#include "ContentBrowserStyle.h"
 #include "Misc/DelayedAutoRegister.h"
 #include "SActionButton.h"
 #include "SAssetSearchBox.h"
@@ -13,13 +14,6 @@
 
 #define LOCTEXT_NAMESPACE "ContentBrowser"
 
-namespace UE::Editor::ContentBrowser
-{
-	extern CONTENTBROWSER_API bool GEnableContentBrowserNewStyle;
-}
-
-using namespace UE::Editor::ContentBrowser;
-
 FDelayedAutoRegisterHelper SContentBrowser::NavigationBarMenuRegistration(
 	EDelayedRegisterRunPhase::EndOfEngineInit,
 	[]
@@ -39,7 +33,7 @@ FDelayedAutoRegisterHelper SContentBrowser::NavigationBarMenuRegistration(
 		.AddDynamicEntry(NAME_None,
 			FNewToolMenuSectionDelegate::CreateLambda([](FToolMenuSection& InSection)
 			{
-				if (!GEnableContentBrowserNewStyle)
+				if (!UE::Editor::ContentBrowser::IsNewStyleEnabled())
 				{
 					return;
 				}
@@ -89,7 +83,7 @@ FDelayedAutoRegisterHelper SContentBrowser::NavigationBarMenuRegistration(
 		.AddDynamicEntry(NAME_None, FNewToolMenuSectionDelegate::CreateLambda(
 			[](FToolMenuSection& InSection)
 			{
-				if (!GEnableContentBrowserNewStyle)
+				if (!UE::Editor::ContentBrowser::IsNewStyleEnabled())
 				{
 					return;
 				}
@@ -278,7 +272,7 @@ FDelayedAutoRegisterHelper SContentBrowser::ToolBarMenuRegistration(
 						const TSharedRef<SContentBrowser> ContentBrowser = Context->ContentBrowser.Pin().ToSharedRef();
 
 						TSharedPtr<SWidget> SaveButton = nullptr;
-						if (GEnableContentBrowserNewStyle)
+						if (UE::Editor::ContentBrowser::IsNewStyleEnabled())
 						{
 							SaveButton =
 								SNew(SActionButton)
@@ -335,7 +329,7 @@ FDelayedAutoRegisterHelper SContentBrowser::ToolBarMenuRegistration(
 				}
 			}));
 
-		if (GEnableContentBrowserNewStyle)
+		if (UE::Editor::ContentBrowser::IsNewStyleEnabled())
 		{
 			ToolBar->AddSection("Filter")
 			.AddDynamicEntry(NAME_None, FNewToolMenuSectionDelegate::CreateLambda(
diff --git a/Engine/Source/Editor/ContentBrowser/Private/SAssetView.cpp b/Engine/Source/Editor/ContentBrowser/Private/SAssetView.cpp
--- a/Engine/Source/Editor/ContentBrowser/Private/SAssetView.cpp
+++ b/Engine/Source/Editor/ContentBrowser/Private/SAssetView.cpp
@@ -86,8 +86,6 @@
 	#define checkAssetList(cond)
 #endif
 
-using namespace UE::Editor::ContentBrowser;
-
 namespace UE::AssetView
 {
 	/** Time delay between recently added items being added to the filtered asset items list */
@@ -1112,7 +1110,7 @@ void SAssetView::Construct( const FArguments& InArgs )
 	TileViewThumbnailResolution = 256;
 
 	// Max Size for the thumbnail
-	const int32 MaxTileViewThumbnailSize = GEnableContentBrowserNewStyle ? 160 : 150;
+	const int32 MaxTileViewThumbnailSize = UE::Editor::ContentBrowser::IsNewStyleEnabled() ? 160 : 150;
 
 	TileViewThumbnailSize = MaxTileViewThumbnailSize;
 
@@ -1370,7 +1368,7 @@ void SAssetView::Construct( const FArguments& InArgs )
 	{
 		const TSharedRef<SHorizontalBox> BottomToolBarBox = SNew(SHorizontalBox);
 
-		if (!GEnableContentBrowserNewStyle)
+		if (!UE::Editor::ContentBrowser::IsNewStyleEnabled())
 		{
 			// Asset count
 			BottomToolBarBox->AddSlot()
@@ -1855,7 +1853,7 @@ void SAssetView::LoadSettings(const FString& IniFilename, const FString& IniSect
 		ThumbnailSizeConfig = FMath::Clamp<int32>(ThumbnailSizeConfig, 0, (int32)EThumbnailSize::MAX-1);
 
 		// TODO: Remove this afterwards, current CB should hide new size
-		if (!GEnableContentBrowserNewStyle)
+		if (!UE::Editor::ContentBrowser::IsNewStyleEnabled())
 		{
 			if ((EThumbnailSize)ThumbnailSizeConfig == EThumbnailSize::XLarge)
 			{
@@ -2491,7 +2489,7 @@ FReply SAssetView::OnMouseWheel( const FGeometry& MyGeometry, const FPointerEven
 	// Make sure to not change the thumbnail scaling when we're in Columns view since thumbnail scaling isn't applicable there.
 	if (MouseEvent.IsControlDown() && IsThumbnailScalingAllowed())
 	{
-		if (GEnableContentBrowserNewStyle)
+		if (UE::Editor::ContentBrowser::IsNewStyleEnabled())
 		{
 			if (MouseEvent.IsShiftDown())
 			{
@@ -4833,7 +4831,7 @@ TSharedPtr<FAssetThumbnail> SAssetView::AddItemToNewThumbnailRelevancyMap(const
 		}
 
 		// The thumbnail newly relevant, create a new thumbnail
-		int32 ThumbnailResolution = GEnableContentBrowserNewStyle
+		int32 ThumbnailResolution = UE::Editor::ContentBrowser::IsNewStyleEnabled()
 			? CurrentThumbnailSize
 			: FMath::TruncToInt((float)CurrentThumbnailSize * MaxThumbnailScale);
 
@@ -5332,7 +5330,7 @@ bool SAssetView::IsThumbnailScalingAllowed() const
 
 float SAssetView::GetTileViewTypeNameHeight() const
 {
-	if (GEnableContentBrowserNewStyle)
+	if (UE::Editor::ContentBrowser::IsNewStyleEnabled())
 	{
 		if (ThumbnailSize == EThumbnailSize::Tiny)
 		{
@@ -5382,21 +5380,21 @@ float SAssetView::GetListViewItemHeight() const
 
 float SAssetView::GetTileViewItemHeight() const
 {
-	return GEnableContentBrowserNewStyle
+	return UE::Editor::ContentBrowser::IsNewStyleEnabled()
 		? GetTileViewItemBaseWidth() + GetTileViewTypeNameHeight() + TileViewHeightPadding
 		: (((float)TileViewNameHeight + GetTileViewTypeNameHeight()) * FMath::Lerp(MinThumbnailScale, MaxThumbnailScale, GetThumbnailScale())) + GetTileViewItemBaseHeight() * FillScale + GetSourceControlIconHeight();
 }
 
 float SAssetView::GetTileViewItemBaseHeight() const
 {
-	return GEnableContentBrowserNewStyle
+	return UE::Editor::ContentBrowser::IsNewStyleEnabled()
 		? GetTileViewItemBaseWidth()
 		: (float)(TileViewThumbnailSize + TileViewThumbnailPadding * 2) * FMath::Lerp(MinThumbnailScale, MaxThumbnailScale, GetThumbnailScale());
 }
 
 float SAssetView::GetTileViewItemWidth() const
 {
-	return GEnableContentBrowserNewStyle
+	return UE::Editor::ContentBrowser::IsNewStyleEnabled()
 		? GetTileViewItemBaseWidth() + TileViewWidthPadding
 		: GetTileViewItemBaseWidth() * FillScale;
 }
@@ -5408,7 +5406,7 @@ float SAssetView::GetTileViewThumbnailDimension() const
 
 float SAssetView::GetTileViewItemBaseWidth() const //-V524
 {
-	return GEnableContentBrowserNewStyle
+	return UE::Editor::ContentBrowser::IsNewStyleEnabled()
 		? GetTileViewThumbnailDimension()
 		: (float)( TileViewThumbnailSize + TileViewThumbnailPadding * 2 ) * FMath::Lerp( MinThumbnailScale, MaxThumbnailScale, GetThumbnailScale() );
 }
diff --git a/Engine/Source/Editor/ContentBrowser/Private/SContentBrowser.cpp b/Engine/Source/Editor/ContentBrowser/Private/SContentBrowser.cpp
--- a/Engine/Source/Editor/ContentBrowser/Private/SContentBrowser.cpp
+++ b/Engine/Source/Editor/ContentBrowser/Private/SContentBrowser.cpp
@@ -165,8 +165,6 @@ struct FSlateBrush;
 
 const FString SContentBrowser::SettingsIniSection = TEXT("ContentBrowser");
 
-using namespace UE::Editor::ContentBrowser;
-
 namespace ContentBrowserSourcesWidgetSwitcherIndex
 {
 	static constexpr int32 PathView = 0;
@@ -490,7 +488,7 @@ void SContentBrowser::Construct( const FArguments& InArgs, const FName& InInstan
 	// Note, for backwards compatibility ShowBottomToolbar controls the visibility of view options so we respect that here
 	if (bShowBottomToolbar)
 	{
-		if (GEnableContentBrowserNewStyle)
+		if (UE::Editor::ContentBrowser::IsNewStyleEnabled())
 		{
 			ViewOptions =
 				SNew(SActionButton)
@@ -530,7 +528,7 @@ void SContentBrowser::Construct( const FArguments& InArgs, const FName& InInstan
 		}
 	}
 
-	if (GEnableContentBrowserNewStyle)
+	if (UE::Editor::ContentBrowser::IsNewStyleEnabled())
 	{
 		constexpr float ToolBarVerticalPadding = 4.0f; // ToolBar total height should be 36.0f
 		constexpr float PanelInsetPadding = 2.0f;
@@ -1225,7 +1223,7 @@ TSharedRef<SWidget> SContentBrowser::CreateLockButton(const FContentBrowserConfi
 
 void SContentBrowser::OnFilterBarLayoutChanging(EFilterBarLayout NewLayout)
 {
-	if (GEnableContentBrowserNewStyle)
+	if (UE::Editor::ContentBrowser::IsNewStyleEnabled())
 	{
 		if (NewLayout == EFilterBarLayout::Horizontal)
 		{
@@ -1509,7 +1507,7 @@ TSharedRef<SWidget> SContentBrowser::CreateAssetView(const FContentBrowserConfig
 				   .Visibility((Config != nullptr ? Config->bCanShowAssetSearch : true) ? EVisibility::Visible : EVisibility::Collapsed)
 				   .AddMetaData<FTagMetaData>(FTagMetaData(TEXT("ContentBrowserSearchAssets")));
 
-	if (GEnableContentBrowserNewStyle)
+	if (UE::Editor::ContentBrowser::IsNewStyleEnabled())
 	{
 		NavigationBar = SNew(SNavigationBar)
 						.OnPathClicked(this, &SContentBrowser::OnPathClicked)
@@ -1526,7 +1524,7 @@ TSharedRef<SWidget> SContentBrowser::CreateAssetView(const FContentBrowserConfig
 	 * which is called initially called through SAssetFilterBar::LoadSettings through SContentBrowser::LoadSettings()
 	 */
 	const FMargin AssetViewPadding =
-		GEnableContentBrowserNewStyle
+		UE::Editor::ContentBrowser::IsNewStyleEnabled()
 			? FMargin(2.0f, 0.0f, 2.0f, 0.0f)
 			: FMargin(2.0f, 2.0f, 2.0f, 0.0f);
 
@@ -1833,7 +1831,7 @@ void SContentBrowser::ExtendViewOptionsMenu(const FContentBrowserConfig* Config)
 				{
 					if (bShowLockButton)
 					{
-						if (GEnableContentBrowserNewStyle)
+						if (UE::Editor::ContentBrowser::IsNewStyleEnabled())
 						{
 							FToolMenuSection& Section = InMenu->FindOrAddSection("Manage");
 							Section.AddMenuEntry(
@@ -1863,10 +1861,10 @@ void SContentBrowser::ExtendViewOptionsMenu(const FContentBrowserConfig* Config)
 
 					if (bShowSourcesView)
 					{
-						FToolMenuSection& Section = InMenu->FindOrAddSection(GEnableContentBrowserNewStyle ? "Show" : "View");
+						FToolMenuSection& Section = InMenu->FindOrAddSection(UE::Editor::ContentBrowser::IsNewStyleEnabled() ? "Show" : "View");
 						Section.AddMenuEntry(
 							"ToggleSources",
-							GEnableContentBrowserNewStyle
+							UE::Editor::ContentBrowser::IsNewStyleEnabled()
 							? LOCTEXT("ToggleSourcesView_NewStyle", "Sources Panel")
 							: LOCTEXT("ToggleSourcesView", "Show Sources Panel"),
 							LOCTEXT("ToggleSourcesView_Tooltip", "Show or hide the sources panel"),
diff --git a/Engine/Source/Editor/ContentBrowser/Public/ContentBrowserModule.h b/Engine/Source/Editor/ContentBrowser/Public/ContentBrowserModule.h
--- a/Engine/Source/Editor/ContentBrowser/Public/ContentBrowserModule.h
+++ b/Engine/Source/Editor/ContentBrowser/Public/ContentBrowserModule.h
@@ -25,11 +25,6 @@ class IContentBrowserSingleton;
 struct FARFilter;
 struct FAssetData;
 
-namespace UE::Editor::ContentBrowser
-{
-	CONTENTBROWSER_API extern bool GEnableContentBrowserNewStyle;
-}
-
 /** Extra state generator that adds an icon and a corresponding legend entry on an asset. */
 class FAssetViewExtraStateGenerator
 {
diff --git a/Engine/Source/Editor/UMGEditor/Private/Library/SLibraryView.cpp b/Engine/Source/Editor/UMGEditor/Private/Library/SLibraryView.cpp
--- a/Engine/Source/Editor/UMGEditor/Private/Library/SLibraryView.cpp
+++ b/Engine/Source/Editor/UMGEditor/Private/Library/SLibraryView.cpp
@@ -37,10 +37,21 @@
 
 namespace UE::Editor::ContentBrowser
 {
-	extern CONTENTBROWSER_API bool GEnableContentBrowserNewStyle;
-}
+	inline bool IsNewStyleEnabled()
+	{
+		static bool bIsNewStyleEnabled = false;
+		UE_CALL_ONCE([&]()
+		{
+			if (const IConsoleVariable* CVar = IConsoleManager::Get().FindConsoleVariable(TEXT("ContentBrowser.EnableNewStyle")))
+			{
+				ensureAlwaysMsgf(!EnumHasAnyFlags(CVar->GetFlags(), ECVF_Default), TEXT("The CVar should have already been set from commandline, @see: UnrealEdGlobals.cpp, UE::Editor::ContentBrowser::EnableContentBrowserNewStyleCVarRegistration."));
+				bIsNewStyleEnabled = CVar->GetBool();
+			}
+		});
 
-using namespace UE::Editor::ContentBrowser;
+		return bIsNewStyleEnabled;
+	}
+}
 
 FText SLibraryViewItem::GetFavoriteToggleToolTipText() const
 {
@@ -163,7 +174,7 @@ TSharedRef<SWidget> SLibraryView::ConstructViewOptions()
 		{
 			for (int32 EnumValue = (int32)EThumbnailSize::Tiny; EnumValue < (int32)EThumbnailSize::MAX; ++EnumValue)
 			{
-				if (GEnableContentBrowserNewStyle)
+				if (UE::Editor::ContentBrowser::IsNewStyleEnabled())
 				{
 					if ((EThumbnailSize)EnumValue == EThumbnailSize::XLarge)
 					{
diff --git a/Engine/Source/Editor/UnrealEd/Private/AssetThumbnail.cpp b/Engine/Source/Editor/UnrealEd/Private/AssetThumbnail.cpp
--- a/Engine/Source/Editor/UnrealEd/Private/AssetThumbnail.cpp
+++ b/Engine/Source/Editor/UnrealEd/Private/AssetThumbnail.cpp
@@ -54,6 +54,23 @@
 #include "Widgets/Layout/SSeparator.h"
 #include "Widgets/Text/SRichTextBlock.h"
 
+namespace UE::Editor::ContentBrowser
+{
+	inline bool IsNewStyleEnabled()
+	{
+		static bool bIsNewStyleEnabled = false;
+		UE_CALL_ONCE([&]()
+		{
+			if (const IConsoleVariable* CVar = IConsoleManager::Get().FindConsoleVariable(TEXT("ContentBrowser.EnableNewStyle")))
+			{
+				ensureAlwaysMsgf(!EnumHasAnyFlags(CVar->GetFlags(), ECVF_Default), TEXT("The CVar should have already been set from commandline, @see: UnrealEdGlobals.cpp, UE::Editor::ContentBrowser::EnableContentBrowserNewStyleCVarRegistration."));
+				bIsNewStyleEnabled = CVar->GetBool();
+			}
+		});
+
+		return bIsNewStyleEnabled;
+	}
+}
 namespace AssetThumbnailPool
 {
 	const FObjectThumbnail* LoadThumbnailsFromPackage(const FAssetData& AssetData, FThumbnailMap& OutThumbnailMap)
@@ -170,10 +187,7 @@ public:
 
 		TSharedRef<SOverlay> OverlayWidget = SNew(SOverlay);
 
-		static const IConsoleVariable* CVarEnableContentBrowserNewStyle = IConsoleManager::Get().FindConsoleVariable(TEXT("ContentBrowser.EnableNewStyle"));
-		const bool bEnableContentBrowserNewStyle = CVarEnableContentBrowserNewStyle && CVarEnableContentBrowserNewStyle->GetBool();
-
-		if (bEnableContentBrowserNewStyle)
+		if (UE::Editor::ContentBrowser::IsNewStyleEnabled())
 		{
 			if (InArgs._TooltipOverride.IsSet())
 			{
@@ -305,8 +319,7 @@ public:
 		}
 
 		TSharedRef<SWidget> ContentWidget = OverlayWidget;
-
-		if (bEnableContentBrowserNewStyle)
+		if (UE::Editor::ContentBrowser::IsNewStyleEnabled())
 		{
 			if (bShowAssetChip)
 			{
@@ -361,7 +374,7 @@ public:
 			];
 		}
 
-		if (bEnableContentBrowserNewStyle)
+		if (UE::Editor::ContentBrowser::IsNewStyleEnabled())
 		{
 			// AssetEditMode, do not create it if there is no config argument to show it
 			if (EditModeVisibility.IsSet())
@@ -377,13 +390,13 @@ public:
 
 		const UAssetDefinition* AssetDefinition = UAssetDefinitionRegistry::Get()->GetAssetDefinitionForClass(Class);
 
-		if (bEnableContentBrowserNewStyle)
+		if (!AssetDefinition)
 		{
-			if (!AssetDefinition)
-			{
-				AssetDefinition = GetDefault<UAssetDefinitionDefault>();
-			}
+			AssetDefinition = GetDefault<UAssetDefinitionDefault>();
+		}
 
+		if (UE::Editor::ContentBrowser::IsNewStyleEnabled())
+		{
 			if (InArgs._AllowAssetStatusThumbnailOverlay && AssetDefinition)
 			{
 				const TSharedPtr<FAssetStatusAssetDataInfoProvider> AssetDataInfoProvider = MakeShared<FAssetStatusAssetDataInfoProvider>(AssetData);
@@ -439,7 +452,8 @@ public:
 		if( InArgs._AllowAssetSpecificThumbnailOverlay && AssetTypeActions.IsValid() )
 		{
 			FAssetActionThumbnailOverlayInfo OutThumbnailInfo;
-			if (bEnableContentBrowserNewStyle && AssetDefinition && AssetDefinition->GetThumbnailActionOverlay(AssetData, OutThumbnailInfo))
+			if (UE::Editor::ContentBrowser::IsNewStyleEnabled()
+				&& AssetDefinition && AssetDefinition->GetThumbnailActionOverlay(AssetData, OutThumbnailInfo))
 			{
 				constexpr int32 OverlayZOrder = 1;
 				if (OutThumbnailInfo.ActionImageWidget.IsValid())
@@ -542,7 +556,7 @@ public:
 			}
 		}
 
-		if (bEnableContentBrowserNewStyle)
+		if (UE::Editor::ContentBrowser::IsNewStyleEnabled())
 		{
 			ChildSlot
 			[
@@ -1059,11 +1073,8 @@ private:
 			AssetColor = AssetTypeActions->GetTypeColor();
 		}
 
-		static const IConsoleVariable* CVarEnableContentBrowserNewStyle = IConsoleManager::Get().FindConsoleVariable(TEXT("ContentBrowser.EnableNewStyle"));
-		const bool bEnableContentBrowserNewStyle = CVarEnableContentBrowserNewStyle && CVarEnableContentBrowserNewStyle->GetBool();
-
 		//AssetBackgroundWidget->SetBorderBackgroundColor(AssetColor.CopyWithNewOpacity(0.3f));
-		if (bEnableContentBrowserNewStyle)
+		if (UE::Editor::ContentBrowser::IsNewStyleEnabled())
 		{
 			if (LayeredImage.IsValid())
 			{
@@ -1483,10 +1494,7 @@ bool SAssetThumbnailToolTip::IsEmpty() const
 
 void SAssetThumbnailToolTip::OnOpening()
 {
-	static const IConsoleVariable* CVarEnableContentBrowserNewStyle = IConsoleManager::Get().FindConsoleVariable(TEXT("ContentBrowser.EnableNewStyle"));
-	const bool bEnableContentBrowserNewStyle = CVarEnableContentBrowserNewStyle && CVarEnableContentBrowserNewStyle->GetBool();
-
-	if (bEnableContentBrowserNewStyle)
+	if (UE::Editor::ContentBrowser::IsNewStyleEnabled())
 	{
 		TSharedPtr<SAssetThumbnail> AssetViewItemPin = AssetThumbnail.Pin();
 		if (AssetViewItemPin.IsValid())
@@ -1600,30 +1608,54 @@ void FAssetThumbnail::SetAsset( const FAssetData& InAssetData )
 
 TSharedRef<SWidget> FAssetThumbnail::MakeThumbnailWidget( const FAssetThumbnailConfig& InConfig )
 {
-	return
-		SNew(SAssetThumbnail)
-		.AssetThumbnail(SharedThis(this))
-		.ThumbnailPool(ThumbnailPool.Pin())
-		.AllowFadeIn(InConfig.bAllowFadeIn)
-		.ForceGenericThumbnail(InConfig.bForceGenericThumbnail)
-		.Label(InConfig.ThumbnailLabel)
-		.HighlightedText(InConfig.HighlightedText)
-		.HintColorAndOpacity(InConfig.HintColorAndOpacity)
-		.AllowHintText(InConfig.bAllowHintText)
-		.AllowRealTimeOnHovered(InConfig.bAllowRealTimeOnHovered)
-		.ClassThumbnailBrushOverride(InConfig.ClassThumbnailBrushOverride)
-		.AllowAssetSpecificThumbnailOverlay(InConfig.bAllowAssetSpecificThumbnailOverlay)
-		.AssetTypeColorOverride(InConfig.AssetTypeColorOverride)
-		.Padding(InConfig.Padding)
-		.GenericThumbnailSize(InConfig.GenericThumbnailSize)
-		.AssetSystemInfoProvider(InConfig.AssetSystemInfoProvider)
-		.TooltipOverride(InConfig.TooltipOverride)
-		.AllowAssetStatusThumbnailOverlay(InConfig.bAllowAssetStatusThumbnailOverlay)
-		.ShowAssetChip(InConfig.bShowAssetChip)
-		.AssetChipBorderImageOverride(InConfig.AssetChipBorderImageOverride)
-		.CanDisplayEditModePrimitiveTools(InConfig.bCanDisplayEditModePrimitiveTools)
-		.IsEditModeVisible(InConfig.IsEditModeVisible)
-		.ColorStripOrientation(InConfig.ColorStripOrientation);
+	TSharedPtr<SWidget> ThumbnailWidget = nullptr;
+	if (UE::Editor::ContentBrowser::IsNewStyleEnabled())
+	{
+		ThumbnailWidget =
+			SNew(SAssetThumbnail)
+			.AssetThumbnail(SharedThis(this))
+			.ThumbnailPool(ThumbnailPool.Pin())
+			.AllowFadeIn(InConfig.bAllowFadeIn)
+			.ForceGenericThumbnail(InConfig.bForceGenericThumbnail)
+			.Label(InConfig.ThumbnailLabel)
+			.HighlightedText(InConfig.HighlightedText)
+			.HintColorAndOpacity(InConfig.HintColorAndOpacity)
+			.AllowHintText(InConfig.bAllowHintText)
+			.AllowRealTimeOnHovered(InConfig.bAllowRealTimeOnHovered)
+			.ClassThumbnailBrushOverride(InConfig.ClassThumbnailBrushOverride)
+			.AllowAssetSpecificThumbnailOverlay(InConfig.bAllowAssetSpecificThumbnailOverlay)
+			.AssetTypeColorOverride(InConfig.AssetTypeColorOverride)
+			.Padding(InConfig.Padding)
+			.GenericThumbnailSize(InConfig.GenericThumbnailSize)
+			.AssetSystemInfoProvider(InConfig.AssetSystemInfoProvider)
+			.TooltipOverride(InConfig.TooltipOverride)
+			.AllowAssetStatusThumbnailOverlay(InConfig.bAllowAssetStatusThumbnailOverlay)
+			.ShowAssetChip(InConfig.bShowAssetChip)
+			.AssetChipBorderImageOverride(InConfig.AssetChipBorderImageOverride)
+			.ColorStripOrientation(InConfig.ColorStripOrientation);
+	}
+	else
+	{
+		ThumbnailWidget =
+			SNew(SAssetThumbnail)
+			.AssetThumbnail(SharedThis(this))
+			.ThumbnailPool(ThumbnailPool.Pin())
+			.AllowFadeIn(InConfig.bAllowFadeIn)
+			.ForceGenericThumbnail(InConfig.bForceGenericThumbnail)
+			.Label(InConfig.ThumbnailLabel)
+			.HighlightedText(InConfig.HighlightedText)
+			.HintColorAndOpacity(InConfig.HintColorAndOpacity)
+			.AllowHintText(InConfig.bAllowHintText)
+			.AllowRealTimeOnHovered(InConfig.bAllowRealTimeOnHovered)
+			.ClassThumbnailBrushOverride(InConfig.ClassThumbnailBrushOverride)
+			.AllowAssetSpecificThumbnailOverlay(InConfig.bAllowAssetSpecificThumbnailOverlay)
+			.AssetTypeColorOverride(InConfig.AssetTypeColorOverride)
+			.Padding(InConfig.Padding)
+			.GenericThumbnailSize(InConfig.GenericThumbnailSize)
+			.ColorStripOrientation(InConfig.ColorStripOrientation);
+	}
+
+	return ThumbnailWidget.ToSharedRef();
 }
 
 void FAssetThumbnail::RefreshThumbnail()
diff --git a/Engine/Source/Editor/UnrealEd/Private/DragAndDrop/AssetDragDropOp.cpp b/Engine/Source/Editor/UnrealEd/Private/DragAndDrop/AssetDragDropOp.cpp
--- a/Engine/Source/Editor/UnrealEd/Private/DragAndDrop/AssetDragDropOp.cpp
+++ b/Engine/Source/Editor/UnrealEd/Private/DragAndDrop/AssetDragDropOp.cpp
@@ -11,6 +11,24 @@
 #include "ThumbnailRendering/ThumbnailManager.h"
 #include "Engine/Level.h"
 
+namespace UE::Editor::ContentBrowser
+{
+	inline bool IsNewStyleEnabled()
+	{
+		static bool bIsNewStyleEnabled = false;
+		UE_CALL_ONCE([&]()
+		{
+			if (const IConsoleVariable* CVar = IConsoleManager::Get().FindConsoleVariable(TEXT("ContentBrowser.EnableNewStyle")))
+			{
+				ensureAlwaysMsgf(!EnumHasAnyFlags(CVar->GetFlags(), ECVF_Default), TEXT("The CVar should have already been set from commandline, @see: UnrealEdGlobals.cpp, UE::Editor::ContentBrowser::EnableContentBrowserNewStyleCVarRegistration."));
+				bIsNewStyleEnabled = CVar->GetBool();
+			}
+		});
+
+		return bIsNewStyleEnabled;
+	}
+}
+
 TSharedRef<FAssetDragDropOp> FAssetDragDropOp::New(const FAssetData& InAssetData, TScriptInterface<IAssetFactoryInterface> AssetFactory)
 {
 	TArray<FAssetData> AssetDataArray;
@@ -85,7 +103,7 @@ TSharedPtr<SWidget> FAssetDragDropOp::GetDefaultDecorator() const
 	}
 	else if (AssetThumbnail.IsValid())
 	{
-		if (0)
+		if (UE::Editor::ContentBrowser::IsNewStyleEnabled())
 		{
 			FAssetThumbnailConfig AssetThumbnailConfig;
 			AssetThumbnailConfig.bShowAssetChip = true;
diff --git a/Engine/Source/Editor/UnrealEd/Private/UnrealEdGlobals.cpp b/Engine/Source/Editor/UnrealEd/Private/UnrealEdGlobals.cpp
--- a/Engine/Source/Editor/UnrealEd/Private/UnrealEdGlobals.cpp
+++ b/Engine/Source/Editor/UnrealEd/Private/UnrealEdGlobals.cpp
@@ -55,6 +55,64 @@ bool GLevelEditorModeToolsIsValid()
 	return Internal::EditorModeToolsSingleton.IsValid();
 }
 
+// @todo: Can remove after permanent switch to new style
+namespace UE::Editor::ContentBrowser
+{
+	namespace Private
+	{
+		static FAutoConsoleVariable CVarEnableContentBrowserNewStyle(
+			TEXT("ContentBrowser.EnableNewStyle"),
+			false,
+			TEXT("Whether or not to enable the Content Browser restyle (~5.6).")
+		);
+
+		void SetEnableNewStyleFromCmdLine()
+		{
+			static constexpr const TCHAR* CVarName = TEXT("ContentBrowser.EnableNewStyle");
+
+			const FString OriginalCmdLine = FCommandLine::GetOriginal();
+			if (const int32 FoundIdx = OriginalCmdLine.Find(CVarName);
+				FoundIdx != INDEX_NONE)
+			{
+				// If present, enable
+				bool bEnableNewStyle = true;
+
+				// Check for explicit disable =0 or \s0
+				for (int32 CharIdx = FoundIdx + FCString::Strlen(CVarName); CharIdx < OriginalCmdLine.Len(); ++CharIdx)
+				{
+					if (OriginalCmdLine[CharIdx] == TEXT(' ')
+						|| OriginalCmdLine[CharIdx] == TEXT('='))
+					{
+						continue;
+					}
+
+					if (FChar::IsDigit(OriginalCmdLine[CharIdx]))
+					{
+						const int32 Value = FChar::ConvertCharDigitToInt(OriginalCmdLine[CharIdx]);
+						// Explicitly disable
+						if (Value == 0)
+						{
+							bEnableNewStyle = false;
+							break;
+						}
+					}
+
+					break;
+				}
+
+				CVarEnableContentBrowserNewStyle->Set(bEnableNewStyle);
+			}
+		}
+
+		static FDelayedAutoRegisterHelper EnableContentBrowserNewStyleCVarRegistration(
+			EDelayedRegisterRunPhase::FileSystemReady,
+			[]()
+			{
+				SetEnableNewStyleFromCmdLine();
+			});
+	}
+}
+
 FEditorModeTools& GLevelEditorModeTools()
 {
 	checkf(!IsRunningCommandlet(), TEXT("The global mode manager should not be created or accessed in a commandlet environment. Check that your mode or module is not accessing the global mode tools or that scriptable features of modes have been moved to subsystems."));
