[ContentBrowser] Updated (new style) view options menu

diff --git a/Engine/Source/Editor/ContentBrowser/Private/ContentBrowserMenuUtils.cpp b/Engine/Source/Editor/ContentBrowser/Private/ContentBrowserMenuUtils.cpp
--- a/Engine/Source/Editor/ContentBrowser/Private/ContentBrowserMenuUtils.cpp
+++ b/Engine/Source/Editor/ContentBrowser/Private/ContentBrowserMenuUtils.cpp
@@ -211,7 +211,7 @@ void ContentBrowserMenuUtils::AddFiltersToMenu(UToolMenu* InMenu, const FName& I
 		ContentMenuInsert.Position = EToolMenuInsertType::After;
 	}
 
-	FToolMenuSection& Section = UE::Editor::ContentBrowser::IsNewStyleEnabled() ? InMenu->FindOrAddSection("Show") : InMenu->FindOrAddSection("Content", LOCTEXT("ContentHeading", "Content"), ContentMenuInsert);
+	FToolMenuSection& Section = InMenu->FindOrAddSection("Content", LOCTEXT("ContentHeading", "Content"), ContentMenuInsert);
 
 	// New style omits the "Show" prefix
 	const FText ShowCppLabel = UE::Editor::ContentBrowser::IsNewStyleEnabled() ?  LOCTEXT("ShowCppClassesOption_NewStyle", "C++ Classes") : LOCTEXT("ShowCppClassesOption", "Show C++ Classes");
diff --git a/Engine/Source/Editor/ContentBrowser/Private/Menus/AssetViewMenus.cpp b/Engine/Source/Editor/ContentBrowser/Private/Menus/AssetViewMenus.cpp
--- a/Engine/Source/Editor/ContentBrowser/Private/Menus/AssetViewMenus.cpp
+++ b/Engine/Source/Editor/ContentBrowser/Private/Menus/AssetViewMenus.cpp
@@ -77,8 +77,6 @@ FDelayedAutoRegisterHelper SAssetView::AssetViewOptionsMenuRegistration(
 		UToolMenu* Menu = UToolMenus::Get()->RegisterMenu(AssetViewOptionsMenuName);
 		Menu->bCloseSelfOnly = true;
 
-		static bool bIsNewStyleEnabled = UE::Editor::ContentBrowser::IsNewStyleEnabled();
-
 		Menu->AddDynamicSection(NAME_None, FNewToolMenuDelegate::CreateLambda([](UToolMenu* InMenu)
 		{
 			FName OwningContentBrowserName = NAME_None;
@@ -87,9 +85,29 @@ FDelayedAutoRegisterHelper SAssetView::AssetViewOptionsMenuRegistration(
 			{
 				if (Context->AssetView.IsValid())
 				{
+					auto SetEntryOrder = [](FToolMenuEntry& InEntry)
+					{
+						// New Style, >5.6: (entry name, entry to place after) map. If the named entry isn't found, the entry will not have an explicit order.
+	                    static const TMap<FName, FName> NewEntryOrder = {
+                    		{ "ShowFolders", "ShowFavorite" },
+                    		{ "ShowEmptyFolders", "ShowFolders" },
+                    		{ "OrganizeFolders", "ShowEmptyFolders" }
+	                    };
+
+						// <5.6 Style: (entry name, entry to place after) map. If the named entry isn't found, the entry will not have an explicit order.
+						static const TMap<FName, FName> EntryOrder = {
+							{ "ShowAllFolder", "FilterRecursively" }
+						};
+
+						if (const FName* InsertAfterEntry = IsNewStyleEnabled() ? NewEntryOrder.Find(InEntry.Name) : EntryOrder.Find(InEntry.Name))
+						{
+							InEntry.InsertPosition = FToolMenuInsert(*InsertAfterEntry, EToolMenuInsertType::After);
+						}
+					};
+
 					const TSharedRef<SAssetView> AssetView = Context->AssetView.Pin().ToSharedRef();
 
-						// ViewType Section
+					// ViewType Section
 					{
 						FToolMenuSection& Section = InMenu->FindOrAddSection("AssetViewType", LOCTEXT("ViewTypeHeading", "View Type"));
 						Section.AddMenuEntry(
@@ -154,7 +172,7 @@ FDelayedAutoRegisterHelper SAssetView::AssetViewOptionsMenuRegistration(
 						);
 					}
 
-					if (!bIsNewStyleEnabled)
+					if (!IsNewStyleEnabled())
 					{
 						// Filter Bar Section
 						if (TSharedPtr<SFilterList> FilterBarPinned = AssetView->FilterBar.Pin())
@@ -207,7 +225,7 @@ FDelayedAutoRegisterHelper SAssetView::AssetViewOptionsMenuRegistration(
 
 					// Thumbnails Section
 					{
-						const FToolMenuInsert SectionInsert(bIsNewStyleEnabled ? "AssetViewType" : "Search", EToolMenuInsertType::After);
+						const FToolMenuInsert SectionInsert("Search", EToolMenuInsertType::After);
 						FToolMenuSection& Section = InMenu->FindOrAddSection("AssetThumbnails", LOCTEXT("ThumbnailsHeading", "Thumbnails"), SectionInsert);
 
 						auto CreateThumbnailSizeSubMenu = [AssetView](UToolMenu* SubMenu)
@@ -216,7 +234,7 @@ FDelayedAutoRegisterHelper SAssetView::AssetViewOptionsMenuRegistration(
 
 							for (int32 EnumValue = static_cast<int32>(EThumbnailSize::Tiny); EnumValue < static_cast<int32>(EThumbnailSize::MAX); ++EnumValue)
 							{
-								if (!bIsNewStyleEnabled)
+								if (!IsNewStyleEnabled())
 								{
 									if ((EThumbnailSize)EnumValue == EThumbnailSize::XLarge)
 									{
@@ -239,90 +257,82 @@ FDelayedAutoRegisterHelper SAssetView::AssetViewOptionsMenuRegistration(
 							}
 						};
 
-						if (!bIsNewStyleEnabled)
-						{
-							Section.AddEntry(FToolMenuEntry::InitSubMenu(
-								"ThumbnailSize",
-								LOCTEXT("ThumbnailSize", "Thumbnail Size"),
-								LOCTEXT("ThumbnailSizeToolTip", "Adjust the size of thumbnails."),
-								FNewToolMenuDelegate::CreateLambda(CreateThumbnailSizeSubMenu)
-							));
-						}
-						else
-						{
-							Section.AddEntry(FToolMenuEntry::InitSubMenu(
-								"ThumbnailSize",
-								FToolUIActionChoice(),
-								Private::GetThumbnailSizeWidget(),
-								FNewToolMenuDelegate::CreateLambda(CreateThumbnailSizeSubMenu)
-							));
-						}
-
-						// Thumbnail EditMode Section
-						{
-							Section.AddMenuEntry(
-								"ThumbnailEditMode",
-								LOCTEXT("ThumbnailEditModeOption", "Thumbnail Edit Mode"),
-								LOCTEXT("ThumbnailEditModeOptionToolTip", "Toggle thumbnail editing mode. When in this mode you can rotate the camera on 3D thumbnails by dragging them."),
-								FSlateIcon(),
-								FUIAction(
-									FExecuteAction::CreateSP(AssetView, &SAssetView::ToggleThumbnailEditMode),
-									FCanExecuteAction::CreateSP(AssetView, &SAssetView::IsThumbnailEditModeAllowed),
-									FIsActionChecked::CreateSP(AssetView, &SAssetView::IsThumbnailEditMode)
-								),
-								EUserInterfaceActionType::ToggleButton
-							);
-
-							Section.AddMenuEntry(
-								"RealTimeThumbnails",
-								LOCTEXT("RealTimeThumbnailsOption", "Real-Time Thumbnails"),
-								LOCTEXT("RealTimeThumbnailsOptionToolTip", "Renders the assets thumbnails in real-time"),
-								FSlateIcon(),
-								FUIAction(
-									FExecuteAction::CreateSP(AssetView, &SAssetView::ToggleRealTimeThumbnails),
-									FCanExecuteAction::CreateSP(AssetView, &SAssetView::CanShowRealTimeThumbnails),
-									FIsActionChecked::CreateSP(AssetView, &SAssetView::IsShowingRealTimeThumbnails)
-								),
-								EUserInterfaceActionType::ToggleButton
-							);
-						}
-					}
-
-					// Show/View Section
-					{
-						FToolMenuSection& Section = InMenu->FindOrAddSection(
-							bIsNewStyleEnabled ? "Show" : "View",
-							bIsNewStyleEnabled ? LOCTEXT("ShowHeading", "Show") : LOCTEXT("ViewHeading", "View"));
+						Section.AddEntry(FToolMenuEntry::InitSubMenu(
+							"ThumbnailSize",
+							LOCTEXT("ThumbnailSize", "Thumbnail Size"),
+							LOCTEXT("ThumbnailSizeToolTip", "Adjust the size of thumbnails."),
+							FNewToolMenuDelegate::CreateLambda(CreateThumbnailSizeSubMenu)
+						));
 
 						Section.AddMenuEntry(
-							"ShowFolders",
-							LOCTEXT("ShowFoldersOption", "Folders"),
-							LOCTEXT("ShowFoldersOptionToolTip", "Show folders in the view as well as assets?"),
+							"ThumbnailEditMode",
+							LOCTEXT("ThumbnailEditModeOption", "Thumbnail Edit Mode"),
+							LOCTEXT("ThumbnailEditModeOptionToolTip", "Toggle thumbnail editing mode. When in this mode you can rotate the camera on 3D thumbnails by dragging them."),
 							FSlateIcon(),
 							FUIAction(
-								FExecuteAction::CreateSP(AssetView, &SAssetView::ToggleShowFolders),
-								FCanExecuteAction::CreateSP(AssetView, &SAssetView::IsToggleShowFoldersAllowed),
-								FIsActionChecked::CreateSP(AssetView, &SAssetView::IsShowingFolders)
+								FExecuteAction::CreateSP(AssetView, &SAssetView::ToggleThumbnailEditMode),
+								FCanExecuteAction::CreateSP(AssetView, &SAssetView::IsThumbnailEditModeAllowed),
+								FIsActionChecked::CreateSP(AssetView, &SAssetView::IsThumbnailEditMode)
 							),
 							EUserInterfaceActionType::ToggleButton
 						);
 
 						Section.AddMenuEntry(
-							"ShowEmptyFolders",
-							LOCTEXT("ShowEmptyFoldersOption", "Empty Folders"),
-							LOCTEXT("ShowEmptyFoldersOptionToolTip", "Show empty folders in the view as well as assets?"),
+							"RealTimeThumbnails",
+							LOCTEXT("RealTimeThumbnailsOption", "Real-Time Thumbnails"),
+							LOCTEXT("RealTimeThumbnailsOptionToolTip", "Renders the assets thumbnails in real-time"),
 							FSlateIcon(),
 							FUIAction(
-								FExecuteAction::CreateSP(AssetView, &SAssetView::ToggleShowEmptyFolders),
-								FCanExecuteAction::CreateSP(AssetView, &SAssetView::IsToggleShowEmptyFoldersAllowed),
-								FIsActionChecked::CreateSP(AssetView, &SAssetView::IsShowingEmptyFolders)
+								FExecuteAction::CreateSP(AssetView, &SAssetView::ToggleRealTimeThumbnails),
+								FCanExecuteAction::CreateSP(AssetView, &SAssetView::CanShowRealTimeThumbnails),
+								FIsActionChecked::CreateSP(AssetView, &SAssetView::IsShowingRealTimeThumbnails)
 							),
 							EUserInterfaceActionType::ToggleButton
 						);
+					}
+
+					// Show/View Section
+					{
+						FToolMenuSection& Section = InMenu->FindOrAddSection(
+							IsNewStyleEnabled() ? "Show" : "View",
+							IsNewStyleEnabled() ? LOCTEXT("ShowHeading", "Show") : LOCTEXT("ViewHeading", "View"));
+
+						const FText ShowFoldersLabel = IsNewStyleEnabled() ? LOCTEXT("ShowFoldersOption_NewStyle", "Folders") :LOCTEXT("ShowFoldersOption", "Show Folders");
+						SetEntryOrder(
+							Section.AddMenuEntry(
+								"ShowFolders",
+								ShowFoldersLabel,
+								LOCTEXT("ShowFoldersOptionToolTip", "Show folders in the view as well as assets?"),
+								FSlateIcon(),
+								FUIAction(
+									FExecuteAction::CreateSP(AssetView, &SAssetView::ToggleShowFolders),
+									FCanExecuteAction::CreateSP(AssetView, &SAssetView::IsToggleShowFoldersAllowed),
+									FIsActionChecked::CreateSP(AssetView, &SAssetView::IsShowingFolders)
+								),
+								EUserInterfaceActionType::ToggleButton
+							)
+						);
+
+						const FText ShowEmptyFoldersLabel = IsNewStyleEnabled() ? LOCTEXT("ShowEmptyFoldersOption_NewStyle", "Empty Folders") :LOCTEXT("ShowEmptyFoldersOption", "Show Emnpty Folders");
+						SetEntryOrder(
+							Section.AddMenuEntry(
+								"ShowEmptyFolders",
+								ShowEmptyFoldersLabel,
+								LOCTEXT("ShowEmptyFoldersOptionToolTip", "Show empty folders in the view as well as assets?"),
+								FSlateIcon(),
+								FUIAction(
+									FExecuteAction::CreateSP(AssetView, &SAssetView::ToggleShowEmptyFolders),
+									FCanExecuteAction::CreateSP(AssetView, &SAssetView::IsToggleShowEmptyFoldersAllowed),
+									FIsActionChecked::CreateSP(AssetView, &SAssetView::IsShowingEmptyFolders)
+								),
+								EUserInterfaceActionType::ToggleButton
+							)
+						);
 
+						const FText ShowFavoriteLabel = IsNewStyleEnabled() ?  LOCTEXT("ShowFavoriteOptions_NewStyle", "Favorites") : LOCTEXT("ShowFavoriteOptions", "Show Favorites");
 						Section.AddMenuEntry(
 							"ShowFavorite",
-							LOCTEXT("ShowFavoriteOptions", "Show Favorites"),
+							ShowFavoriteLabel,
 							LOCTEXT("ShowFavoriteOptionToolTip", "Show the favorite folders in the view?"),
 							FSlateIcon(),
 							FUIAction(
@@ -333,7 +343,7 @@ FDelayedAutoRegisterHelper SAssetView::AssetViewOptionsMenuRegistration(
 							EUserInterfaceActionType::ToggleButton
 						);
 
-						if (!bIsNewStyleEnabled)
+						if (!IsNewStyleEnabled())
 						{
 							Section.AddMenuEntry(
 								"FilterRecursively",
@@ -347,20 +357,9 @@ FDelayedAutoRegisterHelper SAssetView::AssetViewOptionsMenuRegistration(
 								),
 								EUserInterfaceActionType::ToggleButton
 							);
+						}
 
-							Section.AddMenuEntry(
-								"ShowAllFolder",
-								LOCTEXT("ShowAllFolderOption", "Show All Folder"),
-								LOCTEXT("ShowAllFolderOptionToolTip", "Show the all folder in the view?"),
-								FSlateIcon(),
-								FUIAction(
-									FExecuteAction::CreateSP(AssetView, &SAssetView::ToggleShowAllFolder),
-									FCanExecuteAction(),
-									FIsActionChecked::CreateSP(AssetView, &SAssetView::IsShowingAllFolder)
-								),
-								EUserInterfaceActionType::ToggleButton
-							);
-
+						SetEntryOrder(
 							Section.AddMenuEntry(
 								"OrganizeFolders",
 								LOCTEXT("OrganizeFoldersOption", "Organize Folders"),
@@ -372,19 +371,46 @@ FDelayedAutoRegisterHelper SAssetView::AssetViewOptionsMenuRegistration(
 									FIsActionChecked::CreateSP(AssetView, &SAssetView::IsOrganizingFolders)
 								),
 								EUserInterfaceActionType::ToggleButton
-							);
+							)
+						);
 
-							if (AssetView->bShowPathViewFilters)
-							{
-								Section.AddSubMenu(
-									"PathViewFilters",
-									LOCTEXT("PathViewFilters", "Path View Filters"),
-									LOCTEXT("PathViewFilters_ToolTip", "Path View Filters"),
-									FNewToolMenuDelegate());
-							}
+						if (AssetView->bShowPathViewFilters)
+						{
+							Section.AddSubMenu(
+								"PathViewFilters",
+								LOCTEXT("PathViewFilters", "Path View Filters"),
+								LOCTEXT("PathViewFilters_ToolTip", "Path View Filters"),
+								FNewToolMenuDelegate());
 						}
 					}
 
+					// Content Section
+					{
+						FToolMenuSection& Section = InMenu->FindOrAddSection("Content", IsNewStyleEnabled() ? FText::GetEmpty() : LOCTEXT("ContentHeading", "Content"));
+
+						if (IsNewStyleEnabled())
+						{
+							Section.AddSeparator(NAME_None);
+						}
+
+						const FText ShowAllFolderLabel = IsNewStyleEnabled() ?  LOCTEXT("ShowAllFolderOption_NewStyle", "All Folder") : LOCTEXT("ShowAllFolderOption", "Show All Folder");
+						FToolMenuSection& ShowAllFolderSection = IsNewStyleEnabled() ? Section : InMenu->FindOrAddSection("View");
+						SetEntryOrder(
+							ShowAllFolderSection.AddMenuEntry(
+								"ShowAllFolder",
+								ShowAllFolderLabel,
+								LOCTEXT("ShowAllFolderOptionToolTip", "Show the all folder in the view?"),
+								FSlateIcon(),
+								FUIAction(
+									FExecuteAction::CreateSP(AssetView, &SAssetView::ToggleShowAllFolder),
+									FCanExecuteAction(),
+									FIsActionChecked::CreateSP(AssetView, &SAssetView::IsShowingAllFolder)
+								),
+								EUserInterfaceActionType::ToggleButton
+							)
+						);
+					}
+
 					// Search Section
 					{
 						FToolMenuSection& Section = InMenu->FindOrAddSection("Search", LOCTEXT("SearchHeading", "Search"));
@@ -428,27 +454,16 @@ FDelayedAutoRegisterHelper SAssetView::AssetViewOptionsMenuRegistration(
 						);
 					}
 
-					if (bIsNewStyleEnabled)
+					if (IsNewStyleEnabled())
 					{
 						// Manage Section
 						{
+							// @note: This section is used elsewhere - @see: SContentBrowser::ExtendViewOptionsMenu
 							FToolMenuSection& Section = InMenu->FindOrAddSection("Manage", LOCTEXT("ManageHeading", "Manage"));
-							Section.AddMenuEntry(
-								"OrganizeFolders",
-								LOCTEXT("OrganizeFoldersOption", "Organize Folders"),
-								LOCTEXT("OrganizeFoldersOptionToolTip", "Organize folders in the view?"),
-								FSlateIcon(),
-								FUIAction(
-									FExecuteAction::CreateSP(AssetView, &SAssetView::ToggleOrganizeFolders),
-									FCanExecuteAction(),
-									FIsActionChecked::CreateSP(AssetView, &SAssetView::IsOrganizingFolders)
-								),
-								EUserInterfaceActionType::ToggleButton
-							);
 						}
 					}
 
-					if (!bIsNewStyleEnabled)
+					if (!IsNewStyleEnabled())
 					{
 						// Asset Columns Section
 						if (AssetView->GetColumnViewVisibility() == EVisibility::Visible)
diff --git a/Engine/Source/Editor/ContentBrowser/Private/SContentBrowser.cpp b/Engine/Source/Editor/ContentBrowser/Private/SContentBrowser.cpp
--- a/Engine/Source/Editor/ContentBrowser/Private/SContentBrowser.cpp
+++ b/Engine/Source/Editor/ContentBrowser/Private/SContentBrowser.cpp
@@ -1874,7 +1874,11 @@ void SContentBrowser::ExtendViewOptionsMenu(const FContentBrowserConfig* Config)
 								FCanExecuteAction(),
 								FIsActionChecked::CreateLambda([ContentBrowser = Context->OwningContentBrowser]() { return ContentBrowser.Pin()->bSourcesViewExpanded; })),
 							EUserInterfaceActionType::ToggleButton
-						);
+						)
+						.InsertPosition =
+							UE::Editor::ContentBrowser::IsNewStyleEnabled()
+							? FToolMenuInsert(NAME_None, EToolMenuInsertType::First)
+							: FToolMenuInsert();
 					}
 				}
 			}
