[Content Browser][New Style] Added sorting options in UI

diff --git a/Engine/Source/Editor/ContentBrowser/Private/ContentBrowserStyle.cpp b/Engine/Source/Editor/ContentBrowser/Private/ContentBrowserStyle.cpp
--- a/Engine/Source/Editor/ContentBrowser/Private/ContentBrowserStyle.cpp
+++ b/Engine/Source/Editor/ContentBrowser/Private/ContentBrowserStyle.cpp
@@ -314,6 +314,7 @@ namespace UE::ContentBrowser::Private
 		ContentBrowserToolBarStyle.SetSeparatorPadding(FMargin(4.0f, 0.0f));
 		if (UE::Editor::ContentBrowser::IsNewStyleEnabled())
 		{
+			ContentBrowserToolBarStyle.SetIconSize(FVector2D(16.0f, 16.0f));
 			ContentBrowserToolBarStyle.SetBackgroundPadding(FMargin(4.0f, 2.0f, 4.0f, 2.0f));
 		}
 		else
diff --git a/Engine/Source/Editor/ContentBrowser/Private/Menus/AssetViewMenus.cpp b/Engine/Source/Editor/ContentBrowser/Private/Menus/AssetViewMenus.cpp
--- a/Engine/Source/Editor/ContentBrowser/Private/Menus/AssetViewMenus.cpp
+++ b/Engine/Source/Editor/ContentBrowser/Private/Menus/AssetViewMenus.cpp
@@ -505,6 +505,73 @@ FDelayedAutoRegisterHelper SAssetView::AssetViewOptionsMenuRegistration(
 		}));
 	});
 
+FDelayedAutoRegisterHelper SAssetView::ToolBarMenuExtensionRegistration(
+	EDelayedRegisterRunPhase::EndOfEngineInit,
+	[]
+	{
+		static const FName ToolBarName("ContentBrowser.ToolBar");
+		if (UToolMenu* ToolBarMenu = UToolMenus::Get()->ExtendMenu(ToolBarName))
+		{
+			FToolMenuSection& SortSection = ToolBarMenu->FindOrAddSection("Sort");
+			SortSection.InsertPosition = FToolMenuInsert("Search", EToolMenuInsertType::After);
+
+			SortSection.AddDynamicEntry(
+				"Sort",
+				FNewToolMenuSectionDelegate::CreateLambda([](FToolMenuSection& InSection)
+				{
+					if (const UContentBrowserToolbarMenuContext* Context = InSection.FindContext<UContentBrowserToolbarMenuContext>())
+					{
+						if (Context->AssetView.IsValid())
+						{
+							const TSharedRef<SAssetView> AssetView = Context->AssetView.Pin().ToSharedRef();
+
+							const TAttribute<FSlateIcon> GetIcon = TAttribute<FSlateIcon>::Create(
+								TAttribute<FSlateIcon>::FGetter::CreateSPLambda(
+									AssetView,
+									[WeakMenuContext = MakeWeakObjectPtr(Context)]()
+									{
+										EColumnSortMode::Type ActiveSortMode = EColumnSortMode::Ascending;
+										if (WeakMenuContext.IsValid() && WeakMenuContext->AssetView.IsValid())
+										{
+											const TSharedPtr<SAssetView> AssetViewForIcon = WeakMenuContext->AssetView.Pin();
+											if (const TSharedPtr<FAssetViewSortManager> StrongSortManager = AssetViewForIcon->GetSortManager().Pin())
+											{
+												// @note: we don't use Secondary priorities, so this should change if/when we do
+												ActiveSortMode = StrongSortManager->GetSortMode(EColumnSortPriority::Primary);
+											}
+										}
+
+										return ActiveSortMode == EColumnSortMode::Ascending
+											? FSlateIcon(FAppStyle::Get().GetStyleSetName(), "Icons.SortDown")
+											: FSlateIcon(FAppStyle::Get().GetStyleSetName(), "Icons.SortUp");
+									}));
+
+							FToolUIAction SortAction;
+							SortAction.IsActionVisibleDelegate = FToolMenuIsActionButtonVisible::CreateSPLambda(
+								AssetView,
+								[AssetView](const FToolMenuContext&)
+								{
+									// Sort button only visible in Tile view mode - others have their own sorting UI representation.
+									return AssetView->IsCurrentViewType(EAssetViewType::Tile);
+								});
+
+							FToolMenuEntry& SortEntry = InSection.AddEntry(
+								FToolMenuEntry::InitComboButton(
+									"Sort",
+									FToolUIActionChoice(SortAction),
+									FNewToolMenuChoice(
+										FNewToolMenuDelegate::CreateSP(
+											AssetView,
+											&SAssetView::PopulateSortingButtonMenu)),
+									FText::GetEmpty(), // Empty as button...
+									LOCTEXT("SortToolTip", "Sorting options for the current asset view."),
+									GetIcon));
+						}
+					}
+				}));
+		}
+	});
+
 FDelayedAutoRegisterHelper SAssetView::NavigationBarMenuExtensionRegistration(
 	EDelayedRegisterRunPhase::EndOfEngineInit,
 	[]
@@ -533,8 +600,12 @@ FDelayedAutoRegisterHelper SAssetView::NavigationBarMenuExtensionRegistration(
 								FToolMenuEntry& AssetCountEntry = InSection.AddEntry(
 									FToolMenuEntry::InitWidget(
 										"AssetCount",
-										SNew(STextBlock)
-										.Text(AssetView, &SAssetView::GetAssetCountText),
+										SNew(SBox)
+										.Padding(FMargin(4, 0, 0, 0))
+										[
+											SNew(STextBlock)
+											.Text(AssetView, &SAssetView::GetAssetCountText)
+										],
 										FText::GetEmpty()));
 
 								AssetCountEntry.WidgetData.StyleParams.VerticalAlignment = VAlign_Center;
diff --git a/Engine/Source/Editor/ContentBrowser/Private/Menus/ContentBrowserMenus.cpp b/Engine/Source/Editor/ContentBrowser/Private/Menus/ContentBrowserMenus.cpp
--- a/Engine/Source/Editor/ContentBrowser/Private/Menus/ContentBrowserMenus.cpp
+++ b/Engine/Source/Editor/ContentBrowser/Private/Menus/ContentBrowserMenus.cpp
@@ -27,6 +27,7 @@ FDelayedAutoRegisterHelper SContentBrowser::NavigationBarMenuRegistration(
 		FToolMenuOwnerScoped OwnerScoped(UE_MODULE_NAME);
 
 		UToolMenu* ToolBar = UToolMenus::Get()->RegisterMenu(NavigationBarMenuName, NAME_None, EMultiBoxType::SlimHorizontalToolBar);
+		ToolBar->SetStyleSet(&UE::ContentBrowser::Private::FContentBrowserStyle::Get());
 		ToolBar->StyleName = "ContentBrowser.ToolBar";
 
 		ToolBar->AddSection("History")
@@ -225,6 +226,7 @@ FDelayedAutoRegisterHelper SContentBrowser::ToolBarMenuRegistration(
 		FToolMenuOwnerScoped OwnerScoped(UE_MODULE_NAME);
 
 		UToolMenu* ToolBar = UToolMenus::Get()->RegisterMenu(ToolBarMenuName, NAME_None, EMultiBoxType::SlimHorizontalToolBar);
+		ToolBar->SetStyleSet(&UE::ContentBrowser::Private::FContentBrowserStyle::Get());
 		ToolBar->StyleName = ToolBarMenuName; // Style name is the same as the menu entry name
 
 		ToolBar->AddSection("New")
diff --git a/Engine/Source/Editor/ContentBrowser/Private/SAssetView.cpp b/Engine/Source/Editor/ContentBrowser/Private/SAssetView.cpp
--- a/Engine/Source/Editor/ContentBrowser/Private/SAssetView.cpp
+++ b/Engine/Source/Editor/ContentBrowser/Private/SAssetView.cpp
@@ -1122,6 +1122,8 @@ void SAssetView::Construct( const FArguments& InArgs )
 	MinThumbnailScale = 0.2f * ThumbnailScaleRangeScalar;
 	MaxThumbnailScale = 1.9f * ThumbnailScaleRangeScalar;
 
+	SortManager = MakeShared<FAssetViewSortManager>();
+
 	bCanShowClasses = InArgs._CanShowClasses;
 
 	bCanShowFolders = InArgs._CanShowFolders;
@@ -1436,10 +1438,10 @@ void SAssetView::Construct( const FArguments& InArgs )
 	// is a good idea. 
 	if (CurrentViewType == EAssetViewType::Column && bSortByPathInColumnView)
 	{
-		SortManager.SetSortColumnId(EColumnSortPriority::Primary, SortManager.PathColumnId);
-		SortManager.SetSortColumnId(EColumnSortPriority::Secondary, SortManager.NameColumnId);
-		SortManager.SetSortMode(EColumnSortPriority::Primary, EColumnSortMode::Ascending);
-		SortManager.SetSortMode(EColumnSortPriority::Secondary, EColumnSortMode::Ascending);
+		SortManager->SetSortColumnId(EColumnSortPriority::Primary, SortManager->PathColumnId);
+		SortManager->SetSortColumnId(EColumnSortPriority::Secondary, SortManager->NameColumnId);
+		SortManager->SetSortMode(EColumnSortPriority::Primary, EColumnSortMode::Ascending);
+		SortManager->SetSortMode(EColumnSortPriority::Secondary, EColumnSortMode::Ascending);
 		SortList();
 	}
 }
@@ -1604,7 +1606,7 @@ void SAssetView::BeginCreateDeferredItem()
 		DeferredItemToCreate->bWasAddedToView = true;
 
 		FilteredAssetItems.Insert(NewItem, 0);
-		SortManager.SortList(FilteredAssetItems, MajorityAssetType, CustomColumns);
+		SortManager->SortList(FilteredAssetItems, MajorityAssetType, CustomColumns);
 
 		SetSelection(NewItem);
 		RequestScrollIntoView(NewItem);
@@ -2597,7 +2599,7 @@ TSharedRef<SAssetListView> SAssetView::CreateListView()
 			.ResizeMode(ESplitterResizeMode::FixedSize)
 
 			// Revision Control column, currently doesn't support sorting
-			+ SHeaderRow::Column(SortManager.RevisionControlColumnId)
+			+ SHeaderRow::Column(SortManager->RevisionControlColumnId)
 			.FixedWidth(30.f)
 			.HAlignHeader(HAlign_Center)
 			.VAlignHeader(VAlign_Center)
@@ -2608,10 +2610,10 @@ TSharedRef<SAssetListView> SAssetView::CreateListView()
 				RevisionControlColumnIcon
 			]
 			
-			+ SHeaderRow::Column(SortManager.NameColumnId)
+			+ SHeaderRow::Column(SortManager->NameColumnId)
 			.FillWidth(300)
-			.SortMode( TAttribute< EColumnSortMode::Type >::Create( TAttribute< EColumnSortMode::Type >::FGetter::CreateSP( this, &SAssetView::GetColumnSortMode, SortManager.NameColumnId ) ) )
-			.SortPriority(TAttribute< EColumnSortPriority::Type >::Create(TAttribute< EColumnSortPriority::Type >::FGetter::CreateSP(this, &SAssetView::GetColumnSortPriority, SortManager.NameColumnId)))
+			.SortMode(TAttribute<EColumnSortMode::Type>::Create(TAttribute<EColumnSortMode::Type>::FGetter::CreateSP(this, &SAssetView::GetColumnSortMode, SortManager->NameColumnId)))
+			.SortPriority(TAttribute<EColumnSortPriority::Type>::Create(TAttribute<EColumnSortPriority::Type>::FGetter::CreateSP(this, &SAssetView::GetColumnSortPriority, SortManager->NameColumnId)))
 			.OnSort( FOnSortModeChanged::CreateSP( this, &SAssetView::OnSortColumnHeader ) )
 			.DefaultLabel( LOCTEXT("Column_Name", "Name") )
 			.ShouldGenerateWidget(true) // Can't hide name column, so at least one column is visible
@@ -2645,7 +2647,7 @@ TSharedRef<SAssetColumnView> SAssetView::CreateColumnView()
 			.OnHiddenColumnsListChanged(this, &SAssetView::OnHiddenColumnsChanged)
 
 			// Revision Control column, currently doesn't support sorting
-			+ SHeaderRow::Column(SortManager.RevisionControlColumnId)
+			+ SHeaderRow::Column(SortManager->RevisionControlColumnId)
 			.FixedWidth(30.f)
 			.HAlignHeader(HAlign_Center)
 			.VAlignHeader(VAlign_Center)
@@ -2656,18 +2658,18 @@ TSharedRef<SAssetColumnView> SAssetView::CreateColumnView()
 				RevisionControlColumnIcon
 			]
 			
-			+ SHeaderRow::Column(SortManager.NameColumnId)
+			+ SHeaderRow::Column(SortManager->NameColumnId)
 			.FillWidth(300)
-			.SortMode( TAttribute< EColumnSortMode::Type >::Create( TAttribute< EColumnSortMode::Type >::FGetter::CreateSP( this, &SAssetView::GetColumnSortMode, SortManager.NameColumnId ) ) )
-			.SortPriority(TAttribute< EColumnSortPriority::Type >::Create(TAttribute< EColumnSortPriority::Type >::FGetter::CreateSP(this, &SAssetView::GetColumnSortPriority, SortManager.NameColumnId)))
+			.SortMode(TAttribute<EColumnSortMode::Type>::Create(TAttribute<EColumnSortMode::Type>::FGetter::CreateSP(this, &SAssetView::GetColumnSortMode, SortManager->NameColumnId)))
+			.SortPriority(TAttribute<EColumnSortPriority::Type>::Create(TAttribute<EColumnSortPriority::Type>::FGetter::CreateSP(this, &SAssetView::GetColumnSortPriority, SortManager->NameColumnId)))
 			.OnSort( FOnSortModeChanged::CreateSP( this, &SAssetView::OnSortColumnHeader ) )
 			.DefaultLabel( LOCTEXT("Column_Name", "Name") )
 			.ShouldGenerateWidget(true) // Can't hide name column, so at least one column is visible
 		);
 
 	{
-		const bool bIsColumnVisible = !HiddenColumnNames.Contains(SortManager.RevisionControlColumnId.ToString());
-        NewColumnView->GetHeaderRow()->SetShowGeneratedColumn(SortManager.RevisionControlColumnId, bIsColumnVisible);
+		const bool bIsColumnVisible = !HiddenColumnNames.Contains(SortManager->RevisionControlColumnId.ToString());
+        NewColumnView->GetHeaderRow()->SetShowGeneratedColumn(SortManager->RevisionControlColumnId, bIsColumnVisible);
 	}
 
 	NewColumnView->GetHeaderRow()->SetOnGetMaxRowSizeForColumn(FOnGetMaxRowSizeForColumn::CreateRaw(NewColumnView.Get(), &SAssetColumnView::GetMaxRowSizeForColumn));
@@ -2675,33 +2677,33 @@ TSharedRef<SAssetColumnView> SAssetView::CreateColumnView()
 	if(bShowTypeInColumnView)
 	{
 		NewColumnView->GetHeaderRow()->AddColumn(
-				SHeaderRow::Column(SortManager.ClassColumnId)
+				SHeaderRow::Column(SortManager->ClassColumnId)
 				.FillWidth(160)
-				.SortMode(TAttribute< EColumnSortMode::Type >::Create(TAttribute< EColumnSortMode::Type >::FGetter::CreateSP(this, &SAssetView::GetColumnSortMode, SortManager.ClassColumnId)))
-				.SortPriority(TAttribute< EColumnSortPriority::Type >::Create(TAttribute< EColumnSortPriority::Type >::FGetter::CreateSP(this, &SAssetView::GetColumnSortPriority, SortManager.ClassColumnId)))
+				.SortMode(TAttribute< EColumnSortMode::Type >::Create(TAttribute< EColumnSortMode::Type >::FGetter::CreateSP(this, &SAssetView::GetColumnSortMode, SortManager->ClassColumnId)))
+				.SortPriority(TAttribute< EColumnSortPriority::Type >::Create(TAttribute< EColumnSortPriority::Type >::FGetter::CreateSP(this, &SAssetView::GetColumnSortPriority, SortManager->ClassColumnId)))
 				.OnSort(FOnSortModeChanged::CreateSP(this, &SAssetView::OnSortColumnHeader))
 				.DefaultLabel(LOCTEXT("Column_Class", "Type"))
 			);
 
-		const bool bIsColumnVisible = !HiddenColumnNames.Contains(SortManager.ClassColumnId.ToString());
+		const bool bIsColumnVisible = !HiddenColumnNames.Contains(SortManager->ClassColumnId.ToString());
 		
-		NewColumnView->GetHeaderRow()->SetShowGeneratedColumn(SortManager.ClassColumnId, bIsColumnVisible);
+		NewColumnView->GetHeaderRow()->SetShowGeneratedColumn(SortManager->ClassColumnId, bIsColumnVisible);
 	}
 
 	if (bShowPathInColumnView)
 	{
 		NewColumnView->GetHeaderRow()->AddColumn(
-				SHeaderRow::Column(SortManager.PathColumnId)
+				SHeaderRow::Column(SortManager->PathColumnId)
 				.FillWidth(160)
-				.SortMode(TAttribute< EColumnSortMode::Type >::Create(TAttribute< EColumnSortMode::Type >::FGetter::CreateSP(this, &SAssetView::GetColumnSortMode, SortManager.PathColumnId)))
-				.SortPriority(TAttribute< EColumnSortPriority::Type >::Create(TAttribute< EColumnSortPriority::Type >::FGetter::CreateSP(this, &SAssetView::GetColumnSortPriority, SortManager.PathColumnId)))
+				.SortMode(TAttribute< EColumnSortMode::Type >::Create(TAttribute< EColumnSortMode::Type >::FGetter::CreateSP(this, &SAssetView::GetColumnSortMode, SortManager->PathColumnId)))
+				.SortPriority(TAttribute< EColumnSortPriority::Type >::Create(TAttribute< EColumnSortPriority::Type >::FGetter::CreateSP(this, &SAssetView::GetColumnSortPriority, SortManager->PathColumnId)))
 				.OnSort(FOnSortModeChanged::CreateSP(this, &SAssetView::OnSortColumnHeader))
 				.DefaultLabel(LOCTEXT("Column_Path", "Path"))
 			);
 
-		const bool bIsColumnVisible = !HiddenColumnNames.Contains(SortManager.PathColumnId.ToString());
+		const bool bIsColumnVisible = !HiddenColumnNames.Contains(SortManager->PathColumnId.ToString());
 		
-		NewColumnView->GetHeaderRow()->SetShowGeneratedColumn(SortManager.PathColumnId, bIsColumnVisible);
+		NewColumnView->GetHeaderRow()->SetShowGeneratedColumn(SortManager->PathColumnId, bIsColumnVisible);
 	}
 
 	return NewColumnView.ToSharedRef();
@@ -3478,10 +3480,10 @@ void SAssetView::SetMajorityAssetType(FName NewMajorityAssetType)
 
 	auto IsFixedColumn = [this](FName InColumnId)
 	{
-		const bool bIsFixedNameColumn = InColumnId == SortManager.NameColumnId;
-		const bool bIsFixedRevisionControlColumn = InColumnId == SortManager.RevisionControlColumnId;
-		const bool bIsFixedClassColumn = bShowTypeInColumnView && InColumnId == SortManager.ClassColumnId;
-		const bool bIsFixedPathColumn = bShowPathInColumnView && InColumnId == SortManager.PathColumnId;
+		const bool bIsFixedNameColumn = InColumnId == SortManager->NameColumnId;
+		const bool bIsFixedRevisionControlColumn = InColumnId == SortManager->RevisionControlColumnId;
+		const bool bIsFixedClassColumn = bShowTypeInColumnView && InColumnId == SortManager->ClassColumnId;
+		const bool bIsFixedPathColumn = bShowPathInColumnView && InColumnId == SortManager->PathColumnId;
 		return bIsFixedNameColumn || bIsFixedRevisionControlColumn || bIsFixedClassColumn || bIsFixedPathColumn;
 	};
 
@@ -3522,7 +3524,7 @@ void SAssetView::SetMajorityAssetType(FName NewMajorityAssetType)
 		TArray<FSortOrder> CurrentSortOrder;
 		for (int32 PriorityIdx = 0; PriorityIdx < EColumnSortPriority::Max; PriorityIdx++)
 		{
-			const FName SortColumn = SortManager.GetSortColumnId(static_cast<EColumnSortPriority::Type>(PriorityIdx));
+			const FName SortColumn = SortManager->GetSortColumnId(static_cast<EColumnSortPriority::Type>(PriorityIdx));
 			if (SortColumn != NAME_None)
 			{
 				const bool bSortRelevant = SortColumn == FAssetViewSortManager::NameColumnId
@@ -3663,10 +3665,10 @@ void SAssetView::SetMajorityAssetType(FName NewMajorityAssetType)
 			for (int32 SortIdx = 0; SortIdx < CurrentSortOrder.Num(); SortIdx++)
 			{
 				check(CurrentSortOrder[SortIdx].bSortRelevant);
-				if (!SortManager.SetOrToggleSortColumn(static_cast<EColumnSortPriority::Type>(PriorityNum), CurrentSortOrder[SortIdx].SortColumn))
+				if (!SortManager->SetOrToggleSortColumn(static_cast<EColumnSortPriority::Type>(PriorityNum), CurrentSortOrder[SortIdx].SortColumn))
 				{
 					// Toggle twice so mode is preserved if this isn't a new column assignation
-					SortManager.SetOrToggleSortColumn(static_cast<EColumnSortPriority::Type>(PriorityNum), CurrentSortOrder[SortIdx].SortColumn);
+					SortManager->SetOrToggleSortColumn(static_cast<EColumnSortPriority::Type>(PriorityNum), CurrentSortOrder[SortIdx].SortColumn);
 				}				
 				bPendingSortFilteredItems = true;
 				PriorityNum++;
@@ -3675,7 +3677,7 @@ void SAssetView::SetMajorityAssetType(FName NewMajorityAssetType)
 		else if (CurrentSortOrder.Num() == 0)
 		{
 			// If the current sort column is no longer relevant, revert to "Name" and resort when convenient
-			SortManager.ResetSort();
+			SortManager->ResetSort();
 			bPendingSortFilteredItems = true;
 		}
 	}
@@ -3765,7 +3767,7 @@ void SAssetView::SortList(bool bSyncToSelection)
 {
 	if ( !IsRenamingAsset() )
 	{
-		SortManager.SortList(FilteredAssetItems, MajorityAssetType, CustomColumns);
+		SortManager->SortList(FilteredAssetItems, MajorityAssetType, CustomColumns);
 
 		// Update the thumbnails we were using since the order has changed
 		bPendingUpdateThumbnails = true;
@@ -3839,6 +3841,155 @@ void SAssetView::PopulateFilterAdditionalParams(FFiltersAdditionalParams& OutPar
 	OutParams.CanShowLocalizedContent = FCanExecuteAction::CreateSP(this, &SAssetView::IsToggleShowLocalizedContentAllowed);
 }
 
+void SAssetView::OnSetSortParameters(
+	const FToolMenuContext& InMenuContext,
+	const TOptional<const EColumnSortPriority::Type> InSortPriority,
+	const TOptional<const FName> InColumnId,
+	const TOptional<const EColumnSortMode::Type> InNewSortMode)
+{
+	TOptional<const EColumnSortMode::Type> NewSortMode = InNewSortMode;
+
+	// Set sort mode to the currently active one, if none specified
+	if (!InNewSortMode.IsSet())
+	{
+		NewSortMode.Emplace(GetSortManager().Pin()->GetSortMode(EColumnSortPriority::Primary));
+	}
+
+	SetSortParameters(InSortPriority, InColumnId, NewSortMode);
+}
+
+void SAssetView::PopulateSortingButtonMenu(UToolMenu* InToolMenu)
+{
+	UContentBrowserAssetSortingContextMenuContext* SortingContext = NewObject<UContentBrowserAssetSortingContextMenuContext>();
+	SortingContext->OwningContentBrowser = OwningContentBrowser;
+	SortingContext->AssetView = SharedThis(this);
+	SortingContext->AssetViewSortManager = GetSortManager();
+
+	InToolMenu->Context.AddObject(SortingContext);
+
+	FToolMenuSection& SortBySection = InToolMenu->AddSection("SortBy", LOCTEXT("SortByHeading", "Sort By"));
+	{
+		static const TOptional<const EColumnSortPriority::Type> UnsetSortPriority;
+		static const TOptional<const EColumnSortMode::Type> UnsetSortMode;
+
+		static const FName SortableColumnIds[] = { FAssetViewSortManager::NameColumnId, FAssetViewSortManager::DiskSizeColumnId };
+		for (const FName SortableColumnId : SortableColumnIds)
+		{
+			const TOptional<const FName> ColumnId = SortableColumnId;
+
+			FToolUIAction SortByAction;
+			SortByAction.ExecuteAction =
+				FToolMenuExecuteAction::CreateSP(
+					this,
+					&SAssetView::OnSetSortParameters,
+					UnsetSortPriority,
+					ColumnId,
+					UnsetSortMode);
+
+			SortByAction.GetActionCheckState =
+				FToolMenuGetActionCheckState::CreateSPLambda(
+					this,
+					[](const FToolMenuContext& InMenuContext, const FName InId)
+					{
+						if (const UContentBrowserAssetSortingContextMenuContext* SortingContext = InMenuContext.FindContext<UContentBrowserAssetSortingContextMenuContext>())
+						{
+							if (const TSharedPtr<FAssetViewSortManager> StrongSortManager = SortingContext->AssetViewSortManager.Pin())
+							{
+								return StrongSortManager->GetSortColumnId(EColumnSortPriority::Primary) == InId
+									? ECheckBoxState::Checked
+									: ECheckBoxState::Unchecked;
+							}
+						}
+
+						return ECheckBoxState::Unchecked;
+					},
+					SortableColumnId);
+
+			// @todo: should this be localized?
+			FText SortableColumnLabel = FText::FromName(SortableColumnId);
+
+			SortBySection.AddMenuEntry(
+				SortableColumnId,
+				SortableColumnLabel,
+				FText::Format(LOCTEXT("SortByOptionToolTip", "Sorts the items by {0}"), SortableColumnLabel),
+				FSlateIcon(),
+				SortByAction,
+				EUserInterfaceActionType::RadioButton);
+		}
+	}
+
+	FToolMenuSection& SortTypeSection = InToolMenu->AddSection("SortType", LOCTEXT("SortTypeHeading", "Sort Type"));
+	{
+		auto SetSortMode = [](const FToolMenuContext& InMenuContext, const EColumnSortMode::Type InMode)
+		{
+			if (UContentBrowserAssetSortingContextMenuContext* SortingContext = InMenuContext.FindContext<UContentBrowserAssetSortingContextMenuContext>())
+			{
+				if (TSharedPtr<SAssetView> StrongAssetView = SortingContext->AssetView.Pin())
+				{
+					StrongAssetView->SetSortParameters(EColumnSortPriority::Primary, { }, InMode);
+				}
+			}
+		};
+
+		auto IsSortMode = [](const FToolMenuContext& InMenuContext, const EColumnSortMode::Type InMode)
+		{
+			if (const UContentBrowserAssetSortingContextMenuContext* SortingContext = InMenuContext.FindContext<UContentBrowserAssetSortingContextMenuContext>())
+			{
+				if (const TSharedPtr<FAssetViewSortManager> StrongSortManager = SortingContext->AssetViewSortManager.Pin())
+				{
+					return StrongSortManager->GetSortMode(EColumnSortPriority::Primary) == InMode
+						? ECheckBoxState::Checked
+						: ECheckBoxState::Unchecked;
+				}
+			}
+
+			return ECheckBoxState::Unchecked;
+		};
+
+		FToolUIAction SortAscendingAction;
+		SortAscendingAction.ExecuteAction =
+			FToolMenuExecuteAction::CreateSPLambda(
+				this,
+				SetSortMode,
+				EColumnSortMode::Ascending);
+
+		SortAscendingAction.GetActionCheckState =
+			FToolMenuGetActionCheckState::CreateSPLambda(
+				this,
+				IsSortMode,
+				EColumnSortMode::Ascending);
+
+		SortTypeSection.AddMenuEntry(
+			"Ascending",
+			LOCTEXT("AscendingOrder", "Ascending"),
+			LOCTEXT("AscendingOrderToolTip", "Sort the items in Ascending order"),
+			FSlateIcon(),
+			SortAscendingAction,
+			EUserInterfaceActionType::RadioButton);
+
+		FToolUIAction SortDescendingAction;
+		SortDescendingAction.ExecuteAction =
+			FToolMenuExecuteAction::CreateSPLambda(
+				this,
+				SetSortMode,
+				EColumnSortMode::Descending);
+
+		SortDescendingAction.GetActionCheckState =
+			FToolMenuGetActionCheckState::CreateSPLambda(
+				this,
+				IsSortMode,
+				EColumnSortMode::Descending);
+
+		SortTypeSection.AddMenuEntry(
+			"Descending",
+			LOCTEXT("DescendingOrder", "Descending"),
+			LOCTEXT("AscendingOrderToolTip", "Sort the items in Descending order"),
+			FSlateIcon(),
+			SortDescendingAction,
+			EUserInterfaceActionType::RadioButton);
+	}
+}
+
 void SAssetView::ToggleShowFolders()
 {
 	check(IsToggleShowFoldersAllowed());
@@ -5416,12 +5567,14 @@ EColumnSortMode::Type SAssetView::GetColumnSortMode(const FName ColumnId) const
 	for (int32 PriorityIdx = 0; PriorityIdx < EColumnSortPriority::Max; PriorityIdx++)
 	{
 		const EColumnSortPriority::Type SortPriority = static_cast<EColumnSortPriority::Type>(PriorityIdx);
-		if (ColumnId == SortManager.GetSortColumnId(SortPriority))
+		if (ColumnId == SortManager->GetSortColumnId(SortPriority))
 		{
-			return SortManager.GetSortMode(SortPriority);
+			return SortManager->GetSortMode(SortPriority);
 		}
 	}
-	return EColumnSortMode::None;
+
+	static constexpr EColumnSortMode::Type DefaultSortMode = EColumnSortMode::Ascending;
+	return DefaultSortMode;
 }
 
 EColumnSortPriority::Type SAssetView::GetColumnSortPriority(const FName ColumnId) const
@@ -5429,21 +5582,52 @@ EColumnSortPriority::Type SAssetView::GetColumnSortPriority(const FName ColumnId
 	for (int32 PriorityIdx = 0; PriorityIdx < EColumnSortPriority::Max; PriorityIdx++)
 	{
 		const EColumnSortPriority::Type SortPriority = static_cast<EColumnSortPriority::Type>(PriorityIdx);
-		if (ColumnId == SortManager.GetSortColumnId(SortPriority))
+		if (ColumnId == SortManager->GetSortColumnId(SortPriority))
 		{
 			return SortPriority;
 		}
 	}
-	return EColumnSortPriority::Primary;
+
+	static constexpr EColumnSortPriority::Type DefaultSortPriority = EColumnSortPriority::Primary;
+	return DefaultSortPriority;
 }
 
 void SAssetView::OnSortColumnHeader(const EColumnSortPriority::Type SortPriority, const FName& ColumnId, const EColumnSortMode::Type NewSortMode)
 {
-	SortManager.SetSortColumnId(SortPriority, ColumnId);
-	SortManager.SetSortMode(SortPriority, NewSortMode);
+	SortManager->SetSortColumnId(SortPriority, ColumnId);
+	SortManager->SetSortMode(SortPriority, NewSortMode);
 	SortList();
 }
 
+void SAssetView::SetSortParameters(
+	const TOptional<const EColumnSortPriority::Type>& InSortPriority,
+	const TOptional<const FName>& InColumnId,
+	const TOptional<const EColumnSortMode::Type>& InNewSortMode)
+{
+	FName ColumnId = InColumnId.Get(NAME_None);
+
+	static constexpr EColumnSortPriority::Type DefaultSortPriority = EColumnSortPriority::Primary;
+
+	// Use specified priority OR default (primary)...
+	EColumnSortPriority::Type SortPriority = InSortPriority.Get(DefaultSortPriority);
+
+	// ... unless a ColumnId WAS specified and a priority was NOT
+	if (InColumnId.IsSet() && !InSortPriority.IsSet())
+	{
+		SortPriority = GetColumnSortPriority(ColumnId);
+	}
+
+	if (!InColumnId.IsSet())
+	{
+		ColumnId = SortManager->GetSortColumnId(SortPriority);
+	}
+
+	// Use specified sort mode OR get from the ColumnId
+	const EColumnSortMode::Type SortMode = InNewSortMode.Get(GetColumnSortMode(ColumnId));
+
+	OnSortColumnHeader(SortPriority, ColumnId, SortMode);
+}
+
 EVisibility SAssetView::IsAssetShowWarningTextVisible() const
 {
 	return (FilteredAssetItems.Num() > 0 || bQuickFrontendListRefreshRequested) ? EVisibility::Collapsed : EVisibility::HitTestInvisible;
@@ -5678,7 +5862,7 @@ void SAssetView::ExportColumns()
 		}
 
 		FString SaveString;
-		SortManager.ExportColumnsToCSV(FilteredAssetItems, ColumnNames, CustomColumns, SaveString);
+		SortManager->ExportColumnsToCSV(FilteredAssetItems, ColumnNames, CustomColumns, SaveString);
 
 		FFileHelper::SaveStringToFile(SaveString, *OutFilenames[0]);
 	}
@@ -5929,6 +6113,11 @@ void SAssetView::SetShouldFilterItem(FOnShouldFilterItem InCallback)
 	RequestQuickFrontendListRefresh();
 }
 
+TWeakPtr<FAssetViewSortManager> SAssetView::GetSortManager() const
+{
+	return SortManager;
+}
+
 void SAssetView::OnCompleteFiltering(double InAmortizeDuration)
 {
 	CurrentFrontendFilterTelemetry.AmortizeDuration = InAmortizeDuration;
diff --git a/Engine/Source/Editor/ContentBrowser/Public/ContentBrowserMenuContexts.h b/Engine/Source/Editor/ContentBrowser/Public/ContentBrowserMenuContexts.h
--- a/Engine/Source/Editor/ContentBrowser/Public/ContentBrowserMenuContexts.h
+++ b/Engine/Source/Editor/ContentBrowser/Public/ContentBrowserMenuContexts.h
@@ -16,6 +16,7 @@
 
 #include "ContentBrowserMenuContexts.generated.h"
 
+class FAssetViewSortManager;
 class FAssetContextMenu;
 class IAssetTypeActions;
 class SAssetView;
@@ -271,6 +272,17 @@ public:
 	TWeakPtr<SAssetView> AssetView;
 };
 
+UCLASS()
+class CONTENTBROWSER_API UContentBrowserAssetSortingContextMenuContext : public UObject
+{
+	GENERATED_BODY()
+
+public:
+	TWeakPtr<SContentBrowser> OwningContentBrowser;
+	TWeakPtr<SAssetView> AssetView;
+	TWeakPtr<FAssetViewSortManager> AssetViewSortManager;
+};
+
 UCLASS()
 class CONTENTBROWSER_API UContentBrowserMenuContext : public UObject
 {
diff --git a/Engine/Source/Editor/ContentBrowser/Public/SAssetView.h b/Engine/Source/Editor/ContentBrowser/Public/SAssetView.h
--- a/Engine/Source/Editor/ContentBrowser/Public/SAssetView.h
+++ b/Engine/Source/Editor/ContentBrowser/Public/SAssetView.h
@@ -490,6 +490,16 @@ public:
 
 	/** True if we have thumbnails that are pending update */
 	bool HasThumbnailsPendingUpdate() const;
+
+	/** Get the underlying sort manager */
+	TWeakPtr<FAssetViewSortManager> GetSortManager() const;
+
+	/** Set's sorting based on the given parameters, using the current value if not specified */
+	void SetSortParameters(
+		const TOptional<const EColumnSortPriority::Type>& InSortPriority,
+		const TOptional<const FName>& InColumnId,
+		const TOptional<const EColumnSortMode::Type>& InNewSortMode);
+
 private:
 
 	/** Sets the pending selection to the current selection (used when changing views or refreshing the view). */
@@ -572,9 +582,20 @@ private:
 	/** Registers the menu used by the settings button */
 	static FDelayedAutoRegisterHelper AssetViewOptionsMenuRegistration;
 
+	/** Extend the top toolbar */
+	static FDelayedAutoRegisterHelper ToolBarMenuExtensionRegistration;
+
 	/** Extend the bottom navigation bar */
 	static FDelayedAutoRegisterHelper NavigationBarMenuExtensionRegistration;
 
+	void OnSetSortParameters(const FToolMenuContext& InMenuContext,
+		const TOptional<const EColumnSortPriority::Type> InSortPriority,
+		const TOptional<const FName> InColumnId,
+		const TOptional<const EColumnSortMode::Type> InNewSortMode);
+
+	/** Fill in menu content for when the sorting combo button is clicked */
+	void PopulateSortingButtonMenu(UToolMenu* Menu);
+
 	/** Toggle whether folders should be shown or not */
 	void ToggleShowFolders();
 
@@ -1144,7 +1165,7 @@ private:
 	FName MajorityAssetType;
 
 	/** The manager responsible for sorting assets in the view */
-	FAssetViewSortManager SortManager;
+	TSharedPtr<FAssetViewSortManager> SortManager;
 
 	/** Flag indicating if we will be filling the empty space in the tile view. */
 	bool bFillEmptySpaceInTileView;
