[ContentBrowser] Fixed localization issue in SAssetView

#rnx

diff --git a/Engine/Source/Editor/ContentBrowser/Private/SAssetView.cpp b/Engine/Source/Editor/ContentBrowser/Private/SAssetView.cpp
--- a/Engine/Source/Editor/ContentBrowser/Private/SAssetView.cpp
+++ b/Engine/Source/Editor/ContentBrowser/Private/SAssetView.cpp
@@ -3983,7 +3983,7 @@ void SAssetView::PopulateSortingButtonMenu(UToolMenu* InToolMenu)
 		SortTypeSection.AddMenuEntry(
 			"Descending",
 			LOCTEXT("DescendingOrder", "Descending"),
-			LOCTEXT("AscendingOrderToolTip", "Sort the items in Descending order"),
+			LOCTEXT("DescendingOrderToolTip", "Sort the items in Descending order"),
 			FSlateIcon(),
 			SortDescendingAction,
 			EUserInterfaceActionType::RadioButton);
@@ -5538,6 +5538,7 @@ float SAssetView::GetTileViewItemHeight() const
 
 float SAssetView::GetTileViewItemBaseHeight() const
 {
+
 	return UE::Editor::ContentBrowser::IsNewStyleEnabled()
 		? GetTileViewItemBaseWidth()
 		: (float)(TileViewThumbnailSize + TileViewThumbnailPadding * 2) * FMath::Lerp(MinThumbnailScale, MaxThumbnailScale, GetThumbnailScale());
