Fix typo

#rnx

diff --git a/Engine/Source/Editor/ContentBrowser/Private/Menus/AssetViewMenus.cpp b/Engine/Source/Editor/ContentBrowser/Private/Menus/AssetViewMenus.cpp
--- a/Engine/Source/Editor/ContentBrowser/Private/Menus/AssetViewMenus.cpp
+++ b/Engine/Source/Editor/ContentBrowser/Private/Menus/AssetViewMenus.cpp
@@ -313,7 +313,7 @@ FDelayedAutoRegisterHelper SAssetView::AssetViewOptionsMenuRegistration(
 							)
 						);
 
-						const FText ShowEmptyFoldersLabel = IsNewStyleEnabled() ? LOCTEXT("ShowEmptyFoldersOption_NewStyle", "Empty Folders") :LOCTEXT("ShowEmptyFoldersOption", "Show Emnpty Folders");
+						const FText ShowEmptyFoldersLabel = IsNewStyleEnabled() ? LOCTEXT("ShowEmptyFoldersOption_NewStyle", "Empty Folders") :LOCTEXT("ShowEmptyFoldersOption", "Show Empty Folders");
 						SetEntryOrder(
 							Section.AddMenuEntry(
 								"ShowEmptyFolders",
