Fixed user search regression in asset dialogs

Asset dialogs disable recursion for their type filter, but a recent change had made that  flag also apply to user searches. This change restores the old behavior of a user search always being implicitly recursive.

#rnx

diff --git a/Engine/Source/Editor/ContentBrowser/Private/SAssetView.cpp b/Engine/Source/Editor/ContentBrowser/Private/SAssetView.cpp
--- a/Engine/Source/Editor/ContentBrowser/Private/SAssetView.cpp
+++ b/Engine/Source/Editor/ContentBrowser/Private/SAssetView.cpp
@@ -3153,14 +3153,14 @@ void SAssetView::ToggleFilteringRecursively()
 
 bool SAssetView::ShouldFilterRecursively() const
 {
-	if (IsFilteringRecursively())
+	// Quick check for conditions that activate the recursive filtering
+	if (bUserSearching)
 	{
-		// Quick check for conditions that activate the recursive filtering
-		if (bUserSearching)
-		{
-			return true;
-		}
+		return true;
+	}
 
+	if (IsFilteringRecursively())
+	{
 		if (!BackendFilter.IsEmpty() )
 		{
 			return true;
