[ContentBrowser] Nav bar overflow behavior

#rnx

diff --git a/Engine/Source/Editor/ContentBrowser/Private/SContentBrowser.cpp b/Engine/Source/Editor/ContentBrowser/Private/SContentBrowser.cpp
--- a/Engine/Source/Editor/ContentBrowser/Private/SContentBrowser.cpp
+++ b/Engine/Source/Editor/ContentBrowser/Private/SContentBrowser.cpp
@@ -535,6 +535,11 @@ void SContentBrowser::Construct( const FArguments& InArgs, const FName& InInstan
 
 		const TSharedRef<SWidget> AssetView = CreateAssetView(Config);
 
+		NavigationToolBarWidget = CreateNavigationToolBar(Config);
+
+		TSharedPtr<SBox> AssetViewNavigationToolBarContainer = nullptr;
+		TSharedPtr<SBox> SourceTreeAndAssetViewNavigationToolBarContainer = nullptr;
+
 		ChildSlot
 		[
 			SNew(SVerticalBox)
@@ -701,7 +706,7 @@ void SContentBrowser::Construct( const FArguments& InArgs, const FName& InInstan
 									AssetView
 								]
 
-								// Navigation Bar
+								// Navigation ToolBar
 								+ SVerticalBox::Slot()
 								.AutoHeight()
 								.HAlign(HAlign_Fill)
@@ -711,15 +716,38 @@ void SContentBrowser::Construct( const FArguments& InArgs, const FName& InInstan
 									.Padding(FMargin(3))
 									.BorderImage(bIsDrawer ? FStyleDefaults::GetNoBrush() : FAppStyle::Get().GetBrush("Brushes.Panel"))
 									[
-										CreateNavigationBar(Config)
+										SAssignNew(AssetViewNavigationToolBarContainer, SBox)
+										.Padding(0)
+										[
+											NavigationToolBarWidget.ToSharedRef()
+										]
 									]
 								]
 							]
 						]
 					]
+
+					+ SVerticalBox::Slot()
+					.AutoHeight()
+					.Padding(0)
+					[
+						SAssignNew(SourceTreeAndAssetViewNavigationToolBarContainer, SBox)
+						.Padding(0)
+					]
 				]
 			]
 		];
+
+		NavigationToolBarParentSwitcher = MakeShared<UE::Editor::ContentBrowser::Private::TWidgetSizeParentSwitcher<Orient_Horizontal>>(
+			NavigationToolBarWidget,
+			TArray<TSharedPtr<SBox>>{ SourceTreeAndAssetViewNavigationToolBarContainer, AssetViewNavigationToolBarContainer },
+			TArray<FInt16Range>{ FInt16Range(0, 550), FInt16Range(550, 10000) },
+			[this]()
+			{
+				const float Size = this->GetCachedGeometry().GetLocalSize().X;
+				return Size;
+			}
+		);
 	}
 	else
 	{
@@ -1186,7 +1214,7 @@ TSharedRef<SWidget> SContentBrowser::CreateToolBar(const FContentBrowserConfig*
 	return UToolMenus::Get()->GenerateWidget("ContentBrowser.ToolBar", MenuContext);
 }
 
-TSharedRef<SWidget> SContentBrowser::CreateNavigationBar(const FContentBrowserConfig* Config)
+TSharedRef<SWidget> SContentBrowser::CreateNavigationToolBar(const FContentBrowserConfig* Config)
 {
 	FToolMenuContext MenuContext;
 
@@ -2857,6 +2885,13 @@ void SContentBrowser::OnSearchBoxCommitted(const FText& InSearchText, ETextCommi
 	SetSearchBoxText(InSearchText);
 }
 
+void SContentBrowser::Tick(const FGeometry& AllottedGeometry, const double InCurrentTime, const float InDeltaTime)
+{
+	NavigationToolBarParentSwitcher->Update();
+
+	SCompoundWidget::Tick(AllottedGeometry, InCurrentTime, InDeltaTime);
+}
+
 FReply SContentBrowser::OnSearchKeyDown(const FGeometry& Geometry, const FKeyEvent& InKeyEvent)
 {
 	FInputChord CheckChord(InKeyEvent.GetKey(), EModifierKey::FromBools(InKeyEvent.IsControlDown(), InKeyEvent.IsAltDown(), InKeyEvent.IsShiftDown(), InKeyEvent.IsCommandDown()));
diff --git a/Engine/Source/Editor/ContentBrowser/Private/SContentBrowser.h b/Engine/Source/Editor/ContentBrowser/Private/SContentBrowser.h
--- a/Engine/Source/Editor/ContentBrowser/Private/SContentBrowser.h
+++ b/Engine/Source/Editor/ContentBrowser/Private/SContentBrowser.h
@@ -27,6 +27,7 @@
 #include "Widgets/Layout/SSplitter.h"
 #include "Widgets/Navigation/SBreadcrumbTrail.h"
 #include "Widgets/SCompoundWidget.h"
+#include "Widgets/ParentWidgetSwitcher.h"
 
 class FContentBrowserItemDataTemporaryContext;
 class FContentBrowserItemDataUpdate;
@@ -247,6 +248,8 @@ public:
 	/** Called by the editable text control when the user commits a text change */
 	void OnSearchBoxCommitted(const FText& InSearchText, ETextCommit::Type CommitInfo);
 
+	virtual void Tick(const FGeometry& AllottedGeometry, const double InCurrentTime, const float InDeltaTime) override;
+
 private:
 
 	/** Called prior to syncing the selection in this Content Browser */
@@ -608,7 +611,7 @@ private:
 	/** Creates various widgets for the content browser main view */
 	TSharedRef<SWidget> CreateToolBar(const FContentBrowserConfig* Config);
 
-	TSharedRef<SWidget> CreateNavigationBar(const FContentBrowserConfig* Config);
+	TSharedRef<SWidget> CreateNavigationToolBar(const FContentBrowserConfig* Config);
 
 	TSharedRef<SWidget> CreateLockButton(const FContentBrowserConfig* Config);
 	TSharedRef<SWidget> CreateAssetView(const FContentBrowserConfig* Config);
@@ -701,6 +704,12 @@ private:
 	/** Combine breadcrumb/text-box widget for showing & changing location */
 	TSharedPtr<SNavigationBar> NavigationBar;
 
+	/** Handles the navigation bar location switching */
+	TSharedPtr<UE::Editor::ContentBrowser::Private::TWidgetSizeParentSwitcher<Orient_Horizontal>> NavigationToolBarParentSwitcher;
+
+	/** The generated navigation toolbar widget */
+	TSharedPtr<SWidget> NavigationToolBarWidget;
+
 	/** The text box used to search for assets */
 	TSharedPtr<SAssetSearchBox> SearchBoxPtr;
 
diff --git a/Engine/Source/Editor/ContentBrowser/Private/Widgets/ParentWidgetSwitcher.h b/Engine/Source/Editor/ContentBrowser/Private/Widgets/ParentWidgetSwitcher.h
new file mode 100644
--- /dev/null
+++ b/Engine/Source/Editor/ContentBrowser/Private/Widgets/ParentWidgetSwitcher.h
@@ -0,0 +1,137 @@
+// Copyright Epic Games, Inc. All Rights Reserved.
+
+#pragma once
+
+#include "CoreTypes.h"
+#include "Widgets/Layout/SBox.h"
+#include "Widgets/SNullWidget.h"
+#include "Widgets/SWidget.h"
+
+namespace UE::Editor::ContentBrowser::Private
+{
+	/** Takes a single widget and two or more containing widgets (SBox) to switch between, depending on a custom condition. */
+	class FWidgetParentSwitcher
+	{
+	public:
+		virtual ~FWidgetParentSwitcher() = default;
+
+		FWidgetParentSwitcher(
+	        const TSharedPtr<SWidget>& InSubjectWidget,
+	        const TArray<TSharedPtr<SBox>>& InParentWidgets)
+	        : SubjectWidget(InSubjectWidget)
+	        , ParentWidgets(InParentWidgets)
+	    {
+	    }
+
+	    /** Call to check and switch the widget's parent. Return true if switched. */
+	    [[maybe_unused]] bool Update()
+	    {
+	        const int32 TargetParentIdx = GetParentIndex();
+	        if (CurrentSubjectWidgetParentIdx == TargetParentIdx)
+	        {
+	            return false;
+	        }
+
+	        check(ParentWidgets.IsValidIndex(CurrentSubjectWidgetParentIdx));
+	        check(ParentWidgets.IsValidIndex(TargetParentIdx));
+
+			ParentWidgets[CurrentSubjectWidgetParentIdx]->SetContent(SNullWidget::NullWidget);
+			ParentWidgets[TargetParentIdx]->SetContent(SubjectWidget.ToSharedRef());
+
+			CurrentSubjectWidgetParentIdx = TargetParentIdx;
+
+			return true;
+	    }
+
+	protected:
+		virtual int32 GetParentIndex() const = 0;
+
+	protected:
+	    TSharedPtr<SWidget> SubjectWidget;
+	    int32 CurrentSubjectWidgetParentIdx = 0;
+	    TArray<TSharedPtr<SBox>> ParentWidgets;
+	};
+
+	/** Switches parent widgets based on widget size */
+	template <EOrientation Orientation>
+	class TWidgetSizeParentSwitcher
+	    : public FWidgetParentSwitcher
+	{
+	public:
+		virtual ~TWidgetSizeParentSwitcher() override = default;
+
+	    TWidgetSizeParentSwitcher(
+	        const TSharedPtr<SWidget>& InSubjectWidget,
+	        const TArray<TSharedPtr<SBox>>& InParentWidgets,
+	        const TArray<FInt16Range>& InSwitchSizeRanges,
+	        TUniqueFunction<float()>&& InSizeGetter)
+	        : FWidgetParentSwitcher(InSubjectWidget, InParentWidgets)
+			, SwitchSizeRanges(InSwitchSizeRanges)
+			, SizeGetter(MoveTemp(InSizeGetter))
+	    {
+			check(SwitchSizeRanges.Num() == ParentWidgets.Num());
+	    }
+
+	protected:
+		virtual int32 GetParentIndex() const override
+		{
+			const float CurrentSize = SizeGetter();
+			int32 TargetIdx = CurrentSubjectWidgetParentIdx;
+
+			// Returns -1 if less, 0 if in range, 1 if greater
+			auto InRange = [&](const int32 Idx, const int32 Value)
+			{
+				const FInt16Range& Range = SwitchSizeRanges[Idx];
+				if (Value <= Range.GetLowerBoundValue())
+				{
+					return -1;
+				}
+
+				if (Value > Range.GetUpperBoundValue())
+				{
+					return 1;
+				}
+
+				return 0;
+			};
+
+			// Starting from the existing widget, check each size range for the matching parent
+			int32 NumIterations = 0;
+			while (NumIterations < SwitchSizeRanges.Num())
+			{
+				const int32 InRangeResult = InRange(TargetIdx, FMath::FloorToInt(CurrentSize - UE_KINDA_SMALL_NUMBER));
+				if (InRangeResult <= 0)
+				{
+					// Effectively clamp to first widget
+					if (TargetIdx == 0 || InRangeResult == 0)
+					{
+						return TargetIdx;
+					}
+
+					// Otherwise offset TargetIdx
+					TargetIdx -= 1;
+				}
+				else
+				{
+					// Effectively clamp to last widget
+					if (TargetIdx == SwitchSizeRanges.Num() - 1)
+					{
+						return TargetIdx;
+					}
+
+					// Otherwise offset TargetIdx
+					TargetIdx += 1;
+				}
+
+				++NumIterations;
+			}
+
+			checkNoEntry();
+			return INDEX_NONE;
+		}
+
+	private:
+		TArray<FInt16Range> SwitchSizeRanges;
+		TUniqueFunction<float()> SizeGetter;
+	};
+}
