[ContentBrowser] Fixed navigation bar crash

diff --git a/Engine/Source/Editor/ContentBrowser/Private/SContentBrowser.cpp b/Engine/Source/Editor/ContentBrowser/Private/SContentBrowser.cpp
--- a/Engine/Source/Editor/ContentBrowser/Private/SContentBrowser.cpp
+++ b/Engine/Source/Editor/ContentBrowser/Private/SContentBrowser.cpp
@@ -2887,7 +2887,10 @@ void SContentBrowser::OnSearchBoxCommitted(const FText& InSearchText, ETextCommi
 
 void SContentBrowser::Tick(const FGeometry& AllottedGeometry, const double InCurrentTime, const float InDeltaTime)
 {
-	NavigationToolBarParentSwitcher->Update();
+	if (UE::Editor::ContentBrowser::IsNewStyleEnabled())
+	{
+		NavigationToolBarParentSwitcher->Update();
+	}
 
 	SCompoundWidget::Tick(AllottedGeometry, InCurrentTime, InDeltaTime);
 }
