[ContentBrowser] Menu registration fix

diff --git a/Engine/Source/Editor/ContentBrowser/Private/Menus/AssetViewMenus.cpp b/Engine/Source/Editor/ContentBrowser/Private/Menus/AssetViewMenus.cpp
--- a/Engine/Source/Editor/ContentBrowser/Private/Menus/AssetViewMenus.cpp
+++ b/Engine/Source/Editor/ContentBrowser/Private/Menus/AssetViewMenus.cpp
@@ -509,6 +509,12 @@ FDelayedAutoRegisterHelper SAssetView::ToolBarMenuExtensionRegistration(
 	EDelayedRegisterRunPhase::EndOfEngineInit,
 	[]
 	{
+		// @todo: Remove when old style is removed, all extensions are for the new style so early-out if it's not enabled
+		if (!IsNewStyleEnabled())
+		{
+			return;
+		}
+
 		static const FName ToolBarName("ContentBrowser.ToolBar");
 		if (UToolMenu* ToolBarMenu = UToolMenus::Get()->ExtendMenu(ToolBarName))
 		{
