[ContentBrowser] Source tree restyle

diff --git a/Engine/Source/Editor/ContentBrowser/Private/ContentBrowserStyle.cpp b/Engine/Source/Editor/ContentBrowser/Private/ContentBrowserStyle.cpp
--- a/Engine/Source/Editor/ContentBrowser/Private/ContentBrowserStyle.cpp
+++ b/Engine/Source/Editor/ContentBrowser/Private/ContentBrowserStyle.cpp
@@ -74,6 +74,21 @@ namespace UE::ContentBrowser::Private
 
 		Set("ContentBrowser.DefaultFolderColor", FStyleColors::AccentFolder);
 
+		if (Editor::ContentBrowser::IsNewStyleEnabled())
+		{
+			Set("ContentBrowser.AssetTreeHeaderBrush", new FSlateColorBrush(FStyleColors::Header));
+			Set("ContentBrowser.AssetTreeBodyBrush", new FSlateColorBrush(FStyleColors::Header));
+		}
+		else
+		{
+			Set("ContentBrowser.AssetTreeHeaderBrush", new FSlateColorBrush(FStyleColors::Header));
+			Set("ContentBrowser.AssetTreeBodyBrush", new FSlateColorBrush(FStyleColors::Recessed));
+		}
+
+		FExpandableAreaStyle AssetTreeExpandableAreaStyle = FStarshipCoreStyle::GetCoreStyle().GetWidgetStyle<FExpandableAreaStyle>("ExpandableArea");
+
+		Set("ContentBrowser.AssetTreeExpandableArea", AssetTreeExpandableAreaStyle);
+
 		// Asset list view
 		Set("ContentBrowser.AssetListViewNameFont", DEFAULT_FONT("Regular", 12));
 		Set("ContentBrowser.AssetListViewNameFontDirty", DEFAULT_FONT("Bold", 12));
@@ -308,12 +323,20 @@ namespace UE::ContentBrowser::Private
 		Set("ContentBrowser.SaveAllCurrentFolder", new IMAGE_BRUSH_SVG("Starship/Common/SaveCurrent", CoreStyleConstants::Icon16x16));
 		Set("ContentBrowser.ResaveAllCurrentFolder", new IMAGE_BRUSH_SVG("Starship/Common/SaveCurrent", CoreStyleConstants::Icon16x16));
 
+		const FSplitterStyle ContentBrowserSplitterStyle =
+			FSplitterStyle()
+			.SetHandleNormalBrush(FSlateColorBrush(FStyleColors::Transparent))
+			.SetHandleHighlightBrush(FSlateColorBrush(FStyleColors::Transparent));
+
+		Set("ContentBrowser.Splitter", ContentBrowserSplitterStyle);
+
 		FToolBarStyle ContentBrowserToolBarStyle = FStarshipCoreStyle::GetCoreStyle().GetWidgetStyle<FToolBarStyle>("SlimToolBar");
 
 		ContentBrowserToolBarStyle.SetSeparatorBrush(FSlateNoResource());
 		ContentBrowserToolBarStyle.SetSeparatorPadding(FMargin(4.0f, 0.0f));
 		if (UE::Editor::ContentBrowser::IsNewStyleEnabled())
 		{
+			ContentBrowserToolBarStyle.SetAllowWrapButton(false); // Never show the wrap button, just clip
 			ContentBrowserToolBarStyle.SetIconSize(FVector2D(16.0f, 16.0f));
 			ContentBrowserToolBarStyle.SetBackgroundPadding(FMargin(4.0f, 2.0f, 4.0f, 2.0f));
 		}
diff --git a/Engine/Source/Editor/ContentBrowser/Private/SCollectionView.cpp b/Engine/Source/Editor/ContentBrowser/Private/SCollectionView.cpp
--- a/Engine/Source/Editor/ContentBrowser/Private/SCollectionView.cpp
+++ b/Engine/Source/Editor/ContentBrowser/Private/SCollectionView.cpp
@@ -269,6 +269,10 @@ void SCollectionView::Construct( const FArguments& InArgs )
 	UpdateCollectionItems();
 }
 
+bool SCollectionView::IsEmpty() const
+{
+	return AvailableCollections.IsEmpty();
+}
 
 void SCollectionView::HandleCollectionCreated( const FCollectionNameType& Collection )
 {
@@ -953,6 +957,31 @@ void SCollectionView::MakeAddCollectionMenu(TSharedRef<SWidget> MenuParent)
 		);
 }
 
+FVector2D SCollectionView::GetScrollDistance()
+{
+	if (!CollectionTreePtr.IsValid())
+	{
+		return FVector2D::ZeroVector;
+	}
+
+	return CollectionTreePtr->GetScrollDistance();
+}
+
+FVector2D SCollectionView::GetScrollDistanceRemaining()
+{
+	if (!CollectionTreePtr.IsValid())
+	{
+		return FVector2D::ZeroVector;
+	}
+
+	return CollectionTreePtr->GetScrollDistanceRemaining();
+}
+
+TSharedRef<SWidget> SCollectionView::GetScrollWidget()
+{
+	return SharedThis(this);
+}
+
 void SCollectionView::CreateCollectionItem( ECollectionShareType::Type CollectionType, ECollectionStorageMode::Type StorageMode, const FCreateCollectionPayload& InCreationPayload )
 {
 	if ( ensure(CollectionType != ECollectionShareType::CST_All) )
diff --git a/Engine/Source/Editor/ContentBrowser/Private/SCollectionView.h b/Engine/Source/Editor/ContentBrowser/Private/SCollectionView.h
--- a/Engine/Source/Editor/ContentBrowser/Private/SCollectionView.h
+++ b/Engine/Source/Editor/ContentBrowser/Private/SCollectionView.h
@@ -50,7 +50,9 @@ struct FSlateBrush;
 /**
  * The list view of collections.
  */
-class SCollectionView : public SCompoundWidget
+class SCollectionView
+	: public SCompoundWidget
+	, public IScrollableWidget
 {
 	friend class FCollectionContextMenu;
 public:
@@ -92,6 +94,9 @@ public:
 	/** Constructs this widget with InArgs */
 	void Construct( const FArguments& InArgs );
 
+	/** Returns true if there are no collections being presented in this widget. */
+	virtual bool IsEmpty() const;
+
 	/** Selects the specified collections */
 	void SetSelectedCollections(const TArray<FCollectionNameType>& CollectionsToSelect, const bool bEnsureVisible = true);
 
@@ -130,6 +135,12 @@ public:
 	FReply OnAddCollectionClicked();
 	void MakeAddCollectionMenu(TSharedRef<SWidget> MenuParent);
 
+	//~ Begin IScrollableWidget
+	virtual FVector2D GetScrollDistance() override;
+	virtual FVector2D GetScrollDistanceRemaining() override;
+	virtual TSharedRef<SWidget> GetScrollWidget() override;
+	//~ End IScrollableWidget
+
 private:
 
 	/** Payload data for CreateCollectionItem (likely from a delegate binding) */
diff --git a/Engine/Source/Editor/ContentBrowser/Private/SContentBrowser.cpp b/Engine/Source/Editor/ContentBrowser/Private/SContentBrowser.cpp
--- a/Engine/Source/Editor/ContentBrowser/Private/SContentBrowser.cpp
+++ b/Engine/Source/Editor/ContentBrowser/Private/SContentBrowser.cpp
@@ -153,6 +153,7 @@
 #include "Widgets/Navigation/SBreadcrumbTrail.h"
 #include "Widgets/Notifications/SNotificationList.h"
 #include "Widgets/SBoxPanel.h"
+#include "Widgets/SContentBrowserSourceTree.h"
 #include "Widgets/SNullWidget.h"
 #include "Widgets/SWidget.h"
 #include "Widgets/Text/STextBlock.h"
@@ -171,6 +172,8 @@ namespace ContentBrowserSourcesWidgetSwitcherIndex
 	static constexpr int32 CollectionsView = 1;
 }
 
+using namespace UE::Editor::ContentBrowser::Private;
+
 namespace ContentBrowser
 {
 	const FContentBrowserInstanceConfig* GetConstInstanceConfig(FName ForInstance) 
@@ -532,9 +535,32 @@ void SContentBrowser::Construct( const FArguments& InArgs, const FName& InInstan
 	{
 		constexpr float ToolBarVerticalPadding = 4.0f; // ToolBar total height should be 36.0f
 		constexpr float PanelInsetPadding = 2.0f;
+		constexpr float SourceTreeSectionPadding = 2.0f;
 
 		const TSharedRef<SWidget> AssetView = CreateAssetView(Config);
 
+		CreateFavoritesView(Config);
+		CreatePathView(Config);
+		CreateDockedCollectionsView(Config);
+
+		{
+			SAssignNew(SourceTreePtr, SContentBrowserSourceTree)
+
+			+ SContentBrowserSourceTree::Slot()
+			.AreaWidget(FavoritesArea)
+			.Size(0.2f)
+			.Visibility(this, &SContentBrowser::GetFavoriteFolderVisibility)
+
+			+ SContentBrowserSourceTree::Slot()
+			.AreaWidget(PathArea)
+			.Size(0.8f)
+
+			+ SContentBrowserSourceTree::Slot()
+			.AreaWidget(CollectionArea)
+			.Size(0.4f)
+			.Visibility(this, &SContentBrowser::GetDockedCollectionsVisibility);
+		}
+
 		NavigationToolBarWidget = CreateNavigationToolBar(Config);
 
 		TSharedPtr<SBox> AssetViewNavigationToolBarContainer = nullptr;
@@ -561,18 +587,18 @@ void SContentBrowser::Construct( const FArguments& InArgs, const FName& InInstan
 				// Tree + Assets + Navigation Bar
 				+ SHorizontalBox::Slot()
 				.FillWidth(1.0f)
-				.Padding(FMargin(PanelInsetPadding, 0.0f, 0.0f, 0.0f))
+				.Padding(FMargin(0.0f))
 				[
 					SNew(SVerticalBox)
 
 					// Assets/tree
 					+ SVerticalBox::Slot()
 					.FillHeight(1.0f)
-					.Padding(FMargin(0.0f, 0.0f, 0.0f, PanelInsetPadding))
+					.Padding(FMargin(0.0f))
 					[
 						// The tree/assets splitter
 						SAssignNew(PathAssetSplitterPtr, SSplitter)
-						.PhysicalSplitterHandleSize(2.0f)
+						.PhysicalSplitterHandleSize(PanelInsetPadding)
 
 						// Sources View
 						+ SSplitter::Slot()
@@ -586,8 +612,8 @@ void SContentBrowser::Construct( const FArguments& InArgs, const FName& InInstan
 							.WidthOverride(this, &SContentBrowser::GetPathViewBoxWidthOverride)
 							[
 								SNew(SBorder)
-								.Padding(FMargin(0))
-								.BorderImage(FAppStyle::GetBrush("Brushes.Recessed"))
+								.Padding(SourceTreeSectionPadding)
+								.BorderImage(FAppStyle::GetBrush("Brushes.Panel"))
 								[
 									// Note: If adding more widgets here, fix ContentBrowserSourcesWidgetSwitcherIndex and the code that uses it!
 									SAssignNew(SourcesWidgetSwitcher, SWidgetSwitcher)
@@ -595,50 +621,15 @@ void SContentBrowser::Construct( const FArguments& InArgs, const FName& InInstan
 									// Paths View
 									+ SWidgetSwitcher::Slot()
 									[
-										SAssignNew(PathFavoriteSplitterPtr, SSplitter)
-										.Clipping(EWidgetClipping::ClipToBounds)
-										.PhysicalSplitterHandleSize(2.0f)
-										.HitDetectionSplitterHandleSize(8.0f)
-										.Orientation(EOrientation::Orient_Vertical)
-										.MinimumSlotHeight(26.0f)
-										.Visibility(this, &SContentBrowser::GetSourcesViewVisibility)
-										+ SSplitter::Slot()
-										.SizeRule(TAttribute<SSplitter::ESizeRule>(this, &SContentBrowser::GetFavoritesAreaSizeRule))
-										.MinSize(TAttribute<float>(this, &SContentBrowser::GetFavoritesAreaMinSize))
-										.Value(0.2f)
-										[
-											SNew(SBorder)
-											.BorderImage(FAppStyle::Get().GetBrush("Brushes.Header"))
-											.Padding(0.0f, 2.0f, 0.0f, 0.0f)
-											[
-												CreateFavoritesView(Config)
-											]
-										]
-
-										+ SSplitter::Slot()
-										.SizeRule(TAttribute<SSplitter::ESizeRule>(this, &SContentBrowser::GetPathAreaSizeRule))
-										.MinSize(29.0f)
-										.Value(0.8f)
-										[
-											SNew(SBorder)
-											.BorderImage(FAppStyle::Get().GetBrush("Brushes.Header"))
-											.Padding(0.0f, 2.0f, 0.0f, 0.0f)
-											[
-												CreatePathView(Config)
-											]
-										]
-
-										+ SSplitter::Slot()
-										.SizeRule(TAttribute<SSplitter::ESizeRule>(this, &SContentBrowser::GetCollectionsAreaSizeRule))
-										.MinSize(TAttribute<float>(this, &SContentBrowser::GetCollectionsAreaMinSize))
-										.Value(0.4f)
+										// Panel background, seen between items
+										SNew(SBorder)
+										.BorderImage(FAppStyle::GetBrush("WhiteBrush"))
+										.BorderBackgroundColor(FStyleColors::Panel)
+										.HAlign(HAlign_Fill)
+										.VAlign(VAlign_Fill)
+										.Padding(0.0f)
 										[
-											SNew(SBorder)
-											.BorderImage(FAppStyle::Get().GetBrush("Brushes.Header"))
-											.Padding(0.0f, 2.0f, 0.0f, 0.0f)
-											[
-												CreateDockedCollectionsView(Config)
-											]
+											SourceTreePtr.ToSharedRef()
 										]
 									]
 
@@ -1565,11 +1556,6 @@ TSharedRef<SWidget> SContentBrowser::CreateAssetView(const FContentBrowserConfig
 
 void SContentBrowser::SetFavoritesExpanded(bool bExpanded)
 {
-	if (!bExpanded)
-	{
-		FavoritesSearchToggleButton->SetExpanded(false); 
-	}
-
 	if (FContentBrowserInstanceConfig* EditorConfig = GetMutableInstanceConfig())
 	{
 		EditorConfig->bFavoritesExpanded = bExpanded;
@@ -1579,79 +1565,39 @@ void SContentBrowser::SetFavoritesExpanded(bool bExpanded)
 
 TSharedRef<SWidget> SContentBrowser::CreateFavoritesView(const FContentBrowserConfig* Config)
 {
-	// clang-format off
-	TSharedRef<SExpandableArea> FavoritesViewWidget = 
-		SAssignNew(FavoritesArea, SExpandableArea)
-		.BorderImage(FAppStyle::Get().GetBrush("Brushes.Header"))
-		.BodyBorderImage(FAppStyle::Get().GetBrush("Brushes.Recessed"))
-		.HeaderPadding(FMargin(4.0f, 0.0f))
-		.Visibility(this, &SContentBrowser::GetFavoriteFolderVisibility)
-		.Padding(0.f)
-		.AllowAnimatedTransition(false)
-		.OnAreaExpansionChanged(this, &SContentBrowser::SetFavoritesExpanded)
-		.HeaderContent()
-		[
-			SNew(SHorizontalBox)
-			+ SHorizontalBox::Slot()
-			.VAlign(VAlign_Center)
-			[
-				SNew(STextBlock)
-				.Text(LOCTEXT("Favorites", "Favorites"))
-				.TextStyle(FAppStyle::Get(), "ButtonText")
-				.Font(FAppStyle::Get().GetFontStyle("NormalFontBold"))
-			]
-			+SHorizontalBox::Slot()
-			.VAlign(VAlign_Center)
-			.HAlign(HAlign_Right)
-			.AutoWidth()
-			.Padding(4.0f, 0.0f)
-			[
-				SAssignNew(FavoritesSearchToggleButton, SSearchToggleButton, FavoritesSearch->GetWidget())
-				.OnSearchBoxShown_Lambda([this]() { FavoritesArea->SetExpanded(true); })
-			]
-		]
-		.BodyContent()
-		[
-			SNew(SVerticalBox)
-			+ SVerticalBox::Slot()
-			.AutoHeight()
-			[
-				// Should blend in visually with the header but technically acts like part of the body
-				SNew(SBorder)
-				.BorderImage(FAppStyle::Get().GetBrush("Brushes.Header"))
-				.Padding(FMargin(4.0f, 2.0f))
-				[
-					FavoritesSearch->GetWidget()
-				]
-			]
-			+ SVerticalBox::Slot()
-			[
-				SAssignNew(FavoritePathViewPtr, SFavoritePathView)
-				.OnItemSelectionChanged(this, &SContentBrowser::OnItemSelectionChanged, EContentBrowserViewContext::FavoriteView)
-				.OnGetItemContextMenu(this, &SContentBrowser::GetItemContextMenu, EContentBrowserViewContext::FavoriteView)
-				.FocusSearchBoxWhenOpened(false)
-				.ShowTreeTitle(false)
-				.ShowSeparator(false)
-				.AllowClassesFolder(true)
-				.CanShowDevelopersFolder(true)
-				.OwningContentBrowserName(InstanceName)
-				.AddMetaData<FTagMetaData>(FTagMetaData(TEXT("ContentBrowserFavorites")))
-				.ExternalSearch(FavoritesSearch)
-			]
-		];
-	// clang-format on
-	
+	SAssignNew(FavoritePathViewPtr, SFavoritePathView)
+		.OnItemSelectionChanged(this, &SContentBrowser::OnItemSelectionChanged, EContentBrowserViewContext::FavoriteView)
+		.OnGetItemContextMenu(this, &SContentBrowser::GetItemContextMenu, EContentBrowserViewContext::FavoriteView)
+		.FocusSearchBoxWhenOpened(false)
+		.ShowTreeTitle(false)
+		.ShowSeparator(false)
+		.AllowClassesFolder(true)
+		.CanShowDevelopersFolder(true)
+		.OwningContentBrowserName(InstanceName)
+		.AddMetaData<FTagMetaData>(FTagMetaData(TEXT("ContentBrowserFavorites")))
+		.ExternalSearch(FavoritesSearch);
+
+	TSharedRef<UE::Editor::ContentBrowser::Private::SContentBrowserSourceTreeArea> ViewWidget =
+		SAssignNew(
+			FavoritesArea,
+			UE::Editor::ContentBrowser::Private::SContentBrowserSourceTreeArea,
+			"Favorites",
+			FavoritesSearch,
+			FavoritePathViewPtr.ToSharedRef())
+			.Label(LOCTEXT("Favorites", "Favorites"))
+			.OnExpansionChanged(this, &SContentBrowser::SetFavoritesExpanded)
+			.IsEmpty_Lambda([PathView = FavoritePathViewPtr]
+			{
+				return PathView->IsEmpty();
+			})
+			.EmptyBodyLabel(LOCTEXT("FavoritesEmpty", "Right click a folder to add it to your favorites."));
+
 	FavoritePathViewPtr->SetOnFolderFavoriteAdd(SFavoritePathView::FOnFolderFavoriteAdd::CreateSP(this, &SContentBrowser::AddFolderFavorite));
-	return FavoritesViewWidget; 
+	return ViewWidget;
 }
 
 void SContentBrowser::SetPathViewExpanded(bool bExpanded)
 {
-	if (!bExpanded)
-	{
-		PathSearchToggleButton->SetExpanded(false);
-	}
-
 	if (FContentBrowserInstanceConfig* EditorConfig = GetMutableInstanceConfig())
 	{
 		EditorConfig->PathView.bExpanded = bExpanded;
@@ -1661,96 +1607,46 @@ void SContentBrowser::SetPathViewExpanded(bool bExpanded)
 
 TSharedRef<SWidget> SContentBrowser::CreatePathView(const FContentBrowserConfig* Config)
 {
-	// clang-format off
-	return
-		SAssignNew(PathArea, SExpandableArea)
-		.BorderImage(FAppStyle::Get().GetBrush("Brushes.Header"))
-		.BodyBorderImage(FAppStyle::Get().GetBrush("Brushes.Recessed"))
-		.HeaderPadding(FMargin(4.0f, 0.0f))
-		.Padding(0.f)
-		.AllowAnimatedTransition(false)
-		.OnAreaExpansionChanged(this, &SContentBrowser::SetPathViewExpanded)
-		.HeaderContent()
-		[
-			SNew(SHorizontalBox)
-			+ SHorizontalBox::Slot()
-			.VAlign(VAlign_Center)
-			[
-				SNew(STextBlock)
-				.Text(FText::FromString(FApp::GetProjectName()))
-				.TextStyle(FAppStyle::Get(), "ButtonText")
-				.Font(FAppStyle::Get().GetFontStyle("NormalFontBold"))
-			]
-			+SHorizontalBox::Slot()
-			.VAlign(VAlign_Center)
-			.HAlign(HAlign_Right)
-			.AutoWidth()
-			.Padding(4.0f, 0.0f)
-			[
-				SAssignNew(PathSearchToggleButton, SSearchToggleButton, SourcesSearch->GetWidget())
-				.OnSearchBoxShown_Lambda([this]() { PathArea->SetExpanded(true); })
-			]
-		]
-		.BodyContent()
-		[
-			SNew(SVerticalBox)
-			+ SVerticalBox::Slot()
-			.AutoHeight()
-			[
-				// Should blend in visually with the header but technically acts like part of the body
-				SNew(SBorder)
-				.BorderImage(FAppStyle::Get().GetBrush("Brushes.Header"))
-				.Padding(FMargin(4.0f, 2.0f))
-				[
-					SourcesSearch->GetWidget()
-				]
-			]
-			+ SVerticalBox::Slot()
-			[
-				SNew(SBorder)
-				.BorderImage(FAppStyle::Get().GetBrush("Brushes.Recessed"))
-				[
-					SAssignNew(PathViewPtr, SPathView)
-					.OnItemSelectionChanged(this, &SContentBrowser::OnItemSelectionChanged, EContentBrowserViewContext::PathView)
-					.OnGetItemContextMenu(this, &SContentBrowser::GetItemContextMenu, EContentBrowserViewContext::PathView)
-					.FocusSearchBoxWhenOpened(false)
-					.ShowTreeTitle(false)
-					.ShowSeparator(false)
-				   	.ShowRedirectors_Lambda([this]() { return ContentBrowserUtils::ShouldShowRedirectors(FilterListPtr); })
-					.AddMetaData<FTagMetaData>(FTagMetaData(TEXT("ContentBrowserSources")))
-					.ExternalSearch(SourcesSearch)
-					.PluginPathFilters(PluginPathFilters)
-					.OwningContentBrowserName(InstanceName)
-					.AllowClassesFolder(Config != nullptr ? Config->bCanShowClasses : true)
-					.CanShowDevelopersFolder(Config != nullptr ? Config->bCanShowDevelopersFolder : true)
-				]
-			]
-		];
-	// clang-format on
+	SAssignNew(PathViewPtr, SPathView)
+		.OnItemSelectionChanged(this, &SContentBrowser::OnItemSelectionChanged, EContentBrowserViewContext::PathView)
+		.OnGetItemContextMenu(this, &SContentBrowser::GetItemContextMenu, EContentBrowserViewContext::PathView)
+		.FocusSearchBoxWhenOpened(false)
+		.ShowTreeTitle(false)
+		.ShowSeparator(false)
+		.ShowRedirectors_Lambda([this]() { return ContentBrowserUtils::ShouldShowRedirectors(FilterListPtr); })
+		.AddMetaData<FTagMetaData>(FTagMetaData(TEXT("ContentBrowserSources")))
+		.ExternalSearch(SourcesSearch)
+		.PluginPathFilters(PluginPathFilters)
+		.OwningContentBrowserName(InstanceName)
+		.AllowClassesFolder(Config != nullptr ? Config->bCanShowClasses : true)
+		.CanShowDevelopersFolder(Config != nullptr ? Config->bCanShowDevelopersFolder : true);
+
+	TSharedRef<UE::Editor::ContentBrowser::Private::SContentBrowserSourceTreeArea> ViewWidget =
+		SAssignNew(
+			PathArea,
+			UE::Editor::ContentBrowser::Private::SContentBrowserSourceTreeArea,
+			"Path",
+			SourcesSearch,
+			PathViewPtr.ToSharedRef())
+			.Label(FText::FromString(FApp::GetProjectName()))
+			.OnExpansionChanged(this, &SContentBrowser::SetPathViewExpanded);
+
+	return ViewWidget;
 }
 
 TSharedRef<SWidget> SContentBrowser::CreateDockedCollectionsView(const FContentBrowserConfig* Config)
 {
-	return
-		SAssignNew(CollectionArea, SExpandableArea)
-		.BorderImage(FAppStyle::Get().GetBrush("Brushes.Header"))
-		.BodyBorderImage(FAppStyle::Get().GetBrush("Brushes.Recessed"))
-		.HeaderPadding(FMargin(4.0f, 0.0f))
-		.Padding(0.f)
-		.Visibility(this, &SContentBrowser::GetDockedCollectionsVisibility)
-		.AllowAnimatedTransition(false)
-		.OnAreaExpansionChanged_Lambda([this](bool bIsExpanded) { if (!bIsExpanded) CollectionSearchToggleButton->SetExpanded(false); })
+	return SAssignNew(
+		CollectionArea,
+		UE::Editor::ContentBrowser::Private::SContentBrowserSourceTreeArea,
+		"Collection",
+		CollectionSearch,
+		CollectionViewPtr.ToSharedRef())
+		.Label(LOCTEXT("CollectionsTitle", "Collections"))
 		.HeaderContent()
 		[
 			SNew(SHorizontalBox)
-			+ SHorizontalBox::Slot()
-			.VAlign(VAlign_Center)
-			[
-				SNew(STextBlock)
-				.Text(LOCTEXT("CollectionsTitle", "Collections"))
-				.TextStyle(FAppStyle::Get(), "ButtonText")
-				.Font(FAppStyle::Get().GetFontStyle("NormalFontBold"))
-			]
+
 			+ SHorizontalBox::Slot()
 			.VAlign(VAlign_Center)
 			.HAlign(HAlign_Right)
@@ -1768,39 +1664,12 @@ TSharedRef<SWidget> SContentBrowser::CreateDockedCollectionsView(const FContentB
 					.ColorAndOpacity(FSlateColor::UseForeground())
 				]
 			]
-			+ SHorizontalBox::Slot()
-			.VAlign(VAlign_Center)
-			.HAlign(HAlign_Right)
-			.AutoWidth()
-			.Padding(4.0f, 0.0f)
-			[
-				SAssignNew(CollectionSearchToggleButton, SSearchToggleButton, CollectionSearch->GetWidget())
-				.OnSearchBoxShown_Lambda([this]() { CollectionArea->SetExpanded(true); })
-			]
 		]
-		.BodyContent()
-		[
-			SNew(SVerticalBox)
-			+ SVerticalBox::Slot()
-			.AutoHeight()
-			[
-				// Should blend in visually with the header but technically acts like part of the body
-				SNew(SBorder)
-				.BorderImage(FAppStyle::Get().GetBrush("Brushes.Header"))
-				.Padding(FMargin(4.0f, 2.0f))
-				[
-					CollectionSearch->GetWidget()
-				]
-			]
-			+ SVerticalBox::Slot()
-			[
-				SNew(SBorder)
-				.BorderImage(FAppStyle::Get().GetBrush("Brushes.Recessed"))
-				[
-					CollectionViewPtr.ToSharedRef()
-				]
-			]
-		];
+		.IsEmpty_Lambda([PathView = CollectionViewPtr]
+		{
+			return PathView->IsEmpty();
+		})
+		.EmptyBodyLabel(LOCTEXT("CollectionsEmpty", "Click the <img src=\"Icons.PlusCircle\"/> in the section header to create a collection."));
 }
 
 TSharedRef<SWidget> SContentBrowser::CreateDrawerDockButton(const FContentBrowserConfig* Config)
@@ -1887,26 +1756,29 @@ void SContentBrowser::ExtendViewOptionsMenu(const FContentBrowserConfig* Config)
 						}
 					}
 
-					if (bShowSourcesView)
+					if (!UE::Editor::ContentBrowser::IsNewStyleEnabled())
 					{
-						FToolMenuSection& Section = InMenu->FindOrAddSection(UE::Editor::ContentBrowser::IsNewStyleEnabled() ? "Show" : "View");
-						Section.AddMenuEntry(
-							"ToggleSources",
-							UE::Editor::ContentBrowser::IsNewStyleEnabled()
-							? LOCTEXT("ToggleSourcesView_NewStyle", "Sources Panel")
-							: LOCTEXT("ToggleSourcesView", "Show Sources Panel"),
-							LOCTEXT("ToggleSourcesView_Tooltip", "Show or hide the sources panel"),
-							TAttribute<FSlateIcon>(),
-							FUIAction(
-								FExecuteAction::CreateLambda([ContentBrowser = Context->OwningContentBrowser]() { ContentBrowser.Pin()->SourcesViewExpandClicked(); }),
-								FCanExecuteAction(),
-								FIsActionChecked::CreateLambda([ContentBrowser = Context->OwningContentBrowser]() { return ContentBrowser.Pin()->bSourcesViewExpanded; })),
-							EUserInterfaceActionType::ToggleButton
-						)
-						.InsertPosition =
-							UE::Editor::ContentBrowser::IsNewStyleEnabled()
-							? FToolMenuInsert(NAME_None, EToolMenuInsertType::First)
-							: FToolMenuInsert();
+						if (bShowSourcesView)
+						{
+							FToolMenuSection& Section = InMenu->FindOrAddSection(UE::Editor::ContentBrowser::IsNewStyleEnabled() ? "Show" : "View");
+							Section.AddMenuEntry(
+								"ToggleSources",
+								UE::Editor::ContentBrowser::IsNewStyleEnabled()
+								? LOCTEXT("ToggleSourcesView_NewStyle", "Sources Panel")
+								: LOCTEXT("ToggleSourcesView", "Show Sources Panel"),
+								LOCTEXT("ToggleSourcesView_Tooltip", "Show or hide the sources panel"),
+								TAttribute<FSlateIcon>(),
+								FUIAction(
+									FExecuteAction::CreateLambda([ContentBrowser = Context->OwningContentBrowser]() { ContentBrowser.Pin()->SourcesViewExpandClicked(); }),
+									FCanExecuteAction(),
+									FIsActionChecked::CreateLambda([ContentBrowser = Context->OwningContentBrowser]() { return ContentBrowser.Pin()->bSourcesViewExpanded; })),
+								EUserInterfaceActionType::ToggleButton
+							)
+							.InsertPosition =
+								UE::Editor::ContentBrowser::IsNewStyleEnabled()
+								? FToolMenuInsert(NAME_None, EToolMenuInsertType::First)
+								: FToolMenuInsert();
+						}
 					}
 				}
 			}
@@ -1916,7 +1788,7 @@ void SContentBrowser::ExtendViewOptionsMenu(const FContentBrowserConfig* Config)
 
 SSplitter::ESizeRule SContentBrowser::GetFavoritesAreaSizeRule() const
 {
-	// Make sure the area is expanded and visible 
+	// Make sure the area is expanded and visible
 	return FavoritesArea->IsExpanded() && GetFavoriteFolderVisibility() == EVisibility::Visible ? SSplitter::ESizeRule::FractionOfParent : SSplitter::ESizeRule::SizeToContent;
 }
 
@@ -1927,7 +1799,7 @@ SSplitter::ESizeRule SContentBrowser::GetPathAreaSizeRule() const
 
 SSplitter::ESizeRule SContentBrowser::GetCollectionsAreaSizeRule() const
 {
-	// Make sure the area is expanded and visible 
+	// Make sure the area is expanded and visible
 	return CollectionArea->IsExpanded() && GetDockedCollectionsVisibility() == EVisibility::Visible ? SSplitter::ESizeRule::FractionOfParent : SSplitter::ESizeRule::SizeToContent;
 }
 
@@ -1953,12 +1825,22 @@ FOptionalSize SContentBrowser::GetFilterViewBoxWidthOverride() const
 
 float SContentBrowser::GetFavoritesAreaMinSize() const
 {
-	return GetFavoriteFolderVisibility() == EVisibility::Visible ? 29.0f : 0.0f;
+	static const float SourceTreeHeaderHeightMin =
+		UE::Editor::ContentBrowser::IsNewStyleEnabled()
+		? 36.0f
+		: 26.0f + 3.0f;
+
+	return GetFavoriteFolderVisibility() == EVisibility::Visible ? SourceTreeHeaderHeightMin : 0.0f;
 }
 
 float SContentBrowser::GetCollectionsAreaMinSize() const
 {
-	return GetDockedCollectionsVisibility() == EVisibility::Visible ? 29.0f : 0.0f;
+	static const float SourceTreeHeaderHeightMin =
+		UE::Editor::ContentBrowser::IsNewStyleEnabled()
+		? 36.0f
+		: 26.0f + 3.0f;
+
+	return GetDockedCollectionsVisibility() == EVisibility::Visible ? SourceTreeHeaderHeightMin : 0.0f;
 }
 
 FText SContentBrowser::GetHighlightedText() const
@@ -2366,13 +2248,9 @@ void SContentBrowser::SaveSettings() const
 	GConfig->SetBool(*SettingsIniSection, *(SettingsString + TEXT(".SourcesExpanded")), bSourcesViewExpanded, GEditorPerProjectIni);
 	GConfig->SetBool(*SettingsIniSection, *(SettingsString + TEXT(".IsLocked")), bIsLocked, GEditorPerProjectIni);
 
-	GConfig->SetBool(*SettingsIniSection, *(SettingsString + TEXT(".FavoritesAreaExpanded")), FavoritesArea->IsExpanded(), GEditorPerProjectIni);
-	GConfig->SetBool(*SettingsIniSection, *(SettingsString + TEXT(".PathAreaExpanded")), PathArea->IsExpanded(), GEditorPerProjectIni);
-	GConfig->SetBool(*SettingsIniSection, *(SettingsString + TEXT(".CollectionAreaExpanded")), CollectionArea->IsExpanded(), GEditorPerProjectIni);
-
-	GConfig->SetBool(*SettingsIniSection, *(SettingsString + TEXT(".FavoritesSearchAreaExpanded")), FavoritesSearchToggleButton->IsExpanded(), GEditorPerProjectIni);
-	GConfig->SetBool(*SettingsIniSection, *(SettingsString + TEXT(".PathSearchAreaExpanded")), PathSearchToggleButton->IsExpanded(), GEditorPerProjectIni);
-	GConfig->SetBool(*SettingsIniSection, *(SettingsString + TEXT(".CollectionSearchAreaExpanded")), CollectionSearchToggleButton->IsExpanded(), GEditorPerProjectIni);
+	FavoritesArea->SaveSettings(InstanceName);
+	PathArea->SaveSettings(InstanceName);
+	CollectionArea->SaveSettings(InstanceName);
 
 	for(int32 SlotIndex = 0; SlotIndex < PathAssetSplitterPtr->GetChildren()->Num(); SlotIndex++)
 	{
@@ -2391,12 +2269,18 @@ void SContentBrowser::SaveSettings() const
 		}
 	}
 
-	for (int32 SlotIndex = 0; SlotIndex < PathFavoriteSplitterPtr->GetChildren()->Num(); SlotIndex++)
+	if (UE::Editor::ContentBrowser::IsNewStyleEnabled())
 	{
-		float SplitterSize = PathFavoriteSplitterPtr->SlotAt(SlotIndex).GetSizeValue();
-		GConfig->SetFloat(*SettingsIniSection, *(SettingsString + FString::Printf(TEXT(".FavoriteSplitter.SlotSize%d"), SlotIndex)), SplitterSize, GEditorPerProjectIni);
+		SourceTreePtr->SaveSettings(InstanceName);
+	}
+	else
+	{
+		for (int32 SlotIndex = 0; SlotIndex < PathFavoriteSplitterPtr->GetChildren()->Num(); SlotIndex++)
+		{
+			float SplitterSize = PathFavoriteSplitterPtr->SlotAt(SlotIndex).GetSizeValue();
+			GConfig->SetFloat(*SettingsIniSection, *(SettingsString + FString::Printf(TEXT(".FavoriteSplitter.SlotSize%d"), SlotIndex)), SplitterSize, GEditorPerProjectIni);
+		}
 	}
-
 
 	// Save all our data using the settings string as a key in the user settings ini
 	FilterListPtr->SaveSettings(GEditorPerProjectIni, SettingsIniSection, SettingsString);
@@ -2515,29 +2399,9 @@ void SContentBrowser::LoadSettings(const FName& InInstanceName)
 	bIsLocked = false;
 	GConfig->GetBool(*SettingsIniSection, *(SettingsString + TEXT(".IsLocked")), bIsLocked, GEditorPerProjectIni);
 
-	bool bFavoritesAreaExpanded = false;
-	GConfig->GetBool(*SettingsIniSection, *(SettingsString + TEXT(".FavoritesAreaExpanded")), bFavoritesAreaExpanded, GEditorPerProjectIni);
-	FavoritesArea->SetExpanded(bFavoritesAreaExpanded);
-
-	bool bPathAreaExpanded = true;
-	GConfig->GetBool(*SettingsIniSection, *(SettingsString + TEXT(".PathAreaExpanded")), bPathAreaExpanded, GEditorPerProjectIni);
-	PathArea->SetExpanded(bPathAreaExpanded);
-
-	bool bCollectionAreaExpanded = false;
-	GConfig->GetBool(*SettingsIniSection, *(SettingsString + TEXT(".CollectionAreaExpanded")), bCollectionAreaExpanded, GEditorPerProjectIni);
-	CollectionArea->SetExpanded(bCollectionAreaExpanded);
-
-	bool bFavoritesSearchAreaExpanded = false;
-	GConfig->GetBool(*SettingsIniSection, *(SettingsString + TEXT(".FavoritesSearchAreaExpanded")), bFavoritesSearchAreaExpanded, GEditorPerProjectIni);
-	FavoritesSearchToggleButton->SetExpanded(bFavoritesSearchAreaExpanded);
-
-	bool bPathSearchAreaExpanded = false;
-	GConfig->GetBool(*SettingsIniSection, *(SettingsString + TEXT(".PathSearchAreaExpanded")), bPathSearchAreaExpanded, GEditorPerProjectIni);
-	PathSearchToggleButton->SetExpanded(bPathSearchAreaExpanded);
-
-	bool bCollectionSearchAreaExpanded = false;
-	GConfig->GetBool(*SettingsIniSection, *(SettingsString + TEXT(".CollectionSearchAreaExpanded")), bCollectionSearchAreaExpanded, GEditorPerProjectIni);
-	CollectionSearchToggleButton->SetExpanded(bCollectionSearchAreaExpanded);
+	FavoritesArea->LoadSettings(InInstanceName);
+	PathArea->LoadSettings(InInstanceName);
+	CollectionArea->LoadSettings(InInstanceName);
 
 	for(int32 SlotIndex = 0; SlotIndex < PathAssetSplitterPtr->GetChildren()->Num(); SlotIndex++)
 	{
@@ -2558,11 +2422,18 @@ void SContentBrowser::LoadSettings(const FName& InInstanceName)
 		}
 	}
 
-	for (int32 SlotIndex = 0; SlotIndex < PathFavoriteSplitterPtr->GetChildren()->Num(); SlotIndex++)
+	if (UE::Editor::ContentBrowser::IsNewStyleEnabled())
+	{
+		SourceTreePtr->LoadSettings(InstanceName);
+	}
+	else
 	{
-		float SplitterSize = PathFavoriteSplitterPtr->SlotAt(SlotIndex).GetSizeValue();
-		GConfig->GetFloat(*SettingsIniSection, *(SettingsString + FString::Printf(TEXT(".FavoriteSplitter.SlotSize%d"), SlotIndex)), SplitterSize, GEditorPerProjectIni);
-		PathFavoriteSplitterPtr->SlotAt(SlotIndex).SetSizeValue(SplitterSize);
+		for (int32 SlotIndex = 0; SlotIndex < PathFavoriteSplitterPtr->GetChildren()->Num(); SlotIndex++)
+		{
+			float SplitterSize = PathFavoriteSplitterPtr->SlotAt(SlotIndex).GetSizeValue();
+			GConfig->GetFloat(*SettingsIniSection, *(SettingsString + FString::Printf(TEXT(".FavoriteSplitter.SlotSize%d"), SlotIndex)), SplitterSize, GEditorPerProjectIni);
+			PathFavoriteSplitterPtr->SlotAt(SlotIndex).SetSizeValue(SplitterSize);
+		}
 	}
 
 	// Save all our data using the settings string as a key in the user settings ini
diff --git a/Engine/Source/Editor/ContentBrowser/Private/SContentBrowser.h b/Engine/Source/Editor/ContentBrowser/Private/SContentBrowser.h
--- a/Engine/Source/Editor/ContentBrowser/Private/SContentBrowser.h
+++ b/Engine/Source/Editor/ContentBrowser/Private/SContentBrowser.h
@@ -29,6 +29,13 @@
 #include "Widgets/SCompoundWidget.h"
 #include "Widgets/ParentWidgetSwitcher.h"
 
+namespace UE::Editor::ContentBrowser::Private
+{
+	class SContentBrowserSourceTree;
+	class FContentBrowserSourceTreeEntry;
+	class SContentBrowserSourceTreeArea;
+}
+
 class FContentBrowserItemDataTemporaryContext;
 class FContentBrowserItemDataUpdate;
 class FExtender;
@@ -689,6 +696,8 @@ private:
 	/** The sources search for collections*/
 	TSharedPtr<FSourcesSearch> CollectionSearch;
 
+	TSharedPtr<UE::Editor::ContentBrowser::Private::SContentBrowserSourceTree> SourceTreePtr;
+
 	/** The asset tree widget */
 	TSharedPtr<SPathView> PathViewPtr;
 
@@ -726,22 +735,13 @@ private:
 	TSharedPtr<SComboButton> PathPickerButton;
 
 	/** Favorites area widget */
-	TSharedPtr<SExpandableArea> FavoritesArea;
-
-	/** Toggle button for showing/hiding favorites search area */
-	TSharedPtr<SSearchToggleButton> FavoritesSearchToggleButton;
+	TSharedPtr<UE::Editor::ContentBrowser::Private::SContentBrowserSourceTreeArea> FavoritesArea;
 
 	/** Path area widget */
-	TSharedPtr<SExpandableArea> PathArea;
-
-	/** Toggle button for showing/hiding path search area */
-	TSharedPtr<SSearchToggleButton> PathSearchToggleButton;
+	TSharedPtr<UE::Editor::ContentBrowser::Private::SContentBrowserSourceTreeArea> PathArea;
 
 	/** Collection area widget */
-	TSharedPtr<SExpandableArea> CollectionArea;
-
-	/** Toggle button for showing/hiding collection search area */
-	TSharedPtr<SSearchToggleButton> CollectionSearchToggleButton;
+	TSharedPtr<UE::Editor::ContentBrowser::Private::SContentBrowserSourceTreeArea> CollectionArea;
 
 	/** Index of the active sources widget */
 	int32 ActiveSourcesWidgetIndex = 0;
diff --git a/Engine/Source/Editor/ContentBrowser/Private/SPathView.cpp b/Engine/Source/Editor/ContentBrowser/Private/SPathView.cpp
--- a/Engine/Source/Editor/ContentBrowser/Private/SPathView.cpp
+++ b/Engine/Source/Editor/ContentBrowser/Private/SPathView.cpp
@@ -1223,6 +1223,11 @@ FReply SPathView::OnKeyDown(const FGeometry& MyGeometry, const FKeyEvent& InKeyE
 	return FReply::Unhandled();
 }
 
+bool SPathView::IsEmpty() const
+{
+	return TreeViewPtr->GetRootItems().IsEmpty();
+}
+
 void SPathView::PopulatePathViewFiltersMenu(UToolMenu* Menu)
 {
 	{
@@ -1649,6 +1654,31 @@ EContentBrowserItemAttributeFilter SPathView::GetContentBrowserItemAttributeFilt
 			| (bDisplayL10NContent ? EContentBrowserItemAttributeFilter::IncludeLocalized : EContentBrowserItemAttributeFilter::IncludeNone);
 }
 
+FVector2D SPathView::GetScrollDistance()
+{
+	if (!TreeViewPtr.IsValid())
+	{
+		return FVector2D::ZeroVector;
+	}
+
+	return TreeViewPtr->GetScrollDistance();
+}
+
+FVector2D SPathView::GetScrollDistanceRemaining()
+{
+	if (!TreeViewPtr.IsValid())
+	{
+		return FVector2D::ZeroVector;
+	}
+
+	return TreeViewPtr->GetScrollDistanceRemaining();
+}
+
+TSharedRef<SWidget> SPathView::GetScrollWidget()
+{
+	return SharedThis(this);
+}
+
 bool SPathView::InternalPathPassesBlockLists(const FStringView InInternalPath, const int32 InAlreadyCheckedDepth) const
 {
 	TArray<const FPathPermissionList*, TInlineAllocator<2>> BlockLists;
diff --git a/Engine/Source/Editor/ContentBrowser/Private/SPathView.h b/Engine/Source/Editor/ContentBrowser/Private/SPathView.h
--- a/Engine/Source/Editor/ContentBrowser/Private/SPathView.h
+++ b/Engine/Source/Editor/ContentBrowser/Private/SPathView.h
@@ -61,7 +61,9 @@ struct FContentBrowserInstanceConfig;
 /**
  * The tree view of folders which contain content.
  */
-class SPathView : public SCompoundWidget
+class SPathView
+	: public SCompoundWidget
+	, public IScrollableWidget
 {
 	using Super = SCompoundWidget;
 
@@ -172,6 +174,9 @@ public:
 	/** Process the Commands of the PathView */
 	virtual FReply OnKeyDown( const FGeometry& MyGeometry, const FKeyEvent& InKeyEvent ) override;
 
+	/** Returns true if there are no items being presented in this widget. */
+	virtual bool IsEmpty() const;
+
 	/** Selects the closest matches to the supplied paths in the tree. "/" delimited */
 	void SetSelectedPaths(const TArray<FName>& Paths);
 
@@ -281,6 +286,12 @@ public:
 	/** Get current item attribute filter enum */
 	EContentBrowserItemAttributeFilter GetContentBrowserItemAttributeFilter() const;
 
+	//~ Begin IScrollableWidget
+	virtual FVector2D GetScrollDistance() override;
+	virtual FVector2D GetScrollDistanceRemaining() override;
+	virtual TSharedRef<SWidget> GetScrollWidget() override;
+	//~ End IScrollableWidget
+
 protected:
 	/** Expands all parents of the specified item */
 	void RecursiveExpandParents(const TSharedPtr<FTreeItem>& Item);
diff --git a/Engine/Source/Editor/ContentBrowser/Private/Widgets/SContentBrowserSourceTree.cpp b/Engine/Source/Editor/ContentBrowser/Private/Widgets/SContentBrowserSourceTree.cpp
new file mode 100644
--- /dev/null
+++ b/Engine/Source/Editor/ContentBrowser/Private/Widgets/SContentBrowserSourceTree.cpp
@@ -0,0 +1,470 @@
+// Copyright Epic Games, Inc. All Rights Reserved.
+
+#include "SContentBrowserSourceTree.h"
+
+#include "ContentBrowserStyle.h"
+#include "Framework/Layout/IScrollableWidget.h"
+#include "Misc/ConfigCacheIni.h"
+#include "SContentBrowser.h"
+#include "SlateOptMacros.h"
+#include "SourcesSearch.h"
+#include "SSearchToggleButton.h"
+#include "Widgets/Colors/SColorBlock.h"
+#include "Widgets/Input/SSearchBox.h"
+#include "Widgets/Layout/SExpandableArea.h"
+#include "Widgets/Layout/SScrollBorder.h"
+#include "Widgets/Layout/SWidgetSwitcher.h"
+#include "Widgets/SBoxPanel.h"
+#include "Widgets/SWidget.h"
+#include "Widgets/Text/SRichTextBlock.h"
+#include "Widgets/Text/STextBlock.h"
+
+#define LOCTEXT_NAMESPACE "ContentBrowser"
+
+using namespace UE::ContentBrowser::Private;
+
+namespace UE::Editor::ContentBrowser::Private
+{
+	static constexpr float MinBodyHeight = 88.0f;
+
+	BEGIN_SLATE_FUNCTION_BUILD_OPTIMIZATION
+
+	void SContentBrowserSourceTree::FSlot::Construct(const FChildren& SlotOwner, FSlotArguments&& InArgs)
+	{
+		TSlotBase<FSlot>::Construct(SlotOwner, MoveTemp(InArgs));
+
+		check(EntryWidget.IsValid() || InArgs._AreaWidget.IsValid());
+
+		EntryWidget = InArgs._AreaWidget;
+
+		ExpandedSizeRule = InArgs._ExpandedSizeRule.Get(SSplitter::ESizeRule::FractionOfParent);
+
+		if (InArgs._Visibility.IsSet())
+		{
+			EntryVisibility = InArgs._Visibility;
+		}
+		else
+		{
+			EntryVisibility = EVisibility::Visible;
+		}
+
+		if (InArgs._HeaderHeight.IsSet())
+		{
+			HeaderHeight = InArgs._HeaderHeight.GetValue();
+		}
+		else
+		{
+			HeaderHeight =
+				IsNewStyleEnabled()
+				? 36.0f
+				: 26.0f + 3.0f;
+		}
+
+		SlotSize = InitialSlotSize = InArgs._Size.Get(0.5f);
+	}
+
+	SSplitter::ESizeRule SContentBrowserSourceTree::FSlot::GetExpandedSizeRule() const
+	{
+		// If the slot is empty, we size to content (to enforce a fixed size and disallow resizing)
+		return EntryWidget->IsExpanded() && IsVisible() && !EntryWidget->IsEmpty()
+			? ExpandedSizeRule
+			: SSplitter::ESizeRule::SizeToContent;
+	}
+
+	float SContentBrowserSourceTree::FSlot::GetMinHeight() const
+	{
+		return IsVisible()
+			? EntryWidget->IsExpanded() ? HeaderHeight + MinBodyHeight : HeaderHeight
+			: 0.0f;
+	}
+
+	float SContentBrowserSourceTree::FSlot::GetInitialSlotSize() const
+	{
+		return InitialSlotSize;
+	}
+
+	float SContentBrowserSourceTree::FSlot::GetSlotSize() const
+	{
+ 		return SlotSize;
+	}
+
+	TSharedPtr<SContentBrowserSourceTreeArea> SContentBrowserSourceTree::FSlot::GetEntryWidget() const
+	{
+		return EntryWidget;
+	}
+
+	bool SContentBrowserSourceTree::FSlot::IsVisible() const
+	{
+		return EntryVisibility.Get() == EVisibility::Visible;
+	}
+
+	void SContentBrowserSourceTree::FSlot::OnSlotResized(float InNewSize)
+	{
+		SlotSize = InNewSize;
+	}
+
+	SContentBrowserSourceTree::FSlot::FSlotArguments SContentBrowserSourceTree::Slot()
+	{
+		return FSlot::FSlotArguments(MakeUnique<FSlot>());
+	}
+
+	SContentBrowserSourceTree::FScopedWidgetSlotArguments SContentBrowserSourceTree::AddSlot()
+	{
+		return FScopedWidgetSlotArguments{ MakeUnique<FSlot>(), Slots, INDEX_NONE,
+			[this, Splitter = Splitter](const FSlot* InSlotAdded, int32 InSlotIdx)
+			{
+				// When the slot is added, also add it to the splitter
+				Splitter->AddSlot()
+				.Value(InSlotAdded->GetInitialSlotSize())
+				.SizeRule(this, &SContentBrowserSourceTree::GetExpandedSizeRule, InSlotIdx)
+				.MinSize(this, &SContentBrowserSourceTree::GetMinHeight, InSlotIdx)
+				[
+					InSlotAdded->GetEntryWidget().ToSharedRef()
+				];
+			}};
+	}
+
+	float SContentBrowserSourceTree::GetMinHeight(const int32 InSlotIdx) const
+	{
+		if (Slots.IsValidIndex(InSlotIdx))
+		{
+			return Slots[InSlotIdx].GetMinHeight();
+		}
+
+		return 0.0f;
+	}
+
+	SSplitter::ESizeRule SContentBrowserSourceTree::GetExpandedSizeRule(const int32 InSlotIdx) const
+	{
+		if (Slots.IsValidIndex(InSlotIdx))
+		{
+			return Slots[InSlotIdx].GetExpandedSizeRule();
+		}
+
+		return SSplitter::FractionOfParent;
+	}
+
+	float SContentBrowserSourceTree::GetSlotSize(const int32 InSlotIdx) const
+	{
+		if (Slots.IsValidIndex(InSlotIdx))
+		{
+			return Slots[InSlotIdx].GetSlotSize();
+		}
+
+		return 0.0f;
+	}
+
+	void SContentBrowserSourceTree::OnSlotResized(float InNewSize, const int32 InSlotIdx)
+	{
+		if (Slots.IsValidIndex(InSlotIdx))
+		{
+			Slots[InSlotIdx].OnSlotResized(InNewSize);
+		}
+	}
+
+	void SContentBrowserSourceTree::Construct(const FArguments& InArgs)
+	{
+		constexpr float SourceTreeSectionPadding = 2.0f;
+		constexpr float SourceTreeSectionHandleSize = 8.0f;
+		constexpr float SourceTreeHeaderHeight = 36.0f;
+
+		SAssignNew(Splitter, SSplitter)
+		.Style(&FContentBrowserStyle::Get().GetWidgetStyle<FSplitterStyle>("ContentBrowser.Splitter"))
+		.Clipping(EWidgetClipping::ClipToBounds)
+		.PhysicalSplitterHandleSize(SourceTreeSectionPadding)
+		.HitDetectionSplitterHandleSize(SourceTreeSectionHandleSize)
+		.Orientation(EOrientation::Orient_Vertical)
+		.MinimumSlotHeight(SourceTreeHeaderHeight);
+
+		for (const FSlot::FSlotArguments& Slot : InArgs._Slots)
+		{
+			AddSlot()
+			.AreaWidget(Slot._AreaWidget)
+			.Size(Slot._Size)
+			.Visibility(Slot._Visibility);
+		}
+
+		ChildSlot
+		.HAlign(HAlign_Fill)
+	    .VAlign(VAlign_Fill)
+		[
+			SNew(SVerticalBox)
+
+			+ SVerticalBox::Slot()
+			.HAlign(HAlign_Fill)
+			.VAlign(VAlign_Fill)
+			.FillContentHeight(1.0f)
+			[
+				Splitter.ToSharedRef()
+			]
+
+			+ SVerticalBox::Slot()
+			.HAlign(HAlign_Fill)
+			.VAlign(VAlign_Top)
+			.FillContentHeight(1.0f, 0.0f)
+			[
+				SNew(SColorBlock)
+				.Color(FLinearColor::Green)
+				.Size(FVector2D(0, 0))
+			]
+		];
+	}
+
+	void SContentBrowserSourceTree::SaveSettings(const FName& InInstanceName) const
+	{
+		static const FString SettingsIniSection = SContentBrowser::SettingsIniSection;
+		const FString& SettingsString = InInstanceName.ToString();
+
+		for (int32 SlotIndex = 0; SlotIndex < Splitter->GetChildren()->Num(); SlotIndex++)
+		{
+			float SplitterSize = Splitter->SlotAt(SlotIndex).GetSizeValue();
+			GConfig->SetFloat(*SettingsIniSection, *(SettingsString + FString::Printf(TEXT(".FavoriteSplitter.SlotSize%d"), SlotIndex)), SplitterSize, GEditorPerProjectIni);
+		}
+	}
+
+	void SContentBrowserSourceTree::LoadSettings(const FName& InInstanceName)
+	{
+		static const FString SettingsIniSection = SContentBrowser::SettingsIniSection;
+		const FString& SettingsString = InInstanceName.ToString();
+
+		for (int32 SlotIndex = 0; SlotIndex < Splitter->GetChildren()->Num(); SlotIndex++)
+		{
+			float SplitterSize = Splitter->SlotAt(SlotIndex).GetSizeValue();
+			GConfig->GetFloat(*SettingsIniSection, *(SettingsString + FString::Printf(TEXT(".FavoriteSplitter.SlotSize%d"), SlotIndex)), SplitterSize, GEditorPerProjectIni);
+			Splitter->SlotAt(SlotIndex).SetSizeValue(SplitterSize);
+		}
+	}
+
+	void SContentBrowserSourceTreeArea::Construct(
+		const FArguments& InArgs,
+		const FName InId,
+		const TSharedPtr<FSourcesSearch>& InSearch,
+		TSharedRef<IScrollableWidget> InBody)
+	{
+		check(!InId.IsNone());
+
+		Id = InId;
+		bIsEmpty = InArgs._IsEmpty;
+		OnExpansionChanged = InArgs._OnExpansionChanged;
+		BodyScrollableWidget = InBody;
+
+		static const float HorizontalPadding =
+			UE::Editor::ContentBrowser::IsNewStyleEnabled()
+			? 10.0f
+			: 4.0f;
+
+		static const float VerticalHeaderPadding =
+			UE::Editor::ContentBrowser::IsNewStyleEnabled()
+			? HorizontalPadding - 4.0f
+			: 0.0f;
+
+		static const float SearchButtonRightPadding =
+			UE::Editor::ContentBrowser::IsNewStyleEnabled()
+			? 0.0f
+			: 4.0f;
+
+		static constexpr float EmptyBodyLabelPadding = 8.0f;
+
+		static const FMargin HeaderPadding = FMargin(HorizontalPadding, VerticalHeaderPadding);
+
+		TSharedPtr<SWidget> SearchButtonWidget = SNullWidget::NullWidget;
+		if (InSearch.IsValid())
+		{
+			Search = InSearch;
+			SearchButtonWidget =
+				SAssignNew(SearchToggleButton, SSearchToggleButton, Search->GetWidget())
+				.Visibility(this, &SContentBrowserSourceTreeArea::GetHeaderSearchActionVisibility)
+				.OnSearchBoxShown_Lambda([this]() { SetExpanded(true); });
+		}
+
+		TAttribute<int32> BodyContentIndexAttribute;
+		if (UE::Editor::ContentBrowser::IsNewStyleEnabled()
+			&& (InArgs._IsEmpty.IsSet() || InArgs._IsEmpty.IsBound()))
+		{
+			BodyContentIndexAttribute = TAttribute<int32>::Create(
+			TAttribute<int32>::FGetter::CreateSPLambda(
+				this,
+				[IsEmptyAttribute = InArgs._IsEmpty]()
+			{
+				return IsEmptyAttribute.Get() ? 1 : 0;
+			}));
+		}
+		else
+		{
+			BodyContentIndexAttribute.Set(0);
+		}
+
+		ChildSlot
+		[
+			SAssignNew(ExpandableArea, SExpandableArea)
+			.Style(&FContentBrowserStyle::Get().GetWidgetStyle<FExpandableAreaStyle>("ContentBrowser.AssetTreeExpandableArea"))
+			.BorderImage(FContentBrowserStyle::Get().GetBrush("ContentBrowser.AssetTreeHeaderBrush"))
+			.BodyBorderImage(FContentBrowserStyle::Get().GetBrush("ContentBrowser.AssetTreeBodyBrush"))
+			.HeaderPadding(HeaderPadding)
+			.Visibility(InArgs._Visibility)
+			.Padding(0.0f)
+			.AllowAnimatedTransition(true)
+			.OnAreaExpansionChanged(this, &SContentBrowserSourceTreeArea::OnAreaExpansionChanged)
+			.HeaderContent()
+			[
+				SNew(SHorizontalBox)
+
+				+ SHorizontalBox::Slot()
+				.VAlign(VAlign_Center)
+				[
+					SNew(STextBlock)
+					.Visibility(EVisibility::HitTestInvisible) // Allows click-through to the expander button
+					.Text(InArgs._Label)
+					.TextStyle(FAppStyle::Get(), "ButtonText")
+					.Font(FAppStyle::Get().GetFontStyle("NormalFontBold"))
+				]
+
+				+ SHorizontalBox::Slot()
+				.VAlign(VAlign_Fill)
+				.HAlign(HAlign_Right)
+				.AutoWidth()
+				.Padding(0.0f)
+				[
+					InArgs._HeaderContent.Widget
+				]
+
+				+ SHorizontalBox::Slot()
+				.VAlign(VAlign_Center)
+				.HAlign(HAlign_Right)
+				.AutoWidth()
+				.Padding(FMargin(HorizontalPadding, 0.0f, SearchButtonRightPadding, 0.0f))
+				[
+					SearchButtonWidget.ToSharedRef()
+				]
+			]
+			.BodyContent()
+			[
+				SNew(SWidgetSwitcher)
+				.WidgetIndex(BodyContentIndexAttribute)
+
+				+ SWidgetSwitcher::Slot()
+				.Padding(0)
+				[
+					SNew(SVerticalBox)
+
+					// Search bar (if applicable)
+					+ SVerticalBox::Slot()
+					.AutoHeight()
+					[
+						// Should blend in visually with the header but technically acts like part of the body
+						SNew(SBorder)
+						.BorderImage(FAppStyle::Get().GetBrush("Brushes.Header"))
+						.Padding(FMargin(HorizontalPadding, 2.0f))
+						[
+							Search->GetWidget()
+						]
+					]
+
+					+ SVerticalBox::Slot()
+					.Padding(FMargin(0, 1))
+					[
+						// Surround scrollable with a scrollbox (adds drop shadows)
+						SNew(SScrollBorder, InBody)
+						[
+							BodyScrollableWidget->GetScrollWidget()
+						]
+					]
+				]
+
+				+ SWidgetSwitcher::Slot()
+				.Padding(0)
+				[
+					SNew(SBorder)
+					.BorderImage(FAppStyle::GetBrush("Brushes.Recessed"))
+					.Padding(0)
+					[
+						SNew(SBox)
+						.HAlign(HAlign_Center)
+						.VAlign(VAlign_Top)
+						.HeightOverride(MinBodyHeight)
+						.Padding(EmptyBodyLabelPadding)
+						[
+							SNew(SRichTextBlock)
+							.Text(InArgs._EmptyBodyLabel)
+							.TextStyle(&FAppStyle::Get().GetWidgetStyle<FTextBlockStyle>("RichTextBlock.Italic"))
+							.AutoWrapText(true)
+							.DecoratorStyleSet(&FAppStyle::Get())
+							+ SRichTextBlock::ImageDecorator()
+						]
+					]
+				]
+			]
+		];
+	}
+
+	bool SContentBrowserSourceTreeArea::IsExpanded() const
+	{
+		return ExpandableArea->IsExpanded();
+	}
+
+	void SContentBrowserSourceTreeArea::SetExpanded(bool bInExpanded)
+	{
+		ExpandableArea->SetExpanded(bInExpanded);
+	}
+
+	bool SContentBrowserSourceTreeArea::IsEmpty() const
+	{
+		return bIsEmpty.Get();
+	}
+
+	void SContentBrowserSourceTreeArea::SaveSettings(const FName& InInstanceName) const
+	{
+		static const FString SettingsIniSection = SContentBrowser::SettingsIniSection;
+		const FString& SettingsString = InInstanceName.ToString();
+		const FString& IdString = Id.ToString();
+
+		GConfig->SetBool(*SettingsIniSection, *(SettingsString + TEXT(".") + IdString + TEXT("AreaExpanded")), IsExpanded(), GEditorPerProjectIni);
+
+		if (Search.IsValid())
+		{
+			GConfig->SetBool(*SettingsIniSection, *(SettingsString + TEXT(".") + IdString + TEXT("SearchAreaExpanded")), SearchToggleButton->IsExpanded(), GEditorPerProjectIni);
+		}
+	}
+
+	void SContentBrowserSourceTreeArea::LoadSettings(const FName& InInstanceName)
+	{
+		static const FString SettingsIniSection = SContentBrowser::SettingsIniSection;
+		const FString& SettingsString = InInstanceName.ToString();
+		const FString& IdString = Id.ToString();
+
+		bool bAreaExpanded = false;
+		GConfig->GetBool(*SettingsIniSection, *(SettingsString + TEXT(".") + IdString + TEXT("AreaExpanded")), bAreaExpanded, GEditorPerProjectIni);
+		SetExpanded(bAreaExpanded);
+
+		if (Search.IsValid())
+		{
+			bool bSearchAreaExpanded = false;
+			GConfig->GetBool(*SettingsIniSection, *(SettingsString + TEXT(".") + IdString + TEXT("SearchAreaExpanded")), bSearchAreaExpanded, GEditorPerProjectIni);
+			SearchToggleButton->SetExpanded(bSearchAreaExpanded);
+		}
+	}
+
+	TSharedPtr<SSearchToggleButton> SContentBrowserSourceTreeArea::GetSearchToggleButton() const
+	{
+		return SearchToggleButton;
+	}
+
+	EVisibility SContentBrowserSourceTreeArea::GetHeaderSearchActionVisibility() const
+	{
+		return EVisibility::Visible;
+	}
+
+	void SContentBrowserSourceTreeArea::OnAreaExpansionChanged(bool bInIsExpanded)
+	{
+		if (SearchToggleButton.IsValid() && !bInIsExpanded)
+		{
+			SearchToggleButton->SetExpanded(false);
+		}
+
+		OnExpansionChanged.ExecuteIfBound(bInIsExpanded);
+	}
+
+	END_SLATE_FUNCTION_BUILD_OPTIMIZATION
+}
+
+#undef LOCTEXT_NAMESPACE
diff --git a/Engine/Source/Editor/ContentBrowser/Private/Widgets/SContentBrowserSourceTree.h b/Engine/Source/Editor/ContentBrowser/Private/Widgets/SContentBrowserSourceTree.h
new file mode 100644
--- /dev/null
+++ b/Engine/Source/Editor/ContentBrowser/Private/Widgets/SContentBrowserSourceTree.h
@@ -0,0 +1,184 @@
+// Copyright Epic Games, Inc. All Rights Reserved.
+
+#pragma once
+
+#include "Containers/Array.h"
+#include "Containers/ArrayView.h"
+#include "Containers/UnrealString.h"
+#include "Framework/Layout/IScrollableWidget.h"
+#include "Framework/SlateDelegates.h"
+#include "HAL/Platform.h"
+#include "Input/Reply.h"
+#include "Internationalization/Text.h"
+#include "Layout/Visibility.h"
+#include "Misc/Optional.h"
+#include "Templates/SharedPointer.h"
+#include "Templates/UnrealTemplate.h"
+#include "Types/SlateEnums.h"
+#include "UObject/NameTypes.h"
+#include "Widgets/DeclarativeSyntaxSupport.h"
+#include "Widgets/Docking/SDockTab.h"
+#include "Widgets/Layout/SSplitter.h"
+#include "Widgets/SCompoundWidget.h"
+
+class FSourcesSearch;
+class SSearchToggleButton;
+
+namespace UE::Editor::ContentBrowser::Private
+{
+	class SContentBrowserSourceTree;
+	class SContentBrowserSourceTreeArea;
+
+	/** The Content Browser Source Tree, containing "Favorites", etc. */
+	class SContentBrowserSourceTree : public SCompoundWidget
+	{
+		using Super = SCompoundWidget;
+
+	public:
+		class FSlot : public TSlotBase<FSlot>
+		{
+		public:
+			SLATE_SLOT_BEGIN_ARGS(FSlot, TSlotBase<FSlot>)
+				SLATE_ARGUMENT(TSharedPtr<SContentBrowserSourceTreeArea>, AreaWidget)
+				SLATE_ARGUMENT(TOptional<float>, Size)
+				SLATE_ARGUMENT(TOptional<SSplitter::ESizeRule>, ExpandedSizeRule)
+				SLATE_ATTRIBUTE(EVisibility, Visibility)
+				SLATE_ARGUMENT(TOptional<float>, HeaderHeight)
+			SLATE_SLOT_END_ARGS()
+
+			void Construct(const FChildren& SlotOwner, FSlotArguments&& InArgs);
+
+		public:
+			SSplitter::ESizeRule GetExpandedSizeRule() const;
+			float GetMinHeight() const;
+			float GetInitialSlotSize() const;
+			float GetSlotSize() const;
+			TSharedPtr<SContentBrowserSourceTreeArea> GetEntryWidget() const;
+
+			bool IsVisible() const;
+
+			void OnSlotResized(float InNewSize);
+
+		private:
+			SSplitter::ESizeRule ExpandedSizeRule = SSplitter::ESizeRule::FractionOfParent;
+			float HeaderHeight = 36.0f;
+			float InitialSlotSize = 0.0f;
+			float SlotSize = 0.0f; // The current slot height
+			TSharedPtr<SContentBrowserSourceTreeArea> EntryWidget = nullptr;
+			TAttribute<EVisibility> EntryVisibility;
+		};
+
+	public:
+		/** @return Add a new FSlot(). */
+		static FSlot::FSlotArguments Slot();
+
+		SLATE_BEGIN_ARGS(SContentBrowserSourceTree)
+		{ }
+			SLATE_SLOT_ARGUMENT(FSlot, Slots)
+		SLATE_END_ARGS()
+
+		SContentBrowserSourceTree()
+			: Slots(this)
+		{
+		}
+
+		/**
+		 * Constructs this widget with InArgs.
+		 */
+		void Construct(const FArguments& InArgs);
+
+		TSharedPtr<SSplitter> GetSplitter() const { return Splitter; }
+
+		/** Saves all persistent settings to config. */
+		void SaveSettings(const FName& InInstanceName) const;
+
+		/** Loads settings from config based on the browser's InstanceName. */
+		void LoadSettings(const FName& InInstanceName);
+
+	private:
+		using FScopedWidgetSlotArguments = TPanelChildren<FSlot>::FScopedWidgetSlotArguments;
+
+		FScopedWidgetSlotArguments AddSlot();
+
+		float GetMinHeight(const int32 InSlotIdx) const;
+		SSplitter::ESizeRule GetExpandedSizeRule(const int32 InSlotIdx) const;
+		float GetSlotSize(const int32 InSlotIdx) const;
+		void OnSlotResized(float InNewSize, const int32 InSlotIdx);
+
+	private:
+		TSharedPtr<SSplitter> Splitter;
+		TPanelChildren<FSlot> Slots;
+	};
+
+	/** Represents a single item in the source tree view, ie. "Favorites". */
+	class SContentBrowserSourceTreeArea	: public SCompoundWidget
+	{
+		using Super = SCompoundWidget;
+
+	public:
+		SLATE_BEGIN_ARGS(SContentBrowserSourceTreeArea)
+			: _IsEmpty(false)
+		{ }
+			/** Label to display in the header, ie. "Favorites" */
+			SLATE_ATTRIBUTE(FText, Label)
+
+			/** (Optional) content displayed to the right of the header label. */
+			SLATE_NAMED_SLOT(FArguments, HeaderContent)
+
+			/** (Optional) bind to notify when this widget should switch to an empty state. */
+			SLATE_ATTRIBUTE(bool, IsEmpty)
+
+			/** (Optional) label to display in the body area when the body is empty, shown as RichText. */
+			SLATE_ATTRIBUTE(FText, EmptyBodyLabel)
+
+			/** Called when the area is expanded or collapsed */
+			SLATE_EVENT(FOnBooleanValueChanged, OnExpansionChanged)
+		SLATE_END_ARGS()
+
+		/**
+		 * Constructs this widget with InArgs.
+		 * @param Id		Used to identify the role of this widget (Favorites, etc.) for config settings, etc.
+		 * @param InSearch  (Optional) Search for this area and it's contained items.
+		 * @param InBody	The Body widget, which must implement IScrollableWidget
+		 */
+		void Construct(const FArguments& InArgs, const FName InId, const TSharedPtr<FSourcesSearch>& InSearch, TSharedRef<IScrollableWidget> InBody);
+
+		bool IsExpanded() const;
+		void SetExpanded(bool bInExpanded);
+
+		/** Returns true if the contents is "Empty", regardless of expansion state. */
+		bool IsEmpty() const;
+
+		/** Saves all persistent settings to config. */
+		void SaveSettings(const FName& InInstanceName) const;
+
+		/** Loads settings from config based on the browser's InstanceName. */
+		void LoadSettings(const FName& InInstanceName);
+
+		bool HasSearch() const { return Search.IsValid(); }
+
+		TSharedPtr<SSearchToggleButton> GetSearchToggleButton() const;
+
+	private:
+		EVisibility GetHeaderSearchActionVisibility() const;
+
+		void OnAreaExpansionChanged(bool bInIsExpanded);
+
+	private:
+		FName Id;
+
+		TAttribute<bool> bIsEmpty;
+
+		/** The (optional) source search. */
+		TSharedPtr<FSourcesSearch> Search = nullptr;
+
+		/** Toggle button for showing/hiding the search area, only visible if Search IsValid. */
+		TSharedPtr<SSearchToggleButton> SearchToggleButton = nullptr;
+
+		TSharedPtr<IScrollableWidget> BodyScrollableWidget;
+
+		TSharedPtr<SExpandableArea> ExpandableArea;
+
+		FOnBooleanValueChanged OnExpansionChanged;
+	};
+}
