[ContentBrowser] CIS fix

diff --git a/Engine/Source/Editor/ContentBrowser/Private/Widgets/SContentBrowserSourceTree.cpp b/Engine/Source/Editor/ContentBrowser/Private/Widgets/SContentBrowserSourceTree.cpp
--- a/Engine/Source/Editor/ContentBrowser/Private/Widgets/SContentBrowserSourceTree.cpp
+++ b/Engine/Source/Editor/ContentBrowser/Private/Widgets/SContentBrowserSourceTree.cpp
@@ -111,10 +111,10 @@ namespace UE::Editor::ContentBrowser::Private
 	SContentBrowserSourceTree::FScopedWidgetSlotArguments SContentBrowserSourceTree::AddSlot()
 	{
 		return FScopedWidgetSlotArguments{ MakeUnique<FSlot>(), Slots, INDEX_NONE,
-			[this, Splitter = Splitter](const FSlot* InSlotAdded, int32 InSlotIdx)
+			[this, HostSplitter = Splitter](const FSlot* InSlotAdded, int32 InSlotIdx)
 			{
 				// When the slot is added, also add it to the splitter
-				Splitter->AddSlot()
+				HostSplitter->AddSlot()
 				.Value(InSlotAdded->GetInitialSlotSize())
 				.SizeRule(this, &SContentBrowserSourceTree::GetExpandedSizeRule, InSlotIdx)
 				.MinSize(this, &SContentBrowserSourceTree::GetMinHeight, InSlotIdx)
