[ContentBrowser] (New Style) Fixed entry duplication

diff --git a/Engine/Source/Editor/ContentBrowser/Private/Menus/AssetViewMenus.cpp b/Engine/Source/Editor/ContentBrowser/Private/Menus/AssetViewMenus.cpp
--- a/Engine/Source/Editor/ContentBrowser/Private/Menus/AssetViewMenus.cpp
+++ b/Engine/Source/Editor/ContentBrowser/Private/Menus/AssetViewMenus.cpp
@@ -64,10 +64,14 @@ using namespace UE::Editor::ContentBrowser;
 
 FDelayedAutoRegisterHelper SAssetView::AssetViewOptionsMenuRegistration(
 	EDelayedRegisterRunPhase::EndOfEngineInit,
-	[]
+	[](const EDelayedRegisterRunPhase InPhase)
 	{
 		static const FName AssetViewOptionsMenuName("ContentBrowser.AssetViewOptions");
-		if (!ensureAlways(!UToolMenus::Get()->IsMenuRegistered(AssetViewOptionsMenuName)))
+		if (InPhase == EDelayedRegisterRunPhase::LiveCodingReload)
+		{
+			UToolMenus::Get()->RemoveMenu(AssetViewOptionsMenuName);
+		}
+		else if (!ensureAlways(!UToolMenus::Get()->IsMenuRegistered(AssetViewOptionsMenuName)))
 		{
 			return;
 		}
@@ -177,7 +181,7 @@ FDelayedAutoRegisterHelper SAssetView::AssetViewOptionsMenuRegistration(
 						// Filter Bar Section
 						if (TSharedPtr<SFilterList> FilterBarPinned = AssetView->FilterBar.Pin())
 						{
-							FToolMenuSection& Section = InMenu->AddSection("FilterBar", LOCTEXT("FilterBarHeading", "Filter Display"));
+							FToolMenuSection& Section = InMenu->FindOrAddSection("FilterBar", LOCTEXT("FilterBarHeading", "Filter Display"));
 
 							Section.AddMenuEntry(
 								"VerticalLayout",
@@ -230,7 +234,7 @@ FDelayedAutoRegisterHelper SAssetView::AssetViewOptionsMenuRegistration(
 
 						auto CreateThumbnailSizeSubMenu = [AssetView](UToolMenu* SubMenu)
 						{
-							FToolMenuSection& SizeSection = SubMenu->AddSection("ThumbnailSizes");
+							FToolMenuSection& SizeSection = SubMenu->FindOrAddSection("ThumbnailSizes");
 
 							for (int32 EnumValue = static_cast<int32>(EThumbnailSize::Tiny); EnumValue < static_cast<int32>(EThumbnailSize::MAX); ++EnumValue)
 							{
@@ -469,7 +473,7 @@ FDelayedAutoRegisterHelper SAssetView::AssetViewOptionsMenuRegistration(
 						if (AssetView->GetColumnViewVisibility() == EVisibility::Visible)
 						{
 							{
-								FToolMenuSection& Section = InMenu->AddSection("AssetColumns", LOCTEXT("ToggleColumnsHeading", "Columns"));
+								FToolMenuSection& Section = InMenu->FindOrAddSection("AssetColumns", LOCTEXT("ToggleColumnsHeading", "Columns"));
 
 								Section.AddMenuEntry(
 									"ResetColumns",
@@ -503,11 +507,12 @@ FDelayedAutoRegisterHelper SAssetView::AssetViewOptionsMenuRegistration(
 
 			ContentBrowserMenuUtils::AddFiltersToMenu(InMenu, OwningContentBrowserName, FilterParams);
 		}));
-	});
+	},
+	true);
 
 FDelayedAutoRegisterHelper SAssetView::ToolBarMenuExtensionRegistration(
 	EDelayedRegisterRunPhase::EndOfEngineInit,
-	[]
+	[](const EDelayedRegisterRunPhase InPhase)
 	{
 		// @todo: Remove when old style is removed, all extensions are for the new style so early-out if it's not enabled
 		if (!IsNewStyleEnabled())
@@ -576,17 +581,17 @@ FDelayedAutoRegisterHelper SAssetView::ToolBarMenuExtensionRegistration(
 					}
 				}));
 		}
-	});
+	},
+	true);
 
 FDelayedAutoRegisterHelper SAssetView::NavigationBarMenuExtensionRegistration(
 	EDelayedRegisterRunPhase::EndOfEngineInit,
-	[]
+	[](const EDelayedRegisterRunPhase InPhase)
 	{
 		static const FName NavigationBarName("ContentBrowser.NavigationBar");
 		if (UToolMenu* NavigationBarMenu = UToolMenus::Get()->ExtendMenu(NavigationBarName))
 		{
 			FToolMenuSection& AssetCountSection = NavigationBarMenu->FindOrAddSection("AssetCount");
-			AssetCountSection.InsertPosition = FToolMenuInsert(NAME_None, EToolMenuInsertType::Last);
 			AssetCountSection.Alignment = EToolMenuSectionAlign::Last;
 
 			static bool bIsNewStyleEnabled = UE::Editor::ContentBrowser::IsNewStyleEnabled();
@@ -594,7 +599,7 @@ FDelayedAutoRegisterHelper SAssetView::NavigationBarMenuExtensionRegistration(
 			if (bIsNewStyleEnabled)
 			{
 				AssetCountSection.AddDynamicEntry(
-					"AssetCount",
+					NAME_None,
 					FNewToolMenuSectionDelegate::CreateLambda([](FToolMenuSection& InSection)
 					{
 						if (const UContentBrowserToolbarMenuContext* Context = InSection.FindContext<UContentBrowserToolbarMenuContext>())
@@ -620,6 +625,7 @@ FDelayedAutoRegisterHelper SAssetView::NavigationBarMenuExtensionRegistration(
 					}));
 			}
 		}
-	});
+	},
+	true);
 
 #undef LOCTEXT_NAMESPACE
diff --git a/Engine/Source/Editor/ContentBrowser/Private/Menus/ContentBrowserMenus.cpp b/Engine/Source/Editor/ContentBrowser/Private/Menus/ContentBrowserMenus.cpp
--- a/Engine/Source/Editor/ContentBrowser/Private/Menus/ContentBrowserMenus.cpp
+++ b/Engine/Source/Editor/ContentBrowser/Private/Menus/ContentBrowserMenus.cpp
@@ -16,7 +16,7 @@
 
 FDelayedAutoRegisterHelper SContentBrowser::NavigationBarMenuRegistration(
 	EDelayedRegisterRunPhase::EndOfEngineInit,
-	[]
+	[](const EDelayedRegisterRunPhase InPhase)
 	{
 		static const FName NavigationBarMenuName("ContentBrowser.NavigationBar");
 		if (!ensureAlways(!UToolMenus::Get()->IsMenuRegistered(NavigationBarMenuName)))
@@ -30,7 +30,7 @@ FDelayedAutoRegisterHelper SContentBrowser::NavigationBarMenuRegistration(
 		ToolBar->SetStyleSet(&UE::ContentBrowser::Private::FContentBrowserStyle::Get());
 		ToolBar->StyleName = "ContentBrowser.ToolBar";
 
-		ToolBar->AddSection("History")
+		ToolBar->FindOrAddSection("History")
 		.AddDynamicEntry(NAME_None,
 			FNewToolMenuSectionDelegate::CreateLambda([](FToolMenuSection& InSection)
 			{
@@ -80,7 +80,7 @@ FDelayedAutoRegisterHelper SContentBrowser::NavigationBarMenuRegistration(
 				}
 			}));
 
-		ToolBar->AddSection("Path")
+		ToolBar->FindOrAddSection("Path")
 		.AddDynamicEntry(NAME_None, FNewToolMenuSectionDelegate::CreateLambda(
 			[](FToolMenuSection& InSection)
 			{
@@ -118,11 +118,12 @@ FDelayedAutoRegisterHelper SContentBrowser::NavigationBarMenuRegistration(
 					}
 				}
 			}));
-	});
+	},
+	true);
 
 FDelayedAutoRegisterHelper SContentBrowser::AddNewContextMenuRegistration(
 	EDelayedRegisterRunPhase::EndOfEngineInit,
-	[]
+	[](const EDelayedRegisterRunPhase InPhase)
 	{
 		static const FName AddNewContextMenuName("ContentBrowser.AddNewContextMenu");
 		if (!ensureAlways(!UToolMenus::Get()->IsMenuRegistered(AddNewContextMenuName)))
@@ -153,14 +154,19 @@ FDelayedAutoRegisterHelper SContentBrowser::AddNewContextMenuRegistration(
 				ContentBrowser->PopulateAddNewContextMenu(InMenu);
 			}
 		}));
-	});
+	},
+	true);
 
 FDelayedAutoRegisterHelper SContentBrowser::FolderContextMenuRegistration(
 	EDelayedRegisterRunPhase::EndOfEngineInit,
-	[]
+	[](const EDelayedRegisterRunPhase InPhase)
 	{
 		static const FName FolderContextMenuName("ContentBrowser.FolderContextMenu");
-		if (!ensureAlways(!UToolMenus::Get()->IsMenuRegistered(FolderContextMenuName)))
+		if (InPhase == EDelayedRegisterRunPhase::LiveCodingReload)
+		{
+			UToolMenus::Get()->RemoveMenu(FolderContextMenuName);
+		}
+		else if (!ensureAlways(!UToolMenus::Get()->IsMenuRegistered(FolderContextMenuName)))
 		{
 			return;
 		}
@@ -178,14 +184,19 @@ FDelayedAutoRegisterHelper SContentBrowser::FolderContextMenuRegistration(
 				Context->ContentBrowser.Pin()->PopulateFolderContextMenu(InMenu);
 			}
 		}));
-	});
+	},
+	true);
 
 FDelayedAutoRegisterHelper SContentBrowser::PathViewFiltersMenuRegistration(
 	EDelayedRegisterRunPhase::EndOfEngineInit,
-	[]
+	[](const EDelayedRegisterRunPhase InPhase)
 	{
 		static const FName PathViewFiltersMenuName("ContentBrowser.AssetViewOptions.PathViewFilters");
-		if (!ensureAlways(!UToolMenus::Get()->IsMenuRegistered(PathViewFiltersMenuName)))
+		if (InPhase == EDelayedRegisterRunPhase::LiveCodingReload)
+		{
+			UToolMenus::Get()->RemoveMenu(PathViewFiltersMenuName);
+		}
+		else if (!ensureAlways(!UToolMenus::Get()->IsMenuRegistered(PathViewFiltersMenuName)))
 		{
 			return;
 		}
@@ -211,14 +222,19 @@ FDelayedAutoRegisterHelper SContentBrowser::PathViewFiltersMenuRegistration(
 				}
 			}
 		}));
-	});
+	},
+	true);
 
 FDelayedAutoRegisterHelper SContentBrowser::ToolBarMenuRegistration(
 	EDelayedRegisterRunPhase::EndOfEngineInit,
-	[]
+	[](const EDelayedRegisterRunPhase InPhase)
 	{
 		static const FName ToolBarMenuName("ContentBrowser.ToolBar");
-		if (!ensureAlways(!UToolMenus::Get()->IsMenuRegistered(ToolBarMenuName)))
+		if (InPhase == EDelayedRegisterRunPhase::LiveCodingReload)
+		{
+			UToolMenus::Get()->RemoveMenu(ToolBarMenuName);
+		}
+		else if (!ensureAlways(!UToolMenus::Get()->IsMenuRegistered(ToolBarMenuName)))
 		{
 			return;
 		}
@@ -229,7 +245,7 @@ FDelayedAutoRegisterHelper SContentBrowser::ToolBarMenuRegistration(
 		ToolBar->SetStyleSet(&UE::ContentBrowser::Private::FContentBrowserStyle::Get());
 		ToolBar->StyleName = ToolBarMenuName; // Style name is the same as the menu entry name
 
-		ToolBar->AddSection("New")
+		ToolBar->FindOrAddSection("New")
 		.AddDynamicEntry(NAME_None,
 			FNewToolMenuSectionDelegate::CreateLambda([](FToolMenuSection& InSection)
 			{
@@ -263,7 +279,7 @@ FDelayedAutoRegisterHelper SContentBrowser::ToolBarMenuRegistration(
 				}
 			}));
 
-		ToolBar->AddSection("Save")
+		ToolBar->FindOrAddSection("Save")
 		.AddDynamicEntry(NAME_None,
 			FNewToolMenuSectionDelegate::CreateLambda([](FToolMenuSection& InSection)
 			{
@@ -333,7 +349,7 @@ FDelayedAutoRegisterHelper SContentBrowser::ToolBarMenuRegistration(
 
 		if (UE::Editor::ContentBrowser::IsNewStyleEnabled())
 		{
-			ToolBar->AddSection("Filter")
+			ToolBar->FindOrAddSection("Filter")
 			.AddDynamicEntry(NAME_None, FNewToolMenuSectionDelegate::CreateLambda(
 				[](FToolMenuSection& InSection)
 				{
@@ -369,7 +385,7 @@ FDelayedAutoRegisterHelper SContentBrowser::ToolBarMenuRegistration(
 					}
 				}));
 
-			ToolBar->AddSection("Search")
+			ToolBar->FindOrAddSection("Search")
 			.AddDynamicEntry(NAME_None, FNewToolMenuSectionDelegate::CreateLambda(
 				[](FToolMenuSection& InSection)
 				{
@@ -403,7 +419,7 @@ FDelayedAutoRegisterHelper SContentBrowser::ToolBarMenuRegistration(
 					}
 				}));
 
-			FToolMenuSection& SettingsSection = ToolBar->AddSection("Settings");
+			FToolMenuSection& SettingsSection = ToolBar->FindOrAddSection("Settings");
 			SettingsSection.Alignment = EToolMenuSectionAlign::Last;
 
 			SettingsSection.AddDynamicEntry(NAME_None,
@@ -434,6 +450,7 @@ FDelayedAutoRegisterHelper SContentBrowser::ToolBarMenuRegistration(
 				}
 			}));
 		}
-	});
+	},
+	true);
 
 #undef LOCTEXT_NAMESPACE
