[ContentBrowser] ToolBar entry order fix

diff --git a/Engine/Plugins/Editor/ContentBrowser/ContentBrowserAssetDataSource/Source/ContentBrowserAssetDataSource/Private/ContentBrowserAssetDataSource.cpp b/Engine/Plugins/Editor/ContentBrowser/ContentBrowserAssetDataSource/Source/ContentBrowserAssetDataSource/Private/ContentBrowserAssetDataSource.cpp
--- a/Engine/Plugins/Editor/ContentBrowser/ContentBrowserAssetDataSource/Source/ContentBrowserAssetDataSource/Private/ContentBrowserAssetDataSource.cpp
+++ b/Engine/Plugins/Editor/ContentBrowser/ContentBrowserAssetDataSource/Source/ContentBrowserAssetDataSource/Private/ContentBrowserAssetDataSource.cpp
@@ -3653,7 +3653,7 @@ void UContentBrowserAssetDataSource::PopulateContentBrowserToolBar(UToolMenu* In
 
 	Section.AddSeparator(NAME_None);
 
-	Section.AddEntry(
+	FToolMenuEntry& ImportEntry = Section.AddEntry(
 		FToolMenuEntry::InitWidget(
 			"Import",
 			ImportButton.ToSharedRef(),
@@ -3661,7 +3661,8 @@ void UContentBrowserAssetDataSource::PopulateContentBrowserToolBar(UToolMenu* In
 			true,
 			false
 		));
-	
+
+	ImportEntry.InsertPosition.Position = EToolMenuInsertType::Last;
 }
 
 void UContentBrowserAssetDataSource::PopulateAssetFolderContextMenu(UToolMenu* InMenu)
diff --git a/Engine/Source/Editor/ContentBrowser/Private/Menus/AssetViewMenus.cpp b/Engine/Source/Editor/ContentBrowser/Private/Menus/AssetViewMenus.cpp
--- a/Engine/Source/Editor/ContentBrowser/Private/Menus/AssetViewMenus.cpp
+++ b/Engine/Source/Editor/ContentBrowser/Private/Menus/AssetViewMenus.cpp
@@ -181,7 +181,7 @@ FDelayedAutoRegisterHelper SAssetView::AssetViewOptionsMenuRegistration(
 						// Filter Bar Section
 						if (TSharedPtr<SFilterList> FilterBarPinned = AssetView->FilterBar.Pin())
 						{
-							FToolMenuSection& Section = InMenu->FindOrAddSection("FilterBar", LOCTEXT("FilterBarHeading", "Filter Display"));
+							FToolMenuSection& Section = InMenu->FindOrAddSection("FilterBar", LOCTEXT("FilterBarHeading", "Filter Display"), FToolMenuInsert("AssetViewType", EToolMenuInsertType::After));;
 
 							Section.AddMenuEntry(
 								"VerticalLayout",
@@ -299,7 +299,8 @@ FDelayedAutoRegisterHelper SAssetView::AssetViewOptionsMenuRegistration(
 					{
 						FToolMenuSection& Section = InMenu->FindOrAddSection(
 							IsNewStyleEnabled() ? "Show" : "View",
-							IsNewStyleEnabled() ? LOCTEXT("ShowHeading", "Show") : LOCTEXT("ViewHeading", "View"));
+							IsNewStyleEnabled() ? LOCTEXT("ShowHeading", "Show") : LOCTEXT("ViewHeading", "View"),
+							FToolMenuInsert("AssetThumbnails", EToolMenuInsertType::After));
 
 						const FText ShowFoldersLabel = IsNewStyleEnabled() ? LOCTEXT("ShowFoldersOption_NewStyle", "Folders") :LOCTEXT("ShowFoldersOption", "Show Folders");
 						SetEntryOrder(
@@ -390,7 +391,10 @@ FDelayedAutoRegisterHelper SAssetView::AssetViewOptionsMenuRegistration(
 
 					// Content Section
 					{
-						FToolMenuSection& Section = InMenu->FindOrAddSection("Content", IsNewStyleEnabled() ? FText::GetEmpty() : LOCTEXT("ContentHeading", "Content"));
+						FToolMenuSection& Section = InMenu->FindOrAddSection(
+							"Content",
+							IsNewStyleEnabled() ? FText::GetEmpty() : LOCTEXT("ContentHeading", "Content"),
+							FToolMenuInsert(IsNewStyleEnabled() ? "Show" : "View", EToolMenuInsertType::After));
 
 						if (IsNewStyleEnabled())
 						{
@@ -398,7 +402,7 @@ FDelayedAutoRegisterHelper SAssetView::AssetViewOptionsMenuRegistration(
 						}
 
 						const FText ShowAllFolderLabel = IsNewStyleEnabled() ?  LOCTEXT("ShowAllFolderOption_NewStyle", "All Folder") : LOCTEXT("ShowAllFolderOption", "Show All Folder");
-						FToolMenuSection& ShowAllFolderSection = IsNewStyleEnabled() ? Section : InMenu->FindOrAddSection("View");
+						FToolMenuSection& ShowAllFolderSection = IsNewStyleEnabled() ? Section : InMenu->FindOrAddSection("View", { }, FToolMenuInsert("Content", EToolMenuInsertType::After));
 						SetEntryOrder(
 							ShowAllFolderSection.AddMenuEntry(
 								"ShowAllFolder",
@@ -417,7 +421,7 @@ FDelayedAutoRegisterHelper SAssetView::AssetViewOptionsMenuRegistration(
 
 					// Search Section
 					{
-						FToolMenuSection& Section = InMenu->FindOrAddSection("Search", LOCTEXT("SearchHeading", "Search"));
+						FToolMenuSection& Section = InMenu->FindOrAddSection("Search", LOCTEXT("SearchHeading", "Search"), FToolMenuInsert("Content", EToolMenuInsertType::After));
 						Section.AddMenuEntry(
 							"IncludeClassName",
 							LOCTEXT("IncludeClassNameOption", "Search Asset Class Names"),
@@ -463,7 +467,7 @@ FDelayedAutoRegisterHelper SAssetView::AssetViewOptionsMenuRegistration(
 						// Manage Section
 						{
 							// @note: This section is used elsewhere - @see: SContentBrowser::ExtendViewOptionsMenu
-							FToolMenuSection& Section = InMenu->FindOrAddSection("Manage", LOCTEXT("ManageHeading", "Manage"));
+							FToolMenuSection& Section = InMenu->FindOrAddSection("Manage", LOCTEXT("ManageHeading", "Manage"), FToolMenuInsert("Search", EToolMenuInsertType::After));
 						}
 					}
 
@@ -473,7 +477,7 @@ FDelayedAutoRegisterHelper SAssetView::AssetViewOptionsMenuRegistration(
 						if (AssetView->GetColumnViewVisibility() == EVisibility::Visible)
 						{
 							{
-								FToolMenuSection& Section = InMenu->FindOrAddSection("AssetColumns", LOCTEXT("ToggleColumnsHeading", "Columns"));
+								FToolMenuSection& Section = InMenu->FindOrAddSection("AssetColumns", LOCTEXT("ToggleColumnsHeading", "Columns"), FToolMenuInsert("Search", EToolMenuInsertType::After));
 
 								Section.AddMenuEntry(
 									"ResetColumns",
@@ -591,7 +595,7 @@ FDelayedAutoRegisterHelper SAssetView::NavigationBarMenuExtensionRegistration(
 		static const FName NavigationBarName("ContentBrowser.NavigationBar");
 		if (UToolMenu* NavigationBarMenu = UToolMenus::Get()->ExtendMenu(NavigationBarName))
 		{
-			FToolMenuSection& AssetCountSection = NavigationBarMenu->FindOrAddSection("AssetCount");
+			FToolMenuSection& AssetCountSection = NavigationBarMenu->FindOrAddSection("AssetCount", { }, FToolMenuInsert("History", EToolMenuInsertType::Last));
 			AssetCountSection.Alignment = EToolMenuSectionAlign::Last;
 
 			static bool bIsNewStyleEnabled = UE::Editor::ContentBrowser::IsNewStyleEnabled();
diff --git a/Engine/Source/Editor/ContentBrowser/Private/Menus/ContentBrowserMenus.cpp b/Engine/Source/Editor/ContentBrowser/Private/Menus/ContentBrowserMenus.cpp
--- a/Engine/Source/Editor/ContentBrowser/Private/Menus/ContentBrowserMenus.cpp
+++ b/Engine/Source/Editor/ContentBrowser/Private/Menus/ContentBrowserMenus.cpp
@@ -80,7 +80,7 @@ FDelayedAutoRegisterHelper SContentBrowser::NavigationBarMenuRegistration(
 				}
 			}));
 
-		ToolBar->FindOrAddSection("Path")
+		ToolBar->FindOrAddSection("Path", { }, FToolMenuInsert("History", EToolMenuInsertType::After))
 		.AddDynamicEntry(NAME_None, FNewToolMenuSectionDelegate::CreateLambda(
 			[](FToolMenuSection& InSection)
 			{
@@ -279,7 +279,7 @@ FDelayedAutoRegisterHelper SContentBrowser::ToolBarMenuRegistration(
 				}
 			}));
 
-		ToolBar->FindOrAddSection("Save")
+		ToolBar->FindOrAddSection("Save", { }, FToolMenuInsert("New", EToolMenuInsertType::After))
 		.AddDynamicEntry(NAME_None,
 			FNewToolMenuSectionDelegate::CreateLambda([](FToolMenuSection& InSection)
 			{
@@ -349,7 +349,7 @@ FDelayedAutoRegisterHelper SContentBrowser::ToolBarMenuRegistration(
 
 		if (UE::Editor::ContentBrowser::IsNewStyleEnabled())
 		{
-			ToolBar->FindOrAddSection("Filter")
+			ToolBar->FindOrAddSection("Filter", { }, FToolMenuInsert("Save", EToolMenuInsertType::After))
 			.AddDynamicEntry(NAME_None, FNewToolMenuSectionDelegate::CreateLambda(
 				[](FToolMenuSection& InSection)
 				{
@@ -385,7 +385,7 @@ FDelayedAutoRegisterHelper SContentBrowser::ToolBarMenuRegistration(
 					}
 				}));
 
-			ToolBar->FindOrAddSection("Search")
+			ToolBar->FindOrAddSection("Search", { }, FToolMenuInsert("Filter", EToolMenuInsertType::After))
 			.AddDynamicEntry(NAME_None, FNewToolMenuSectionDelegate::CreateLambda(
 				[](FToolMenuSection& InSection)
 				{
@@ -419,7 +419,7 @@ FDelayedAutoRegisterHelper SContentBrowser::ToolBarMenuRegistration(
 					}
 				}));
 
-			FToolMenuSection& SettingsSection = ToolBar->FindOrAddSection("Settings");
+			FToolMenuSection& SettingsSection = ToolBar->FindOrAddSection("Settings", { }, FToolMenuInsert("Search", EToolMenuInsertType::After));
 			SettingsSection.Alignment = EToolMenuSectionAlign::Last;
 
 			SettingsSection.AddDynamicEntry(NAME_None,
