[ContentBrowser] Menu fixes

diff --git a/Engine/Source/Editor/ContentBrowser/Private/Menus/AssetViewMenus.cpp b/Engine/Source/Editor/ContentBrowser/Private/Menus/AssetViewMenus.cpp
--- a/Engine/Source/Editor/ContentBrowser/Private/Menus/AssetViewMenus.cpp
+++ b/Engine/Source/Editor/ContentBrowser/Private/Menus/AssetViewMenus.cpp
@@ -64,24 +64,15 @@ using namespace UE::Editor::ContentBrowser;
 
 FDelayedAutoRegisterHelper SAssetView::AssetViewOptionsMenuRegistration(
 	EDelayedRegisterRunPhase::EndOfEngineInit,
-	[](const EDelayedRegisterRunPhase InPhase)
+	[]
 	{
-		static const FName AssetViewOptionsMenuName("ContentBrowser.AssetViewOptions");
-		if (InPhase == EDelayedRegisterRunPhase::LiveCodingReload)
-		{
-			UToolMenus::Get()->RemoveMenu(AssetViewOptionsMenuName);
-		}
-		else if (!ensureAlways(!UToolMenus::Get()->IsMenuRegistered(AssetViewOptionsMenuName)))
-		{
-			return;
-		}
-
 		FToolMenuOwnerScoped OwnerScoped(UE_MODULE_NAME);
 
+		static const FName AssetViewOptionsMenuName("ContentBrowser.AssetViewOptions");
 		UToolMenu* Menu = UToolMenus::Get()->RegisterMenu(AssetViewOptionsMenuName);
 		Menu->bCloseSelfOnly = true;
 
-		Menu->AddDynamicSection(NAME_None, FNewToolMenuDelegate::CreateLambda([](UToolMenu* InMenu)
+		Menu->AddDynamicSection(NAME_Default, FNewToolMenuDelegate::CreateLambda([](UToolMenu* InMenu)
 		{
 			FName OwningContentBrowserName = NAME_None;
 			FFiltersAdditionalParams FilterParams;
@@ -92,11 +83,11 @@ FDelayedAutoRegisterHelper SAssetView::AssetViewOptionsMenuRegistration(
 					auto SetEntryOrder = [](FToolMenuEntry& InEntry)
 					{
 						// New Style, >5.6: (entry name, entry to place after) map. If the named entry isn't found, the entry will not have an explicit order.
-	                    static const TMap<FName, FName> NewEntryOrder = {
-                    		{ "ShowFolders", "ShowFavorite" },
-                    		{ "ShowEmptyFolders", "ShowFolders" },
-                    		{ "OrganizeFolders", "ShowEmptyFolders" }
-	                    };
+						static const TMap<FName, FName> NewEntryOrder = {
+							{ "ShowFolders", "ShowFavorite" },
+							{ "ShowEmptyFolders", "ShowFolders" },
+							{ "OrganizeFolders", "ShowEmptyFolders" }
+						};
 
 						// <5.6 Style: (entry name, entry to place after) map. If the named entry isn't found, the entry will not have an explicit order.
 						static const TMap<FName, FName> EntryOrder = {
@@ -113,8 +104,11 @@ FDelayedAutoRegisterHelper SAssetView::AssetViewOptionsMenuRegistration(
 
 					// ViewType Section
 					{
-						FToolMenuSection& Section = InMenu->FindOrAddSection("AssetViewType", LOCTEXT("ViewTypeHeading", "View Type"));
-						Section.AddMenuEntry(
+						FToolMenuSection& ViewTypeSection = InMenu->FindOrAddSection(
+							"AssetViewType",
+							LOCTEXT("ViewTypeHeading", "View Type"));
+
+						ViewTypeSection.AddMenuEntry(
 							"TileView",
 							LOCTEXT("TileViewOption", "Tiles"),
 							LOCTEXT("TileViewOptionToolTip", "View assets as tiles in a grid."),
@@ -127,7 +121,7 @@ FDelayedAutoRegisterHelper SAssetView::AssetViewOptionsMenuRegistration(
 							EUserInterfaceActionType::RadioButton
 						);
 
-						Section.AddMenuEntry(
+						ViewTypeSection.AddMenuEntry(
 							"ListView",
 							LOCTEXT("ListViewOption", "List"),
 							LOCTEXT("ListViewOptionToolTip", "View assets in a list with thumbnails."),
@@ -140,7 +134,7 @@ FDelayedAutoRegisterHelper SAssetView::AssetViewOptionsMenuRegistration(
 							EUserInterfaceActionType::RadioButton
 						);
 
-						Section.AddMenuEntry(
+						ViewTypeSection.AddMenuEntry(
 							"ColumnView",
 							LOCTEXT("ColumnViewOption", "Columns"),
 							LOCTEXT("ColumnViewOptionToolTip", "View assets in a list with columns of details."),
@@ -161,7 +155,7 @@ FDelayedAutoRegisterHelper SAssetView::AssetViewOptionsMenuRegistration(
 							? AssetView->ViewExtender->GetViewTooltipText()
 							: LOCTEXT("CustomViewOptionToolTip", "A user specified custom view.");
 
-						Section.AddMenuEntry(
+						ViewTypeSection.AddMenuEntry(
 							"CustomView",
 							CustomViewLabel,
 							CustomViewTooltip,
@@ -181,7 +175,10 @@ FDelayedAutoRegisterHelper SAssetView::AssetViewOptionsMenuRegistration(
 						// Filter Bar Section
 						if (TSharedPtr<SFilterList> FilterBarPinned = AssetView->FilterBar.Pin())
 						{
-							FToolMenuSection& Section = InMenu->FindOrAddSection("FilterBar", LOCTEXT("FilterBarHeading", "Filter Display"), FToolMenuInsert("AssetViewType", EToolMenuInsertType::After));;
+							FToolMenuSection& Section = InMenu->FindOrAddSection(
+								"FilterBar",
+								LOCTEXT("FilterBarHeading", "Filter Display"),
+								FToolMenuInsert("AssetViewType", EToolMenuInsertType::After));;
 
 							Section.AddMenuEntry(
 								"VerticalLayout",
@@ -229,8 +226,11 @@ FDelayedAutoRegisterHelper SAssetView::AssetViewOptionsMenuRegistration(
 
 					// Thumbnails Section
 					{
-						const FToolMenuInsert SectionInsert("Search", EToolMenuInsertType::After);
-						FToolMenuSection& Section = InMenu->FindOrAddSection("AssetThumbnails", LOCTEXT("ThumbnailsHeading", "Thumbnails"), SectionInsert);
+						const FToolMenuInsert ThumbnailSectionInsert("Search", EToolMenuInsertType::After);
+						FToolMenuSection& ThumbnailSection = InMenu->FindOrAddSection(
+							"AssetThumbnails",
+							LOCTEXT("ThumbnailsHeading", "Thumbnails"),
+							ThumbnailSectionInsert);
 
 						auto CreateThumbnailSizeSubMenu = [AssetView](UToolMenu* SubMenu)
 						{
@@ -247,7 +247,7 @@ FDelayedAutoRegisterHelper SAssetView::AssetViewOptionsMenuRegistration(
 								}
 
 								SizeSection.AddMenuEntry(
-									NAME_None,
+									"ThumbnailSizeValue",
 									SAssetView::ThumbnailSizeToDisplayName((EThumbnailSize)EnumValue),
 									FText::GetEmpty(),
 									FSlateIcon(),
@@ -261,14 +261,14 @@ FDelayedAutoRegisterHelper SAssetView::AssetViewOptionsMenuRegistration(
 							}
 						};
 
-						Section.AddEntry(FToolMenuEntry::InitSubMenu(
+						ThumbnailSection.AddEntry(FToolMenuEntry::InitSubMenu(
 							"ThumbnailSize",
 							LOCTEXT("ThumbnailSize", "Thumbnail Size"),
 							LOCTEXT("ThumbnailSizeToolTip", "Adjust the size of thumbnails."),
 							FNewToolMenuDelegate::CreateLambda(CreateThumbnailSizeSubMenu)
 						));
 
-						Section.AddMenuEntry(
+						ThumbnailSection.AddMenuEntry(
 							"ThumbnailEditMode",
 							LOCTEXT("ThumbnailEditModeOption", "Thumbnail Edit Mode"),
 							LOCTEXT("ThumbnailEditModeOptionToolTip", "Toggle thumbnail editing mode. When in this mode you can rotate the camera on 3D thumbnails by dragging them."),
@@ -281,7 +281,7 @@ FDelayedAutoRegisterHelper SAssetView::AssetViewOptionsMenuRegistration(
 							EUserInterfaceActionType::ToggleButton
 						);
 
-						Section.AddMenuEntry(
+						ThumbnailSection.AddMenuEntry(
 							"RealTimeThumbnails",
 							LOCTEXT("RealTimeThumbnailsOption", "Real-Time Thumbnails"),
 							LOCTEXT("RealTimeThumbnailsOptionToolTip", "Renders the assets thumbnails in real-time"),
@@ -296,15 +296,13 @@ FDelayedAutoRegisterHelper SAssetView::AssetViewOptionsMenuRegistration(
 					}
 
 					// Show/View Section
+					FToolMenuSection& ShowSection = InMenu->FindOrAddSection(
+						IsNewStyleEnabled() ? "Show" : "View",
+						IsNewStyleEnabled() ? LOCTEXT("ShowHeading", "Show") : LOCTEXT("ViewHeading", "View"));
 					{
-						FToolMenuSection& Section = InMenu->FindOrAddSection(
-							IsNewStyleEnabled() ? "Show" : "View",
-							IsNewStyleEnabled() ? LOCTEXT("ShowHeading", "Show") : LOCTEXT("ViewHeading", "View"),
-							FToolMenuInsert("AssetThumbnails", EToolMenuInsertType::After));
-
 						const FText ShowFoldersLabel = IsNewStyleEnabled() ? LOCTEXT("ShowFoldersOption_NewStyle", "Folders") :LOCTEXT("ShowFoldersOption", "Show Folders");
 						SetEntryOrder(
-							Section.AddMenuEntry(
+							ShowSection.AddMenuEntry(
 								"ShowFolders",
 								ShowFoldersLabel,
 								LOCTEXT("ShowFoldersOptionToolTip", "Show folders in the view as well as assets?"),
@@ -320,7 +318,7 @@ FDelayedAutoRegisterHelper SAssetView::AssetViewOptionsMenuRegistration(
 
 						const FText ShowEmptyFoldersLabel = IsNewStyleEnabled() ? LOCTEXT("ShowEmptyFoldersOption_NewStyle", "Empty Folders") :LOCTEXT("ShowEmptyFoldersOption", "Show Empty Folders");
 						SetEntryOrder(
-							Section.AddMenuEntry(
+							ShowSection.AddMenuEntry(
 								"ShowEmptyFolders",
 								ShowEmptyFoldersLabel,
 								LOCTEXT("ShowEmptyFoldersOptionToolTip", "Show empty folders in the view as well as assets?"),
@@ -335,7 +333,7 @@ FDelayedAutoRegisterHelper SAssetView::AssetViewOptionsMenuRegistration(
 						);
 
 						const FText ShowFavoriteLabel = IsNewStyleEnabled() ?  LOCTEXT("ShowFavoriteOptions_NewStyle", "Favorites") : LOCTEXT("ShowFavoriteOptions", "Show Favorites");
-						Section.AddMenuEntry(
+						ShowSection.AddMenuEntry(
 							"ShowFavorite",
 							ShowFavoriteLabel,
 							LOCTEXT("ShowFavoriteOptionToolTip", "Show the favorite folders in the view?"),
@@ -350,7 +348,7 @@ FDelayedAutoRegisterHelper SAssetView::AssetViewOptionsMenuRegistration(
 
 						if (!IsNewStyleEnabled())
 						{
-							Section.AddMenuEntry(
+							ShowSection.AddMenuEntry(
 								"FilterRecursively",
 								LOCTEXT("FilterRecursivelyOption", "Filter Recursively"),
 								LOCTEXT("FilterRecursivelyOptionToolTip", "Should filters apply recursively in the view?"),
@@ -365,7 +363,7 @@ FDelayedAutoRegisterHelper SAssetView::AssetViewOptionsMenuRegistration(
 						}
 
 						SetEntryOrder(
-							Section.AddMenuEntry(
+							ShowSection.AddMenuEntry(
 								"OrganizeFolders",
 								LOCTEXT("OrganizeFoldersOption", "Organize Folders"),
 								LOCTEXT("OrganizeFoldersOptionToolTip", "Organize folders in the view?"),
@@ -381,7 +379,7 @@ FDelayedAutoRegisterHelper SAssetView::AssetViewOptionsMenuRegistration(
 
 						if (AssetView->bShowPathViewFilters)
 						{
-							Section.AddSubMenu(
+							ShowSection.AddSubMenu(
 								"PathViewFilters",
 								LOCTEXT("PathViewFilters", "Path View Filters"),
 								LOCTEXT("PathViewFilters_ToolTip", "Path View Filters"),
@@ -391,18 +389,18 @@ FDelayedAutoRegisterHelper SAssetView::AssetViewOptionsMenuRegistration(
 
 					// Content Section
 					{
-						FToolMenuSection& Section = InMenu->FindOrAddSection(
+						FToolMenuSection& ContentSection = InMenu->FindOrAddSection(
 							"Content",
 							IsNewStyleEnabled() ? FText::GetEmpty() : LOCTEXT("ContentHeading", "Content"),
 							FToolMenuInsert(IsNewStyleEnabled() ? "Show" : "View", EToolMenuInsertType::After));
 
 						if (IsNewStyleEnabled())
 						{
-							Section.AddSeparator(NAME_None);
+							ContentSection.AddSeparator("ContentSeparator");
 						}
 
 						const FText ShowAllFolderLabel = IsNewStyleEnabled() ?  LOCTEXT("ShowAllFolderOption_NewStyle", "All Folder") : LOCTEXT("ShowAllFolderOption", "Show All Folder");
-						FToolMenuSection& ShowAllFolderSection = IsNewStyleEnabled() ? Section : InMenu->FindOrAddSection("View", { }, FToolMenuInsert("Content", EToolMenuInsertType::After));
+						FToolMenuSection& ShowAllFolderSection = IsNewStyleEnabled() ? ContentSection : ShowSection;
 						SetEntryOrder(
 							ShowAllFolderSection.AddMenuEntry(
 								"ShowAllFolder",
@@ -421,8 +419,12 @@ FDelayedAutoRegisterHelper SAssetView::AssetViewOptionsMenuRegistration(
 
 					// Search Section
 					{
-						FToolMenuSection& Section = InMenu->FindOrAddSection("Search", LOCTEXT("SearchHeading", "Search"), FToolMenuInsert("Content", EToolMenuInsertType::After));
-						Section.AddMenuEntry(
+						FToolMenuSection& SearchSection = InMenu->FindOrAddSection(
+							"Search",
+							LOCTEXT("SearchHeading", "Search"),
+							FToolMenuInsert("Content", EToolMenuInsertType::After));
+
+						SearchSection.AddMenuEntry(
 							"IncludeClassName",
 							LOCTEXT("IncludeClassNameOption", "Search Asset Class Names"),
 							LOCTEXT("IncludeClassesNameOptionTooltip", "Include asset type names in search criteria?  (e.g. Blueprint, Texture, Sound)"),
@@ -435,7 +437,7 @@ FDelayedAutoRegisterHelper SAssetView::AssetViewOptionsMenuRegistration(
 							EUserInterfaceActionType::ToggleButton
 						);
 
-						Section.AddMenuEntry(
+						SearchSection.AddMenuEntry(
 							"IncludeAssetPath",
 							LOCTEXT("IncludeAssetPathOption", "Search Asset Path"),
 							LOCTEXT("IncludeAssetPathOptionTooltip", "Include entire asset path in search criteria?"),
@@ -448,7 +450,7 @@ FDelayedAutoRegisterHelper SAssetView::AssetViewOptionsMenuRegistration(
 							EUserInterfaceActionType::ToggleButton
 						);
 
-						Section.AddMenuEntry(
+						SearchSection.AddMenuEntry(
 							"IncludeCollectionName",
 							LOCTEXT("IncludeCollectionNameOption", "Search Collection Names"),
 							LOCTEXT("IncludeCollectionNameOptionTooltip", "Include Collection names in search criteria?"),
@@ -511,12 +513,11 @@ FDelayedAutoRegisterHelper SAssetView::AssetViewOptionsMenuRegistration(
 
 			ContentBrowserMenuUtils::AddFiltersToMenu(InMenu, OwningContentBrowserName, FilterParams);
 		}));
-	},
-	true);
+	});
 
 FDelayedAutoRegisterHelper SAssetView::ToolBarMenuExtensionRegistration(
 	EDelayedRegisterRunPhase::EndOfEngineInit,
-	[](const EDelayedRegisterRunPhase InPhase)
+	[]
 	{
 		// @todo: Remove when old style is removed, all extensions are for the new style so early-out if it's not enabled
 		if (!IsNewStyleEnabled())
@@ -585,25 +586,24 @@ FDelayedAutoRegisterHelper SAssetView::ToolBarMenuExtensionRegistration(
 					}
 				}));
 		}
-	},
-	true);
+	});
 
 FDelayedAutoRegisterHelper SAssetView::NavigationBarMenuExtensionRegistration(
 	EDelayedRegisterRunPhase::EndOfEngineInit,
-	[](const EDelayedRegisterRunPhase InPhase)
+	[]
 	{
 		static const FName NavigationBarName("ContentBrowser.NavigationBar");
 		if (UToolMenu* NavigationBarMenu = UToolMenus::Get()->ExtendMenu(NavigationBarName))
 		{
-			FToolMenuSection& AssetCountSection = NavigationBarMenu->FindOrAddSection("AssetCount", { }, FToolMenuInsert("History", EToolMenuInsertType::Last));
-			AssetCountSection.Alignment = EToolMenuSectionAlign::Last;
-
 			static bool bIsNewStyleEnabled = UE::Editor::ContentBrowser::IsNewStyleEnabled();
 
 			if (bIsNewStyleEnabled)
 			{
+				FToolMenuSection& AssetCountSection = NavigationBarMenu->FindOrAddSection("AssetCount", { }, FToolMenuInsert("History", EToolMenuInsertType::Last));
+				AssetCountSection.Alignment = EToolMenuSectionAlign::Last;
+
 				AssetCountSection.AddDynamicEntry(
-					NAME_None,
+					"AssetCount",
 					FNewToolMenuSectionDelegate::CreateLambda([](FToolMenuSection& InSection)
 					{
 						if (const UContentBrowserToolbarMenuContext* Context = InSection.FindContext<UContentBrowserToolbarMenuContext>())
@@ -629,7 +629,6 @@ FDelayedAutoRegisterHelper SAssetView::NavigationBarMenuExtensionRegistration(
 					}));
 			}
 		}
-	},
-	true);
+	});
 
 #undef LOCTEXT_NAMESPACE
diff --git a/Engine/Source/Editor/ContentBrowser/Private/Menus/ContentBrowserMenus.cpp b/Engine/Source/Editor/ContentBrowser/Private/Menus/ContentBrowserMenus.cpp
--- a/Engine/Source/Editor/ContentBrowser/Private/Menus/ContentBrowserMenus.cpp
+++ b/Engine/Source/Editor/ContentBrowser/Private/Menus/ContentBrowserMenus.cpp
@@ -16,22 +16,18 @@
 
 FDelayedAutoRegisterHelper SContentBrowser::NavigationBarMenuRegistration(
 	EDelayedRegisterRunPhase::EndOfEngineInit,
-	[](const EDelayedRegisterRunPhase InPhase)
+	[]
 	{
-		static const FName NavigationBarMenuName("ContentBrowser.NavigationBar");
-		if (!ensureAlways(!UToolMenus::Get()->IsMenuRegistered(NavigationBarMenuName)))
-		{
-			return;
-		}
-
 		FToolMenuOwnerScoped OwnerScoped(UE_MODULE_NAME);
 
+		static const FName NavigationBarMenuName("ContentBrowser.NavigationBar");
 		UToolMenu* ToolBar = UToolMenus::Get()->RegisterMenu(NavigationBarMenuName, NAME_None, EMultiBoxType::SlimHorizontalToolBar);
 		ToolBar->SetStyleSet(&UE::ContentBrowser::Private::FContentBrowserStyle::Get());
 		ToolBar->StyleName = "ContentBrowser.ToolBar";
 
 		ToolBar->FindOrAddSection("History")
-		.AddDynamicEntry(NAME_None,
+		.AddDynamicEntry(
+			"History",
 			FNewToolMenuSectionDelegate::CreateLambda([](FToolMenuSection& InSection)
 			{
 				if (!UE::Editor::ContentBrowser::IsNewStyleEnabled())
@@ -81,7 +77,9 @@ FDelayedAutoRegisterHelper SContentBrowser::NavigationBarMenuRegistration(
 			}));
 
 		ToolBar->FindOrAddSection("Path", { }, FToolMenuInsert("History", EToolMenuInsertType::After))
-		.AddDynamicEntry(NAME_None, FNewToolMenuSectionDelegate::CreateLambda(
+		.AddDynamicEntry(
+			"Path",
+			FNewToolMenuSectionDelegate::CreateLambda(
 			[](FToolMenuSection& InSection)
 			{
 				if (!UE::Editor::ContentBrowser::IsNewStyleEnabled())
@@ -118,65 +116,50 @@ FDelayedAutoRegisterHelper SContentBrowser::NavigationBarMenuRegistration(
 					}
 				}
 			}));
-	},
-	true);
+	});
 
 FDelayedAutoRegisterHelper SContentBrowser::AddNewContextMenuRegistration(
 	EDelayedRegisterRunPhase::EndOfEngineInit,
-	[](const EDelayedRegisterRunPhase InPhase)
+	[]
 	{
-		static const FName AddNewContextMenuName("ContentBrowser.AddNewContextMenu");
-		if (!ensureAlways(!UToolMenus::Get()->IsMenuRegistered(AddNewContextMenuName)))
-		{
-			return;
-		}
-
-		FToolMenuOwnerScoped OwnerScoped(UE_MODULE_NAME);
-
-		UToolMenu* Menu = UToolMenus::Get()->RegisterMenu(AddNewContextMenuName);
-		Menu->AddDynamicSection("DynamicSection_Common", FNewToolMenuDelegate::CreateLambda([](UToolMenu* InMenu)
-		{
-			TSharedPtr<SContentBrowser> ContentBrowser;
-			if (const UContentBrowserMenuContext* MenuContext = InMenu->FindContext<UContentBrowserMenuContext>())
-			{
-				ContentBrowser = MenuContext->ContentBrowser.Pin();
-			}
-			else
-			{
-				if (const UContentBrowserToolbarMenuContext* ToolbarContext  = InMenu->FindContext<UContentBrowserToolbarMenuContext>())
-				{
-					ContentBrowser = ToolbarContext->ContentBrowser.Pin();
-				}
-			}
-
-			if (ContentBrowser)
-			{
-				ContentBrowser->PopulateAddNewContextMenu(InMenu);
-			}
-		}));
-	},
-	true);
+        FToolMenuOwnerScoped OwnerScoped(UE_MODULE_NAME);
+
+        static const FName AddNewContextMenuName("ContentBrowser.AddNewContextMenu");
+        UToolMenu* Menu = UToolMenus::Get()->RegisterMenu(AddNewContextMenuName);
+
+        Menu->AddDynamicSection("DynamicSection_Common", FNewToolMenuDelegate::CreateLambda([](UToolMenu* InMenu)
+        {
+            TSharedPtr<SContentBrowser> ContentBrowser;
+            if (const UContentBrowserMenuContext* MenuContext = InMenu->FindContext<UContentBrowserMenuContext>())
+            {
+                ContentBrowser = MenuContext->ContentBrowser.Pin();
+            }
+            else
+            {
+                if (const UContentBrowserToolbarMenuContext* ToolbarContext  = InMenu->FindContext<UContentBrowserToolbarMenuContext>())
+                {
+                    ContentBrowser = ToolbarContext->ContentBrowser.Pin();
+                }
+            }
+
+            if (ContentBrowser)
+            {
+                ContentBrowser->PopulateAddNewContextMenu(InMenu);
+            }
+        }));
+    });
 
 FDelayedAutoRegisterHelper SContentBrowser::FolderContextMenuRegistration(
 	EDelayedRegisterRunPhase::EndOfEngineInit,
-	[](const EDelayedRegisterRunPhase InPhase)
+	[]
 	{
-		static const FName FolderContextMenuName("ContentBrowser.FolderContextMenu");
-		if (InPhase == EDelayedRegisterRunPhase::LiveCodingReload)
-		{
-			UToolMenus::Get()->RemoveMenu(FolderContextMenuName);
-		}
-		else if (!ensureAlways(!UToolMenus::Get()->IsMenuRegistered(FolderContextMenuName)))
-		{
-			return;
-		}
-
 		FToolMenuOwnerScoped OwnerScoped(UE_MODULE_NAME);
 
+		static const FName FolderContextMenuName("ContentBrowser.FolderContextMenu");
 		UToolMenu* Menu = UToolMenus::Get()->RegisterMenu(FolderContextMenuName);
 		Menu->bCloseSelfOnly = true;
 
-		Menu->AddDynamicSection(NAME_None, FNewToolMenuDelegate::CreateLambda([](UToolMenu* InMenu)
+		Menu->AddDynamicSection(NAME_Default, FNewToolMenuDelegate::CreateLambda([](UToolMenu* InMenu)
 		{
 			UContentBrowserFolderContext* Context = InMenu->FindContext<UContentBrowserFolderContext>();
 			if (Context && Context->ContentBrowser.IsValid())
@@ -184,273 +167,259 @@ FDelayedAutoRegisterHelper SContentBrowser::FolderContextMenuRegistration(
 				Context->ContentBrowser.Pin()->PopulateFolderContextMenu(InMenu);
 			}
 		}));
-	},
-	true);
+	});
 
 FDelayedAutoRegisterHelper SContentBrowser::PathViewFiltersMenuRegistration(
 	EDelayedRegisterRunPhase::EndOfEngineInit,
-	[](const EDelayedRegisterRunPhase InPhase)
+	[]
 	{
-		static const FName PathViewFiltersMenuName("ContentBrowser.AssetViewOptions.PathViewFilters");
-		if (InPhase == EDelayedRegisterRunPhase::LiveCodingReload)
-		{
-			UToolMenus::Get()->RemoveMenu(PathViewFiltersMenuName);
-		}
-		else if (!ensureAlways(!UToolMenus::Get()->IsMenuRegistered(PathViewFiltersMenuName)))
-		{
-			return;
-		}
-
-		FToolMenuOwnerScoped OwnerScoped(UE_MODULE_NAME);
-
-		UToolMenu* Menu = UToolMenus::Get()->RegisterMenu(PathViewFiltersMenuName);
-
-		Menu->AddDynamicSection(NAME_None, FNewToolMenuDelegate::CreateLambda([](UToolMenu* InMenu)
-		{
-			if (const UContentBrowserAssetViewContextMenuContext* AssetViewContext = InMenu->FindContext<UContentBrowserAssetViewContextMenuContext>())
-			{
-				if (const TSharedPtr<SContentBrowser> ContentBrowser = AssetViewContext->OwningContentBrowser.Pin())
-				{
-					ContentBrowser->PopulatePathViewFiltersMenu(InMenu);
-				}
-			}
-			else if (const UContentBrowserMenuContext* ContentBrowserContext = InMenu->FindContext<UContentBrowserMenuContext>())
-			{
-				if (const TSharedPtr<SContentBrowser> ContentBrowser = ContentBrowserContext->ContentBrowser.Pin())
-				{
-					ContentBrowser->PopulatePathViewFiltersMenu(InMenu);
-				}
-			}
-		}));
-	},
-	true);
+        FToolMenuOwnerScoped OwnerScoped(UE_MODULE_NAME);
+
+        static const FName PathViewFiltersMenuName("ContentBrowser.AssetViewOptions.PathViewFilters");
+        UToolMenu* Menu = UToolMenus::Get()->RegisterMenu(PathViewFiltersMenuName);
+
+        Menu->AddDynamicSection(NAME_Default, FNewToolMenuDelegate::CreateLambda([](UToolMenu* InMenu)
+        {
+            if (const UContentBrowserAssetViewContextMenuContext* AssetViewContext = InMenu->FindContext<UContentBrowserAssetViewContextMenuContext>())
+            {
+                if (const TSharedPtr<SContentBrowser> ContentBrowser = AssetViewContext->OwningContentBrowser.Pin())
+                {
+                    ContentBrowser->PopulatePathViewFiltersMenu(InMenu);
+                }
+            }
+            else if (const UContentBrowserMenuContext* ContentBrowserContext = InMenu->FindContext<UContentBrowserMenuContext>())
+            {
+                if (const TSharedPtr<SContentBrowser> ContentBrowser = ContentBrowserContext->ContentBrowser.Pin())
+                {
+                    ContentBrowser->PopulatePathViewFiltersMenu(InMenu);
+                }
+            }
+        }));
+    });
 
 FDelayedAutoRegisterHelper SContentBrowser::ToolBarMenuRegistration(
 	EDelayedRegisterRunPhase::EndOfEngineInit,
-	[](const EDelayedRegisterRunPhase InPhase)
+	[]
 	{
-		static const FName ToolBarMenuName("ContentBrowser.ToolBar");
-		if (InPhase == EDelayedRegisterRunPhase::LiveCodingReload)
-		{
-			UToolMenus::Get()->RemoveMenu(ToolBarMenuName);
-		}
-		else if (!ensureAlways(!UToolMenus::Get()->IsMenuRegistered(ToolBarMenuName)))
-		{
-			return;
-		}
-
-		FToolMenuOwnerScoped OwnerScoped(UE_MODULE_NAME);
-
-		UToolMenu* ToolBar = UToolMenus::Get()->RegisterMenu(ToolBarMenuName, NAME_None, EMultiBoxType::SlimHorizontalToolBar);
-		ToolBar->SetStyleSet(&UE::ContentBrowser::Private::FContentBrowserStyle::Get());
-		ToolBar->StyleName = ToolBarMenuName; // Style name is the same as the menu entry name
-
-		ToolBar->FindOrAddSection("New")
-		.AddDynamicEntry(NAME_None,
-			FNewToolMenuSectionDelegate::CreateLambda([](FToolMenuSection& InSection)
-			{
-				if (UContentBrowserToolbarMenuContext* Context = InSection.FindContext<UContentBrowserToolbarMenuContext>())
-				{
-					if (Context->ContentBrowser.IsValid())
-					{
-						const TSharedRef<SContentBrowser> ContentBrowser = Context->ContentBrowser.Pin().ToSharedRef();
-
-						const TSharedRef<SPositiveActionButton> NewButton =
-							SNew(SPositiveActionButton)
-							.OnGetMenuContent_Lambda([Context]
-							{
-								return Context->ContentBrowser.Pin()->MakeAddNewContextMenu(EContentBrowserDataMenuContext_AddNewMenuDomain::Toolbar, Context);
-							})
-							.ToolTipText(ContentBrowser, &SContentBrowser::GetAddNewToolTipText)
-							.IsEnabled(ContentBrowser, &SContentBrowser::IsAddNewEnabled)
-							.AddMetaData<FTagMetaData>(FTagMetaData(TEXT("ContentBrowserNewAsset")))
-							.Icon(FAppStyle::Get().GetBrush("Icons.Plus"))
-							.Text(LOCTEXT("AddAssetButton", "Add"));
-
-						InSection.AddEntry(
-							FToolMenuEntry::InitWidget(
-								"NewButton",
-								NewButton,
-								FText::GetEmpty(),
-								true,
-								false
-							));
-					}
-				}
-			}));
+        FToolMenuOwnerScoped OwnerScoped(UE_MODULE_NAME);
 
-		ToolBar->FindOrAddSection("Save", { }, FToolMenuInsert("New", EToolMenuInsertType::After))
-		.AddDynamicEntry(NAME_None,
-			FNewToolMenuSectionDelegate::CreateLambda([](FToolMenuSection& InSection)
-			{
-				if (const UContentBrowserToolbarMenuContext* Context = InSection.FindContext<UContentBrowserToolbarMenuContext>())
-				{
-					if (Context->ContentBrowser.IsValid())
-					{
-						const TSharedRef<SContentBrowser> ContentBrowser = Context->ContentBrowser.Pin().ToSharedRef();
-
-						TSharedPtr<SWidget> SaveButton = nullptr;
-						if (UE::Editor::ContentBrowser::IsNewStyleEnabled())
-						{
-							SaveButton =
-								SNew(SActionButton)
-								.ToolTipText(LOCTEXT("SaveDirtyPackagesTooltip", "Save all modified assets."))
-								.OnClicked(ContentBrowser, &SContentBrowser::OnSaveClicked)
-								.AddMetaData<FTagMetaData>(FTagMetaData(TEXT("ContentBrowserSaveDirtyPackages")))
-								.Icon(FAppStyle::Get().GetBrush("MainFrame.SaveAll"))
-								.Text(LOCTEXT("SaveAll", "Save All"));
-						}
-						else
-						{
-							SaveButton =
-								SNew(SButton)
-								.ButtonStyle(FAppStyle::Get(), "SimpleButton")
-								.ToolTipText(LOCTEXT("SaveDirtyPackagesTooltip", "Save all modified assets."))
-								.ContentPadding(2.f)
-								.OnClicked(ContentBrowser, &SContentBrowser::OnSaveClicked)
-								.AddMetaData<FTagMetaData>(FTagMetaData(TEXT("ContentBrowserSaveDirtyPackages")))
-								[
-									SNew(SHorizontalBox)
-									// Save All Icon
-									+ SHorizontalBox::Slot()
-									.AutoWidth()
-									.HAlign(HAlign_Center)
-									.VAlign(VAlign_Center)
-									[
-										SNew(SImage)
-										.Image(FAppStyle::Get().GetBrush("MainFrame.SaveAll"))
-										.ColorAndOpacity(FSlateColor::UseForeground())
-									]
-
-									// Save All Text
-									+ SHorizontalBox::Slot()
-									.Padding(FMargin(3, 0, 0, 0))
-									.VAlign(VAlign_Center)
-									.AutoWidth()
-									[
-										SNew(STextBlock)
-										.TextStyle(FAppStyle::Get(), "NormalText")
-										.Text(LOCTEXT("SaveAll", "Save All"))
-									]
-								];
-						}
-
-						InSection.AddEntry(
-							FToolMenuEntry::InitWidget(
-								"SaveButton",
-								SaveButton.ToSharedRef(),
-								FText::GetEmpty(),
-								true,
-								false
-							));
-					}
-				}
-			}));
-
-		if (UE::Editor::ContentBrowser::IsNewStyleEnabled())
-		{
-			ToolBar->FindOrAddSection("Filter", { }, FToolMenuInsert("Save", EToolMenuInsertType::After))
-			.AddDynamicEntry(NAME_None, FNewToolMenuSectionDelegate::CreateLambda(
-				[](FToolMenuSection& InSection)
-				{
-					if (const UContentBrowserToolbarMenuContext* Context = InSection.FindContext<UContentBrowserToolbarMenuContext>())
-					{
-						if (Context->ContentBrowser.IsValid())
-						{
-							const TSharedRef<SContentBrowser> ContentBrowser = Context->ContentBrowser.Pin().ToSharedRef();
-
-							if (!ContentBrowser->FilterComboButton.IsValid())
-							{
-								return;
-							}
-
-							FToolMenuEntry& FilterEntry = InSection.AddEntry(
-								FToolMenuEntry::InitWidget(
-									"Filter",
-									SNew(SBox)
-									.HAlign(HAlign_Left)
-									.VAlign(VAlign_Center)
-									[
-										ContentBrowser->FilterComboButton.ToSharedRef()
-									],
-									LOCTEXT("Filter", "Filter")
-								));
-
-							FilterEntry.WidgetData.StyleParams.HorizontalAlignment = HAlign_Left;
-							FilterEntry.WidgetData.StyleParams.SizeRule = FSizeParam::ESizeRule::SizeRule_Auto;
-
-							// ... no label when in the toolbar
-							FilterEntry.ToolBarData.LabelOverride = FText::GetEmpty();
-						}
-					}
-				}));
-
-			ToolBar->FindOrAddSection("Search", { }, FToolMenuInsert("Filter", EToolMenuInsertType::After))
-			.AddDynamicEntry(NAME_None, FNewToolMenuSectionDelegate::CreateLambda(
-				[](FToolMenuSection& InSection)
-				{
-					static constexpr float SearchBoxMinWidth = 180.0f;
-					static constexpr float SearchBoxMaxWidth = 640.0f;
-
-					if (const UContentBrowserToolbarMenuContext* Context = InSection.FindContext<UContentBrowserToolbarMenuContext>())
-					{
-						if (Context->ContentBrowser.IsValid())
-						{
-							const TSharedRef<SContentBrowser> ContentBrowser = Context->ContentBrowser.Pin().ToSharedRef();
-
-							// Should have been initialized in SContentBrowser::CreateAssetView
-							if (!ensure(ContentBrowser->SearchBoxPtr.IsValid()))
-							{
-								return;
-							}
-
-							FToolMenuEntry& SearchBoxEntry =
-								InSection.AddEntry(
-								FToolMenuEntry::InitWidget(
-									"SearchBox",
-									ContentBrowser->SearchBoxPtr.ToSharedRef(),
-									FText::GetEmpty()));
-
-							SearchBoxEntry.WidgetData.StyleParams.SizeRule = FSizeParam::SizeRule_StretchContent;
-							SearchBoxEntry.WidgetData.StyleParams.MinSize = SearchBoxMinWidth;
-							SearchBoxEntry.WidgetData.StyleParams.MaxSize = SearchBoxMaxWidth;
-							SearchBoxEntry.WidgetData.StyleParams.DesiredWidth = SearchBoxMaxWidth;
-						}
-					}
-				}));
-
-			FToolMenuSection& SettingsSection = ToolBar->FindOrAddSection("Settings", { }, FToolMenuInsert("Search", EToolMenuInsertType::After));
-			SettingsSection.Alignment = EToolMenuSectionAlign::Last;
-
-			SettingsSection.AddDynamicEntry(NAME_None,
-			FNewToolMenuSectionDelegate::CreateLambda([](FToolMenuSection& InSection)
-			{
-				if (const UContentBrowserToolbarMenuContext* Context = InSection.FindContext<UContentBrowserToolbarMenuContext>())
-				{
-					if (Context->ContentBrowser.IsValid())
-					{
-						const TSharedRef<SContentBrowser> ContentBrowser = Context->ContentBrowser.Pin().ToSharedRef();
-
-						const TSharedRef<SWidget> LockToggleButton =
-							SNew(SActionButton)
-							.ActionButtonType(EActionButtonType::Simple)
-							.Icon(ContentBrowser, &SContentBrowser::GetLockIconBrush)
-							.IconColorAndOpacity(FSlateColor::UseStyle())
-							.ToolTipText(LOCTEXT("LockToggleTooltip", "Toggle lock. If locked, this browser will ignore Find in Content Browser requests."))
-							.OnClicked(ContentBrowser, &SContentBrowser::ToggleLockClicked)
-							.AddMetaData<FTagMetaData>(FTagMetaData(TEXT("ContentBrowserLock")));
-
-						InSection.AddEntry(
-							FToolMenuEntry::InitWidget(
-								"Lock",
-								LockToggleButton,
-								FText::GetEmpty()
-							));
-					}
-				}
-			}));
-		}
-	},
-	true);
+		static const FName ToolBarMenuName("ContentBrowser.ToolBar");
+        UToolMenu* ToolBar = UToolMenus::Get()->RegisterMenu(ToolBarMenuName, NAME_None, EMultiBoxType::SlimHorizontalToolBar);
+        ToolBar->SetStyleSet(&UE::ContentBrowser::Private::FContentBrowserStyle::Get());
+        ToolBar->StyleName = ToolBarMenuName; // Style name is the same as the menu entry name
+
+        ToolBar->FindOrAddSection("New")
+        .AddDynamicEntry(
+            "New",
+            FNewToolMenuSectionDelegate::CreateLambda([](FToolMenuSection& InSection)
+            {
+                if (UContentBrowserToolbarMenuContext* Context = InSection.FindContext<UContentBrowserToolbarMenuContext>())
+                {
+                    if (Context->ContentBrowser.IsValid())
+                    {
+                        const TSharedRef<SContentBrowser> ContentBrowser = Context->ContentBrowser.Pin().ToSharedRef();
+
+                        const TSharedRef<SPositiveActionButton> NewButton =
+                            SNew(SPositiveActionButton)
+                            .OnGetMenuContent_Lambda([Context]
+                            {
+                                return Context->ContentBrowser.Pin()->MakeAddNewContextMenu(EContentBrowserDataMenuContext_AddNewMenuDomain::Toolbar, Context);
+                            })
+                            .ToolTipText(ContentBrowser, &SContentBrowser::GetAddNewToolTipText)
+                            .IsEnabled(ContentBrowser, &SContentBrowser::IsAddNewEnabled)
+                            .AddMetaData<FTagMetaData>(FTagMetaData(TEXT("ContentBrowserNewAsset")))
+                            .Icon(FAppStyle::Get().GetBrush("Icons.Plus"))
+                            .Text(LOCTEXT("AddAssetButton", "Add"));
+
+                        InSection.AddEntry(
+                            FToolMenuEntry::InitWidget(
+                                "NewButton",
+                                NewButton,
+                                FText::GetEmpty(),
+                                true,
+                                false
+                            ));
+                    }
+                }
+            }));
+
+        ToolBar->FindOrAddSection("Save", { }, FToolMenuInsert("New", EToolMenuInsertType::After))
+        .AddDynamicEntry(
+            "Save",
+            FNewToolMenuSectionDelegate::CreateLambda([](FToolMenuSection& InSection)
+            {
+                if (const UContentBrowserToolbarMenuContext* Context = InSection.FindContext<UContentBrowserToolbarMenuContext>())
+                {
+                    if (Context->ContentBrowser.IsValid())
+                    {
+                        const TSharedRef<SContentBrowser> ContentBrowser = Context->ContentBrowser.Pin().ToSharedRef();
+
+                        TSharedPtr<SWidget> SaveButton = nullptr;
+                        if (UE::Editor::ContentBrowser::IsNewStyleEnabled())
+                        {
+                            SaveButton =
+                                SNew(SActionButton)
+                                .ToolTipText(LOCTEXT("SaveDirtyPackagesTooltip", "Save all modified assets."))
+                                .OnClicked(ContentBrowser, &SContentBrowser::OnSaveClicked)
+                                .AddMetaData<FTagMetaData>(FTagMetaData(TEXT("ContentBrowserSaveDirtyPackages")))
+                                .Icon(FAppStyle::Get().GetBrush("MainFrame.SaveAll"))
+                                .Text(LOCTEXT("SaveAll", "Save All"));
+                        }
+                        else
+                        {
+                            SaveButton =
+                                SNew(SButton)
+                                .ButtonStyle(FAppStyle::Get(), "SimpleButton")
+                                .ToolTipText(LOCTEXT("SaveDirtyPackagesTooltip", "Save all modified assets."))
+                                .ContentPadding(2.f)
+                                .OnClicked(ContentBrowser, &SContentBrowser::OnSaveClicked)
+                                .AddMetaData<FTagMetaData>(FTagMetaData(TEXT("ContentBrowserSaveDirtyPackages")))
+                                [
+                                    SNew(SHorizontalBox)
+                                    // Save All Icon
+                                    + SHorizontalBox::Slot()
+                                    .AutoWidth()
+                                    .HAlign(HAlign_Center)
+                                    .VAlign(VAlign_Center)
+                                    [
+                                        SNew(SImage)
+                                        .Image(FAppStyle::Get().GetBrush("MainFrame.SaveAll"))
+                                        .ColorAndOpacity(FSlateColor::UseForeground())
+                                    ]
+
+                                    // Save All Text
+                                    + SHorizontalBox::Slot()
+                                    .Padding(FMargin(3, 0, 0, 0))
+                                    .VAlign(VAlign_Center)
+                                    .AutoWidth()
+                                    [
+                                        SNew(STextBlock)
+                                        .TextStyle(FAppStyle::Get(), "NormalText")
+                                        .Text(LOCTEXT("SaveAll", "Save All"))
+                                    ]
+                                ];
+                        }
+
+                        InSection.AddEntry(
+                            FToolMenuEntry::InitWidget(
+                                "SaveButton",
+                                SaveButton.ToSharedRef(),
+                                FText::GetEmpty(),
+                                true,
+                                false
+                            ));
+                    }
+                }
+            }));
+
+        if (UE::Editor::ContentBrowser::IsNewStyleEnabled())
+        {
+            ToolBar->FindOrAddSection("Filter", { }, FToolMenuInsert("Save", EToolMenuInsertType::After))
+            .AddDynamicEntry(
+                "Filter",
+                FNewToolMenuSectionDelegate::CreateLambda(
+                [](FToolMenuSection& InSection)
+                {
+                    if (const UContentBrowserToolbarMenuContext* Context = InSection.FindContext<UContentBrowserToolbarMenuContext>())
+                    {
+                        if (Context->ContentBrowser.IsValid())
+                        {
+                            const TSharedRef<SContentBrowser> ContentBrowser = Context->ContentBrowser.Pin().ToSharedRef();
+
+                            if (!ContentBrowser->FilterComboButton.IsValid())
+                            {
+                                return;
+                            }
+
+                            FToolMenuEntry& FilterEntry = InSection.AddEntry(
+                                FToolMenuEntry::InitWidget(
+                                    "Filter",
+                                    SNew(SBox)
+                                    .HAlign(HAlign_Left)
+                                    .VAlign(VAlign_Center)
+                                    [
+                                        ContentBrowser->FilterComboButton.ToSharedRef()
+                                    ],
+                                    LOCTEXT("Filter", "Filter")
+                                ));
+
+                            FilterEntry.WidgetData.StyleParams.HorizontalAlignment = HAlign_Left;
+                            FilterEntry.WidgetData.StyleParams.SizeRule = FSizeParam::ESizeRule::SizeRule_Auto;
+
+                            // ... no label when in the toolbar
+                            FilterEntry.ToolBarData.LabelOverride = FText::GetEmpty();
+                        }
+                    }
+                }));
+
+            ToolBar->FindOrAddSection("Search", { }, FToolMenuInsert("Filter", EToolMenuInsertType::After))
+            .AddDynamicEntry(
+                "Search",
+                FNewToolMenuSectionDelegate::CreateLambda(
+                [](FToolMenuSection& InSection)
+                {
+                    static constexpr float SearchBoxMinWidth = 180.0f;
+                    static constexpr float SearchBoxMaxWidth = 640.0f;
+
+                    if (const UContentBrowserToolbarMenuContext* Context = InSection.FindContext<UContentBrowserToolbarMenuContext>())
+                    {
+                        if (Context->ContentBrowser.IsValid())
+                        {
+                            const TSharedRef<SContentBrowser> ContentBrowser = Context->ContentBrowser.Pin().ToSharedRef();
+
+                            // Should have been initialized in SContentBrowser::CreateAssetView
+                            if (!ensure(ContentBrowser->SearchBoxPtr.IsValid()))
+                            {
+                                return;
+                            }
+
+                            FToolMenuEntry& SearchBoxEntry =
+                                InSection.AddEntry(
+                                FToolMenuEntry::InitWidget(
+                                    "SearchBox",
+                                    ContentBrowser->SearchBoxPtr.ToSharedRef(),
+                                    FText::GetEmpty()));
+
+                            SearchBoxEntry.WidgetData.StyleParams.SizeRule = FSizeParam::SizeRule_StretchContent;
+                            SearchBoxEntry.WidgetData.StyleParams.MinSize = SearchBoxMinWidth;
+                            SearchBoxEntry.WidgetData.StyleParams.MaxSize = SearchBoxMaxWidth;
+                            SearchBoxEntry.WidgetData.StyleParams.DesiredWidth = SearchBoxMaxWidth;
+                        }
+                    }
+                }));
+
+            FToolMenuSection& SettingsSection = ToolBar->FindOrAddSection("Settings", { }, FToolMenuInsert("Search", EToolMenuInsertType::After));
+            SettingsSection.Alignment = EToolMenuSectionAlign::Last;
+
+            SettingsSection.AddDynamicEntry(
+	            "Settings",
+	            FNewToolMenuSectionDelegate::CreateLambda([](FToolMenuSection& InSection)
+	            {
+	                if (const UContentBrowserToolbarMenuContext* Context = InSection.FindContext<UContentBrowserToolbarMenuContext>())
+	                {
+	                    if (Context->ContentBrowser.IsValid())
+	                    {
+	                        const TSharedRef<SContentBrowser> ContentBrowser = Context->ContentBrowser.Pin().ToSharedRef();
+
+	                        const TSharedRef<SWidget> LockToggleButton =
+	                            SNew(SActionButton)
+	                            .ActionButtonType(EActionButtonType::Simple)
+	                            .Icon(ContentBrowser, &SContentBrowser::GetLockIconBrush)
+	                            .IconColorAndOpacity(FSlateColor::UseStyle())
+	                            .ToolTipText(LOCTEXT("LockToggleTooltip", "Toggle lock. If locked, this browser will ignore Find in Content Browser requests."))
+	                            .OnClicked(ContentBrowser, &SContentBrowser::ToggleLockClicked)
+	                            .AddMetaData<FTagMetaData>(FTagMetaData(TEXT("ContentBrowserLock")));
+
+	                        InSection.AddEntry(
+	                            FToolMenuEntry::InitWidget(
+	                                "Lock",
+	                                LockToggleButton,
+	                                FText::GetEmpty()
+	                            ));
+	                    }
+	                }
+	            }));
+        }
+    });
 
 #undef LOCTEXT_NAMESPACE
