[ContentBrowser] ShouldShowRedirectors changed to regular callback - not leaky lambda

diff --git a/Engine/Source/Editor/ContentBrowser/Private/SContentBrowser.cpp b/Engine/Source/Editor/ContentBrowser/Private/SContentBrowser.cpp
--- a/Engine/Source/Editor/ContentBrowser/Private/SContentBrowser.cpp
+++ b/Engine/Source/Editor/ContentBrowser/Private/SContentBrowser.cpp
@@ -386,6 +386,11 @@ SContentBrowser::~SContentBrowser()
 	}
 }
 
+bool SContentBrowser::ShouldShowRedirectors() const
+{
+	return FilterListPtr.IsValid() ? ContentBrowserUtils::ShouldShowRedirectors(FilterListPtr) : false;
+}
+
 BEGIN_SLATE_FUNCTION_BUILD_OPTIMIZATION
 void SContentBrowser::Construct( const FArguments& InArgs, const FName& InInstanceName, const FContentBrowserConfig* Config )
 {
@@ -460,13 +465,13 @@ void SContentBrowser::Construct( const FArguments& InArgs, const FName& InInstan
 			//.ThumbnailScale(Config != nullptr ? Config->ThumbnailScale : 0.18f)
 			.InitialViewType(Config != nullptr ? Config->InitialAssetViewType : EAssetViewType::Tile)
 			.OnNewItemRequested(this, &SContentBrowser::OnNewItemRequested)
-					   .OnItemSelectionChanged(this, &SContentBrowser::OnItemSelectionChanged, EContentBrowserViewContext::AssetView)
+			.OnItemSelectionChanged(this, &SContentBrowser::OnItemSelectionChanged, EContentBrowserViewContext::AssetView)
 			.OnItemsActivated(this, &SContentBrowser::OnItemsActivated)
 			.OnGetItemContextMenu(this, &SContentBrowser::GetItemContextMenu, EContentBrowserViewContext::AssetView)
 			.OnItemRenameCommitted(this, &SContentBrowser::OnItemRenameCommitted)
 			.FrontendFilters(FrontendFilters)
 			.TextFilter(TextFilter)
-			.ShowRedirectors_Lambda([this]() { return ContentBrowserUtils::ShouldShowRedirectors(FilterListPtr); })
+			.ShowRedirectors(this, &SContentBrowser::ShouldShowRedirectors)
 			.HighlightedText(this, &SContentBrowser::GetHighlightedText)
 			.ShowBottomToolbar(bShowBottomToolbar)
 			.ShowViewOptions(false) // We control this for the main content browser
diff --git a/Engine/Source/Editor/ContentBrowser/Private/SContentBrowser.h b/Engine/Source/Editor/ContentBrowser/Private/SContentBrowser.h
--- a/Engine/Source/Editor/ContentBrowser/Private/SContentBrowser.h
+++ b/Engine/Source/Editor/ContentBrowser/Private/SContentBrowser.h
@@ -106,6 +106,8 @@ public:
 
 	~SContentBrowser();
 
+	bool ShouldShowRedirectors() const;
+
 	/** Constructs this widget with InArgs */
 	void Construct( const FArguments& InArgs, const FName& InstanceName, const FContentBrowserConfig* Config );
 
