Move the Fab button after Save All in the content browser

diff --git a/Engine/Plugins/Fab/Source/Fab/Private/FabBrowser.cpp b/Engine/Plugins/Fab/Source/Fab/Private/FabBrowser.cpp
--- a/Engine/Plugins/Fab/Source/Fab/Private/FabBrowser.cpp
+++ b/Engine/Plugins/Fab/Source/Fab/Private/FabBrowser.cpp
@@ -110,7 +110,8 @@ void FFabBrowser::SetupEntryPoints()
 		FCanExecuteAction()
 	);
 
-	FToolMenuEntry& ToolMenuEntry = UToolMenus::Get()->ExtendMenu("ContentBrowser.Toolbar")->FindOrAddSection("New").AddEntry(
+	FToolMenuSection& SaveSection = UToolMenus::Get()->ExtendMenu("ContentBrowser.Toolbar")->FindOrAddSection("Save");
+	FToolMenuEntry& ToolMenuEntry = SaveSection.AddEntry(
 		FToolMenuEntry::InitToolBarButton(
 			"OpenFabWindow",
 			InvokeTabAction,
@@ -120,7 +121,8 @@ void FFabBrowser::SetupEntryPoints()
 			EUserInterfaceActionType::Button
 		)
 	);
-	ToolMenuEntry.StyleNameOverride = FName("CalloutToolbar");
+	ToolMenuEntry.InsertPosition.Position = EToolMenuInsertType::Last;
+	ToolMenuEntry.StyleNameOverride = "ContentBrowser.ToolBar.Buttons";
 
 	UToolMenu* WindowMenu = UToolMenus::Get()->ExtendMenu("MainFrame.MainMenu.Window");
 
diff --git a/Engine/Source/Editor/ContentBrowser/Private/ContentBrowserStyle.cpp b/Engine/Source/Editor/ContentBrowser/Private/ContentBrowserStyle.cpp
--- a/Engine/Source/Editor/ContentBrowser/Private/ContentBrowserStyle.cpp
+++ b/Engine/Source/Editor/ContentBrowser/Private/ContentBrowserStyle.cpp
@@ -330,7 +330,7 @@ namespace UE::ContentBrowser::Private
 
 		Set("ContentBrowser.Splitter", ContentBrowserSplitterStyle);
 
-		FToolBarStyle ContentBrowserToolBarStyle = FStarshipCoreStyle::GetCoreStyle().GetWidgetStyle<FToolBarStyle>("SlimToolBar");
+		FToolBarStyle ContentBrowserToolBarStyle = FAppStyle::Get().GetWidgetStyle<FToolBarStyle>("SlimToolBar");
 
 		ContentBrowserToolBarStyle.SetSeparatorBrush(FSlateNoResource());
 		ContentBrowserToolBarStyle.SetSeparatorPadding(FMargin(4.0f, 0.0f));
@@ -347,6 +347,14 @@ namespace UE::ContentBrowser::Private
 
 		Set("ContentBrowser.ToolBar", ContentBrowserToolBarStyle);
 
+		if (UE::Editor::ContentBrowser::IsNewStyleEnabled())
+		{
+			ContentBrowserToolBarStyle.SetButtonStyle(FAppStyle::Get().GetWidgetStyle<FButtonStyle>("Button"));
+		}
+
+		// Separate style for buttons prior to the search box to get around SToolBarComboButtonBlock referencing the wrong style
+		Set("ContentBrowser.ToolBar.Buttons", ContentBrowserToolBarStyle);
+
 		FSlateStyleRegistry::RegisterSlateStyle(*this);
 	}
 
