Fixed MacArm64 compiler error: declaration shadows a field of 'SContentBrowser'

diff --git a/Engine/Source/Editor/ContentBrowser/Private/SContentBrowser.cpp b/Engine/Source/Editor/ContentBrowser/Private/SContentBrowser.cpp
--- a/Engine/Source/Editor/ContentBrowser/Private/SContentBrowser.cpp
+++ b/Engine/Source/Editor/ContentBrowser/Private/SContentBrowser.cpp
@@ -390,18 +390,18 @@ public:
 		return GetCollectionContainer() == FCollectionManagerModule::GetModule().Get().GetProjectCollectionContainer();
 	}
 
-	void LoadSettings(const FName& InstanceName)
+	void LoadSettings(const FName& InInstanceName)
 	{
-		const FString SettingsString = GetSettingsString(InstanceName);
+		const FString SettingsString = GetSettingsString(InInstanceName);
 		const FString EditorPerProjectIni = GetCollectionContainer()->GetCollectionSource()->GetEditorPerProjectIni();
 
 		CollectionViewPtr->LoadSettings(EditorPerProjectIni, SettingsIniSection, SettingsString);
 		CollectionArea->LoadSettings(EditorPerProjectIni, SettingsIniSection, SettingsString);
 	}
 
-	void SaveSettings(const FName& InstanceName) const
+	void SaveSettings(const FName& InInstanceName) const
 	{
-		const FString SettingsString = GetSettingsString(InstanceName);
+		const FString SettingsString = GetSettingsString(InInstanceName);
 		const FString EditorPerProjectIni = GetCollectionContainer()->GetCollectionSource()->GetEditorPerProjectIni();
 
 		CollectionViewPtr->SaveSettings(EditorPerProjectIni, SettingsIniSection, SettingsString);
