AssetViewWidgets - various fixes to extra state icon widgets
The widgets will now use .AutoWidth()
The max width and height are now the same and limited to a fraction of the whole area to ensure icons are not stretched

diff --git a/Engine/Source/Editor/ContentBrowser/Private/AssetViewWidgets.cpp b/Engine/Source/Editor/ContentBrowser/Private/AssetViewWidgets.cpp
--- a/Engine/Source/Editor/ContentBrowser/Private/AssetViewWidgets.cpp
+++ b/Engine/Source/Editor/ContentBrowser/Private/AssetViewWidgets.cpp
@@ -240,8 +240,8 @@ TSharedRef<SWidget> FAssetViewItemHelper::CreateListTileItemContents(T* const In
 		.VAlign(VAlign_Top)
 		[
 			SNew(SBox)
-			.MaxDesiredWidth(InTileOrListItem, &T::GetExtraStateIconMaxWidth)
-			.MaxDesiredHeight(InTileOrListItem, &T::GetStateIconImageSize)
+			.MaxDesiredWidth(InTileOrListItem, &T::GetExtraStateIconMaxSize)
+			.MaxDesiredHeight(InTileOrListItem, &T::GetExtraStateIconMaxSize)
 			[
 				InTileOrListItem->GenerateExtraStateIconWidget(TAttribute<float>(InTileOrListItem, &T::GetExtraStateIconWidth))
 			]
@@ -726,6 +726,7 @@ TSharedRef<SWidget> SAssetViewItem::GenerateExtraStateIconWidget(TAttribute<floa
 				{
 					Content->AddSlot()
 						.HAlign(HAlign_Left)
+						.AutoWidth()
 						.MaxWidth(InMaxExtraStateIconWidth)
 						[
 							Generator.IconGenerator.Execute(ItemAssetData)
@@ -2026,10 +2027,11 @@ TSharedRef<SWidget> SAssetListItem::GenerateWidgetForColumn( const FName& Column
 
 float SAssetListItem::GetExtraStateIconWidth() const
 {
-	return GetStateIconImageSize().Get();
+	// Each individual extra state icon will also simply be limited to the max width so things like long text show up properly
+	return GetExtraStateIconMaxSize().Get();
 }
 
-FOptionalSize SAssetListItem::GetExtraStateIconMaxWidth() const
+FOptionalSize SAssetListItem::GetExtraStateIconMaxSize() const
 {
 	return GetThumbnailBoxSize().Get() * 0.7f;
 }
@@ -2465,12 +2467,13 @@ void SAssetTileItem::OnRequestingRename()
 
 float SAssetTileItem::GetExtraStateIconWidth() const
 {
-	return GetStateIconImageSize().Get();
+	// Each individual extra state icon will also simply be limited to the max width so things like long text show up properly
+	return GetExtraStateIconMaxSize().Get();
 }
 
-FOptionalSize SAssetTileItem::GetExtraStateIconMaxWidth() const
+FOptionalSize SAssetTileItem::GetExtraStateIconMaxSize() const
 {
-	return GetThumbnailBoxSize().Get() * 0.8f;
+	return GetThumbnailBoxSize().Get() * 0.7f;
 }
 
 FOptionalSize SAssetTileItem::GetStateIconImageSize() const
diff --git a/Engine/Source/Editor/ContentBrowser/Public/AssetViewWidgets.h b/Engine/Source/Editor/ContentBrowser/Public/AssetViewWidgets.h
--- a/Engine/Source/Editor/ContentBrowser/Public/AssetViewWidgets.h
+++ b/Engine/Source/Editor/ContentBrowser/Public/AssetViewWidgets.h
@@ -463,8 +463,8 @@ private:
 	/** Get the expected width of an extra state icon. */
 	float GetExtraStateIconWidth() const;
 
-	/** Returns the max width size to be used by extra state icons. */
-	FOptionalSize GetExtraStateIconMaxWidth() const;
+	/** Returns the max size to be used by extra state icons. */
+	FOptionalSize GetExtraStateIconMaxSize() const;
 
 	/** Returns the size of the state icon box widget (i.e dirty image, scc)*/
 	FOptionalSize GetStateIconImageSize() const;
@@ -715,8 +715,8 @@ protected:
 	/** Get the expected width of an extra state icon. */
 	float GetExtraStateIconWidth() const;
 
-	/** Returns the max width size to be used by extra state icons. */
-	FOptionalSize GetExtraStateIconMaxWidth() const;
+	/** Returns the max size to be used by extra state icons. */
+	FOptionalSize GetExtraStateIconMaxSize() const;
 
 	/** Returns the size of the state icon box widget (i.e dirty image, scc)*/
 	FOptionalSize GetStateIconImageSize() const;
