[ContentBrowser] Added missing ThumbnailSize entries in the view options menu

diff --git a/Engine/Source/Editor/ContentBrowser/Private/Menus/AssetViewMenus.cpp b/Engine/Source/Editor/ContentBrowser/Private/Menus/AssetViewMenus.cpp
--- a/Engine/Source/Editor/ContentBrowser/Private/Menus/AssetViewMenus.cpp
+++ b/Engine/Source/Editor/ContentBrowser/Private/Menus/AssetViewMenus.cpp
@@ -247,7 +247,7 @@ FDelayedAutoRegisterHelper SAssetView::AssetViewOptionsMenuRegistration(
 								}
 
 								SizeSection.AddMenuEntry(
-									"ThumbnailSizeValue",
+									FName(FString::Printf(TEXT("ThumbnailSizeValue_%d"), EnumValue)),
 									SAssetView::ThumbnailSizeToDisplayName((EThumbnailSize)EnumValue),
 									FText::GetEmpty(),
 									FSlateIcon(),
