[ContentBrowser] Fixed missing "Show All Folder" entry in the view options menu

diff --git a/Engine/Source/Editor/ContentBrowser/Private/Menus/AssetViewMenus.cpp b/Engine/Source/Editor/ContentBrowser/Private/Menus/AssetViewMenus.cpp
--- a/Engine/Source/Editor/ContentBrowser/Private/Menus/AssetViewMenus.cpp
+++ b/Engine/Source/Editor/ContentBrowser/Private/Menus/AssetViewMenus.cpp
@@ -296,10 +296,11 @@ FDelayedAutoRegisterHelper SAssetView::AssetViewOptionsMenuRegistration(
 					}
 
 					// Show/View Section
-					FToolMenuSection& ShowSection = InMenu->FindOrAddSection(
-						IsNewStyleEnabled() ? "Show" : "View",
-						IsNewStyleEnabled() ? LOCTEXT("ShowHeading", "Show") : LOCTEXT("ViewHeading", "View"));
 					{
+						FToolMenuSection& ShowSection = InMenu->FindOrAddSection(
+							IsNewStyleEnabled() ? "Show" : "View",
+							IsNewStyleEnabled() ? LOCTEXT("ShowHeading", "Show") : LOCTEXT("ViewHeading", "View"));
+
 						const FText ShowFoldersLabel = IsNewStyleEnabled() ? LOCTEXT("ShowFoldersOption_NewStyle", "Folders") :LOCTEXT("ShowFoldersOption", "Show Folders");
 						SetEntryOrder(
 							ShowSection.AddMenuEntry(
@@ -400,7 +401,7 @@ FDelayedAutoRegisterHelper SAssetView::AssetViewOptionsMenuRegistration(
 						}
 
 						const FText ShowAllFolderLabel = IsNewStyleEnabled() ?  LOCTEXT("ShowAllFolderOption_NewStyle", "All Folder") : LOCTEXT("ShowAllFolderOption", "Show All Folder");
-						FToolMenuSection& ShowAllFolderSection = IsNewStyleEnabled() ? ContentSection : ShowSection;
+						FToolMenuSection& ShowAllFolderSection = IsNewStyleEnabled() ? ContentSection : *InMenu->FindSection("View");
 						SetEntryOrder(
 							ShowAllFolderSection.AddMenuEntry(
 								"ShowAllFolder",
