Added a scrollbox to account for large warning texts

diff --git a/Engine/Source/Editor/ContentBrowser/Private/SAssetView.cpp b/Engine/Source/Editor/ContentBrowser/Private/SAssetView.cpp
--- a/Engine/Source/Editor/ContentBrowser/Private/SAssetView.cpp
+++ b/Engine/Source/Editor/ContentBrowser/Private/SAssetView.cpp
@@ -72,6 +72,7 @@
 #include "Widgets/Images/SImage.h"
 #include "Widgets/Input/SComboButton.h"
 #include "Widgets/Layout/SScrollBorder.h"
+#include "Widgets/Layout/SScrollBox.h"
 #include "Widgets/Layout/SSplitter.h"
 #include "Widgets/Notifications/SProgressBar.h"
 #include "Widgets/SOverlay.h"
@@ -1364,12 +1365,16 @@ void SAssetView::Construct( const FArguments& InArgs )
 			.VAlign(VAlign_Center)
 			.Padding(FMargin(0, 14, 0, 0))
 			[
-				// A warning to display when there are no assets to show
-				SNew( STextBlock )
-				.Justification( ETextJustify::Center )
-				.Text( this, &SAssetView::GetAssetShowWarningText )
-				.Visibility( this, &SAssetView::IsAssetShowWarningTextVisible )
-				.AutoWrapText( true )
+				SNew(SScrollBox)
+				+SScrollBox::Slot()
+				[
+					// A warning to display when there are no assets to show
+					SNew( STextBlock )
+					.Justification( ETextJustify::Center )
+					.Text( this, &SAssetView::GetAssetShowWarningText )
+					.Visibility( this, &SAssetView::IsAssetShowWarningTextVisible )
+					.AutoWrapText( true )
+				]
 			]
 
 			+ SOverlay::Slot()
