[ContentBrowser] Docking dismissal fix

diff --git a/Engine/Source/Editor/ContentBrowser/Private/ContentBrowserSingleton.cpp b/Engine/Source/Editor/ContentBrowser/Private/ContentBrowserSingleton.cpp
--- a/Engine/Source/Editor/ContentBrowser/Private/ContentBrowserSingleton.cpp
+++ b/Engine/Source/Editor/ContentBrowser/Private/ContentBrowserSingleton.cpp
@@ -833,6 +833,8 @@ void FContentBrowserSingleton::DockContentBrowserDrawer()
 				}
 			}
 		}
+
+		GEditor->GetEditorSubsystem<UStatusBarSubsystem>()->ForceDismissDrawer();
 	}
 
 }
