[Content Browser] Adding Go Up to Parent Folder command to the content browser
- Also setting FGenericCommands as the parent context of FContentBrowserCommands

diff --git a/Engine/Source/Editor/ContentBrowser/Private/ContentBrowserCommands.cpp b/Engine/Source/Editor/ContentBrowser/Private/ContentBrowserCommands.cpp
--- a/Engine/Source/Editor/ContentBrowser/Private/ContentBrowserCommands.cpp
+++ b/Engine/Source/Editor/ContentBrowser/Private/ContentBrowserCommands.cpp
@@ -17,6 +17,7 @@ void FContentBrowserCommands::RegisterCommands()
 	UI_COMMAND(OpenAssetsOrFolders, "Open Assets or Folders", "Opens the selected assets or folders, depending on the selection", EUserInterfaceActionType::Button, FInputChord(EKeys::Enter), FInputChord(EModifierKey::Control, EKeys::E));
 	UI_COMMAND(PreviewAssets, "Preview Assets", "Loads the selected assets and previews them if possible", EUserInterfaceActionType::Button, FInputChord(EKeys::SpaceBar));
 	UI_COMMAND(CreateNewFolder, "Create New Folder", "Creates new folder in selected path", EUserInterfaceActionType::Button, FInputChord(EModifierKey::Control | EModifierKey::Shift, EKeys::N));
+	UI_COMMAND(GoUpToParentFolder, "Go Up to Parent Folder", "Opens the folder that contains the currently open one", EUserInterfaceActionType::Button, FInputChord(EModifierKey::Control , EKeys::BackSpace));
 	UI_COMMAND(SaveSelectedAsset, "Save Selected Item", "Save the selected item", EUserInterfaceActionType::Button, FInputChord(EModifierKey::Control, EKeys::S));
 	UI_COMMAND(SaveAllCurrentFolder, "Save All", "Save All in current folder", EUserInterfaceActionType::Button, FInputChord());
 	UI_COMMAND(ResaveAllCurrentFolder, "Resave All", "Resave all assets contained in the current folder", EUserInterfaceActionType::Button, FInputChord());
diff --git a/Engine/Source/Editor/ContentBrowser/Private/ContentBrowserCommands.h b/Engine/Source/Editor/ContentBrowser/Private/ContentBrowserCommands.h
--- a/Engine/Source/Editor/ContentBrowser/Private/ContentBrowserCommands.h
+++ b/Engine/Source/Editor/ContentBrowser/Private/ContentBrowserCommands.h
@@ -19,7 +19,7 @@ public:
 
 	/** Default constructor. */
 	FContentBrowserCommands()
-		: TCommands<FContentBrowserCommands>(TEXT("ContentBrowser"), NSLOCTEXT( "ContentBrowser", "ContentBrowser", "Content Browser" ), NAME_None, FAppStyle::GetAppStyleSetName() )
+		: TCommands<FContentBrowserCommands>(TEXT("ContentBrowser"), NSLOCTEXT( "ContentBrowser", "ContentBrowser", "Content Browser" ), "GenericCommands", FAppStyle::GetAppStyleSetName() )
 	{ }
 
 public:
@@ -32,6 +32,7 @@ public:
 
 	TSharedPtr<FUICommandInfo> CreateNewFolder;
 	TSharedPtr<FUICommandInfo> OpenAssetsOrFolders;
+	TSharedPtr<FUICommandInfo> GoUpToParentFolder;
 	TSharedPtr<FUICommandInfo> PreviewAssets;
 	TSharedPtr<FUICommandInfo> SaveSelectedAsset;
 	TSharedPtr<FUICommandInfo> SaveAllCurrentFolder;
diff --git a/Engine/Source/Editor/ContentBrowser/Private/SContentBrowser.cpp b/Engine/Source/Editor/ContentBrowser/Private/SContentBrowser.cpp
--- a/Engine/Source/Editor/ContentBrowser/Private/SContentBrowser.cpp
+++ b/Engine/Source/Editor/ContentBrowser/Private/SContentBrowser.cpp
@@ -1121,6 +1121,11 @@ void SContentBrowser::BindCommands()
 		FExecuteAction::CreateSP(this, &SContentBrowser::HandleCreateNewFolderCommandExecute)
 	));
 
+	Commands->MapAction(FContentBrowserCommands::Get().GoUpToParentFolder, FUIAction(
+		FExecuteAction::CreateSP(this, &SContentBrowser::HandleGoUpToParentFolder),
+		FCanExecuteAction::CreateSP(this, &SContentBrowser::HandleCanGoUpToParentFolder)
+	));
+
 	Commands->MapAction(FContentBrowserCommands::Get().SaveSelectedAsset, FUIAction(
 		FExecuteAction::CreateSP(this, &SContentBrowser::HandleSaveAssetCommand),
 		FCanExecuteAction::CreateSP(this, &SContentBrowser::HandleSaveAssetCommandCanExecute)
@@ -4128,6 +4133,33 @@ void SContentBrowser::HandleCreateNewFolderCommandExecute()
 	}
 }
 
+void SContentBrowser::HandleGoUpToParentFolder()
+{
+	const FString SelectedPath = PathViewPtr->GetSelectedPath();
+	int32 LastSlashIdx = INDEX_NONE;
+	if (ensure(SelectedPath.FindLastChar('/', LastSlashIdx)))
+	{
+		const int32 ChopCount = SelectedPath.Len() - LastSlashIdx;
+		const FString NewPathSelection = SelectedPath.LeftChop(ChopCount);
+		SetSelectedPaths({NewPathSelection}, true);
+	}
+}
+
+bool SContentBrowser::HandleCanGoUpToParentFolder() const
+{
+	// Allow going up if there's one non-root folder selected
+	TArray<FString> SelectedPaths = PathViewPtr->GetSelectedPaths();
+	if (SelectedPaths.Num() == 1)
+	{
+		int32 LastSlashIdx = INDEX_NONE;
+		if (SelectedPaths[0].FindLastChar('/', LastSlashIdx))
+		{
+			return LastSlashIdx > 0;
+		}
+	}
+	return false;
+}
+
 void SContentBrowser::GetSelectionState(TArray<FAssetData>& SelectedAssets, TArray<FString>& SelectedPaths)
 {
 	SelectedAssets.Reset();
diff --git a/Engine/Source/Editor/ContentBrowser/Private/SContentBrowser.h b/Engine/Source/Editor/ContentBrowser/Private/SContentBrowser.h
--- a/Engine/Source/Editor/ContentBrowser/Private/SContentBrowser.h
+++ b/Engine/Source/Editor/ContentBrowser/Private/SContentBrowser.h
@@ -481,6 +481,12 @@ private:
 	/** Handler for creating new folder */
 	void HandleCreateNewFolderCommandExecute();
 
+	/** Handler for navigating up to the parent folder of the one currently open */
+	void HandleGoUpToParentFolder();
+
+	/** True if the user may navigate from the currently selected folder to its parent */
+	bool HandleCanGoUpToParentFolder() const;
+	
 	/** True if the user may use the history back button */
 	bool IsBackEnabled() const;
 
