Fixed empty folders showing up in the content browser path bread crumbs when Show Empty Folders is unchecked.

diff --git a/Engine/Source/Editor/ContentBrowser/Private/SContentBrowser.cpp b/Engine/Source/Editor/ContentBrowser/Private/SContentBrowser.cpp
--- a/Engine/Source/Editor/ContentBrowser/Private/SContentBrowser.cpp
+++ b/Engine/Source/Editor/ContentBrowser/Private/SContentBrowser.cpp
@@ -209,6 +209,15 @@ namespace ContentBrowser
 			bDisplayEmpty = EditorConfig->bShowEmptyFolders;
 		}
 
+		EContentBrowserIsFolderVisibleFlags FolderFlags = ContentBrowserUtils::GetIsFolderVisibleFlags(bDisplayEmpty);
+
+		TOptional<FContentBrowserFolderContentsFilter> FolderFilter;
+		if (!bDisplayEmpty)
+		{
+			FolderFilter = FContentBrowserFolderContentsFilter{};
+			FolderFilter->ItemCategoryFilter = ItemCategoryFilter;
+		}
+
 		UContentBrowserDataSubsystem* ContentBrowserData = IContentBrowserDataModule::Get().GetSubsystem();
 
 		FContentBrowserDataFilter SubItemsFilter;
@@ -275,7 +284,7 @@ namespace ContentBrowser
 				}
 			}
 
-			if (!ContentBrowserData->IsFolderVisible(Item.GetVirtualPath(), ContentBrowserUtils::GetIsFolderVisibleFlags(bDisplayEmpty)))
+			if (!ContentBrowserData->IsFolderVisible(Item.GetVirtualPath(), FolderFlags, FolderFilter))
 			{
 				It.RemoveCurrent();
 				continue;
