Add the Navigation Bar to the Asset Browser.
- Move some shared behavior out of ContentBrowser.cpp and into ContentBrowserUtils.
- Add support for dialog-local history with back/forward slate buttons and back/forward mouse buttons.
- Add `SyncToItems()` functions to Path and Asset Picker.

diff --git a/Engine/Source/Editor/ContentBrowser/Private/ContentBrowserUtils.cpp b/Engine/Source/Editor/ContentBrowser/Private/ContentBrowserUtils.cpp
--- a/Engine/Source/Editor/ContentBrowser/Private/ContentBrowserUtils.cpp
+++ b/Engine/Source/Editor/ContentBrowser/Private/ContentBrowserUtils.cpp
@@ -6,12 +6,14 @@
 #include "AssetRegistry/ARFilter.h"
 #include "AssetRegistry/AssetData.h"
 #include "CollectionManagerModule.h"
+#include "CollectionViewUtils.h"
 #include "Containers/Map.h"
 #include "Containers/StringView.h"
 #include "ContentBrowserConfig.h"
 #include "ContentBrowserDataFilter.h"
 #include "ContentBrowserDataSource.h"
 #include "ContentBrowserDataSubsystem.h"
+#include "ContentBrowserDataUtils.h"
 #include "ContentBrowserItem.h"
 #include "ContentBrowserItemData.h"
 #include "ContentBrowserSingleton.h"
@@ -66,6 +68,7 @@
 #include "Widgets/Text/STextBlock.h"
 #include "Algo/Transform.h"
 #include "ContentBrowserStyle.h"
+#include "SNavigationBar.h"
 
 class SWidget;
 struct FGeometry;
@@ -1267,4 +1270,276 @@ FContentBrowserItem ContentBrowserUtils::TryGetItemFromUserProvidedPath(FStringV
 
 	return FContentBrowserItem();
 }
+
+FString ContentBrowserUtils::FormatCollectionCrumbData(const ICollectionContainer& CollectionContainer, const FCollectionNameType& Collection)
+{
+	return FString::Printf(
+		TEXT("%s?%s?%s"),
+		*CollectionContainer.GetCollectionSource()->GetName().ToString(),
+		*Collection.Name.ToString(),
+		*FString::FromInt(Collection.Type));
+}
+
+void ContentBrowserUtils::UpdateNavigationBar(const TSharedPtr<SNavigationBar>& NavigationBar, const TSharedPtr<SAssetView>& AssetView, const TSharedPtr<SPathView>& PathView)
+{
+	const FContentSources& ContentSources = AssetView->GetContentSources();
+
+	NavigationBar->ClearPaths();
+
+	if (ContentSources.HasVirtualPaths())
+	{
+		UContentBrowserDataSubsystem* ContentBrowserData = IContentBrowserDataModule::Get().GetSubsystem();
+
+		TArray<FString> Crumbs;
+		ContentSources.GetVirtualPaths()[0].ToString().ParseIntoArray(Crumbs, TEXT("/"), true);
+
+		FContentBrowserDataFilter SubItemsFilter;
+		SubItemsFilter.ItemTypeFilter = EContentBrowserItemTypeFilter::IncludeFolders;
+		SubItemsFilter.bRecursivePaths = false;
+		SubItemsFilter.ItemCategoryFilter = PathView->GetContentBrowserItemCategoryFilter();
+		SubItemsFilter.ItemAttributeFilter = PathView->GetContentBrowserItemAttributeFilter();
+
+		FString CrumbPath = TEXT("/");
+		for (const FString& Crumb : Crumbs)
+		{
+			CrumbPath += Crumb;
+
+			bool bHasSubItems = false;
+			ContentBrowserData->EnumerateItemsUnderPath(*CrumbPath, SubItemsFilter, [&bHasSubItems](FContentBrowserItemData&& InSubItem)
+			{
+				bHasSubItems = true;
+				return false;
+			});
+
+			const FContentBrowserItem CrumbFolderItem = ContentBrowserData->GetItemAtPath(*CrumbPath, EContentBrowserItemTypeFilter::IncludeFolders);
+			NavigationBar->PushPath(CrumbFolderItem.IsValid() ? CrumbFolderItem.GetDisplayName() : FText::FromString(Crumb), CrumbPath, bHasSubItems);
+
+			CrumbPath += TEXT("/");
+		}
+	}
+	else if (ContentSources.HasCollections())
+	{
+		const FCollectionRef& Collection = ContentSources.GetCollections()[0];
+
+		TArray<FCollectionNameType> CollectionPathItems;
+
+		// Walk up the parents of this collection so that we can generate a complete path (this loop also adds the child collection to the array)
+		for (TOptional<FCollectionNameType> CurrentCollection = FCollectionNameType(Collection.Name, Collection.Type);
+			CurrentCollection.IsSet(); 
+			CurrentCollection = Collection.Container->GetParentCollection(CurrentCollection->Name, CurrentCollection->Type))
+		{
+			CollectionPathItems.Insert(CurrentCollection.GetValue(), 0);
+		}
+
+		// Now add each part of the path to the breadcrumb trail
+		for (const FCollectionNameType& CollectionPathItem : CollectionPathItems)
+		{
+			const FString CrumbData = FormatCollectionCrumbData(*Collection.Container, CollectionPathItem);
+			
+			TArray<FCollectionNameType> ChildCollections;
+			Collection.Container->GetChildCollections(CollectionPathItem.Name, CollectionPathItem.Type, ChildCollections);
+			const bool bHasChildren = ChildCollections.Num() > 0;
+
+			FFormatNamedArguments Args;
+			Args.Add(TEXT("CollectionName"), FText::FromName(CollectionPathItem.Name));
+			const FText DisplayName = FText::Format(LOCTEXT("CollectionPathIndicator", "{CollectionName} (Collection)"), Args);
+
+			NavigationBar->PushPath(DisplayName, CrumbData, bHasChildren);
+		}
+	}
+	else
+	{
+		NavigationBar->PushPath(LOCTEXT("AllAssets", "All Assets"), TEXT(""), true); 
+	}
+}
+
+TArray<FContentBrowserItem> ContentBrowserUtils::FilterOrAliasItems(TArrayView<const FContentBrowserItem> Items)
+{
+	UContentBrowserDataSubsystem* ContentBrowserData = IContentBrowserDataModule::Get().GetSubsystem();
+	const TSharedRef<FPathPermissionList>& FolderPermissions = FAssetToolsModule::GetModule().Get().GetFolderPermissionList();
+
+	// If any of the items to sync don't pass the permission filter, try to find an alias that might be visible
+	TArray<FContentBrowserItem> NewItems;
+	for (const FContentBrowserItem& Item : Items)
+	{
+		if (FolderPermissions->PassesStartsWithFilter(Item.GetInternalPath()))
+		{
+			NewItems.Add(Item);
+		}
+		else
+		{
+			TArray<FContentBrowserItemPath> Aliases = ContentBrowserData->GetAliasesForPath(Item.GetInternalPath());
+			for (const FContentBrowserItemPath& Alias : Aliases)
+			{
+				if (FolderPermissions->PassesStartsWithFilter(Alias.GetInternalPathName()))
+				{
+					FContentBrowserItem AliasedItem = ContentBrowserData->GetItemAtPath(Alias.GetVirtualPathName(), EContentBrowserItemTypeFilter::IncludeFiles);
+					if (AliasedItem.IsValid())
+					{
+						NewItems.Add(MoveTemp(AliasedItem));
+						break;
+					}
+				}
+			}
+		}
+	}
+	
+	return NewItems;
+}
+
+const FContentBrowserInstanceConfig* ContentBrowserUtils::GetConstInstanceConfig(const FName& ForInstance) 
+{
+	if (ForInstance.IsNone())
+	{
+		return nullptr;
+	}
+
+	UContentBrowserConfig* Config = UContentBrowserConfig::Get();
+	if (Config == nullptr)
+	{
+		return nullptr;
+	}
+
+	const FContentBrowserInstanceConfig* InstanceConfig = Config->Instances.Find(ForInstance);
+	return InstanceConfig;
+}
+
+TArray<FContentBrowserItem> ContentBrowserUtils::GetChildItemsFromVirtualPath(
+	const FName& Path,
+	EContentBrowserItemCategoryFilter ItemCategoryFilter, 
+	EContentBrowserItemAttributeFilter ItemAttributeFilter,
+	const FName& ConfigInstanceName,
+	const SPathView& PathViewForFiltering 
+)
+{
+	const UContentBrowserSettings* ContentBrowserSettings = GetDefault<UContentBrowserSettings>();
+	bool bDisplayEmpty = ContentBrowserSettings->DisplayEmptyFolders;
+	// check to see if we have an instance config that overrides the default in UContentBrowserSettings
+	if (const FContentBrowserInstanceConfig* EditorConfig = GetConstInstanceConfig(ConfigInstanceName))
+	{
+		bDisplayEmpty = EditorConfig->bShowEmptyFolders;
+	}
+	
+	EContentBrowserIsFolderVisibleFlags FolderFlags = ContentBrowserUtils::GetIsFolderVisibleFlags(bDisplayEmpty);
+
+	TOptional<FContentBrowserFolderContentsFilter> FolderFilter;
+	if (!bDisplayEmpty)
+	{
+		FolderFilter = FContentBrowserFolderContentsFilter{};
+		FolderFilter->ItemCategoryFilter = ItemCategoryFilter;
+	}
+
+	UContentBrowserDataSubsystem* ContentBrowserData = IContentBrowserDataModule::Get().GetSubsystem();
+
+	FContentBrowserDataFilter SubItemsFilter;
+	SubItemsFilter.ItemTypeFilter = EContentBrowserItemTypeFilter::IncludeFolders;
+	SubItemsFilter.bRecursivePaths = false;
+	SubItemsFilter.ItemCategoryFilter = ItemCategoryFilter;
+	SubItemsFilter.ItemAttributeFilter = ItemAttributeFilter;
+
+	TArray<FContentBrowserItem> SubItems = ContentBrowserData->GetItemsUnderPath(Path, SubItemsFilter);
+
+	for (auto It = SubItems.CreateIterator(); It; ++It)
+	{
+		const FContentBrowserItem& Item = *It;
+		if (!Item.GetInternalPath().IsNone())
+		{
+			if (!PathViewForFiltering.InternalPathPassesBlockLists(FNameBuilder(Item.GetInternalPath())))
+			{
+				It.RemoveCurrent();
+				continue;
+			}
+		}
+		else
+		{
+			// Test if any child internal paths pass for this fully virtual path
+			bool bPasses = false;
+			for (const FContentBrowserItemData& ItemData : Item.GetInternalItems())
+			{
+				UContentBrowserDataSource* ItemDataSource = ItemData.GetOwnerDataSource();
+				if (!ItemDataSource) 
+				{
+					continue;
+				}
+
+				auto OnSubPath = [&PathViewForFiltering, &bPasses, &SubItemsFilter](FName VirtualSubPath, FName InternalPath)
+				{
+					if (InternalPath.IsNone())
+					{
+						// Keep enumerating, this path only exist virtually e.g. /All/Plugins
+						return true;
+					}
+
+					const FNameBuilder InternalPathBuilder(InternalPath);
+					if (ContentBrowserDataUtils::PathPassesAttributeFilter(InternalPathBuilder, 0, SubItemsFilter.ItemAttributeFilter) 
+					&& PathViewForFiltering.InternalPathPassesBlockLists(InternalPathBuilder))
+					{
+						bPasses = true;
+						// Stop enumerating
+						return false;
+					}
+					return true;
+				}; 
+				ItemDataSource->GetRootPathVirtualTree().EnumerateSubPaths(Item.GetVirtualPath(), OnSubPath, /*bRecurse*/ true);
+
+				if (bPasses)
+				{
+					break;
+				}
+			}
+
+			if (!bPasses)
+			{
+				It.RemoveCurrent();
+				continue;
+			}
+		}
+
+		if (!ContentBrowserData->IsFolderVisible(Item.GetVirtualPath(), FolderFlags, FolderFilter))
+		{
+			It.RemoveCurrent();
+			continue;
+		}
+	}
+
+	return SubItems;
+}
+
+TSharedPtr<SWidget> ContentBrowserUtils::GetFolderWidgetForNavigationBar(const FText& InFolderName, const FName& InFolderBrushName, const FLinearColor& InFolderColor)
+{
+	const FSlateBrush* FolderBrush = FAppStyle::GetBrush(InFolderBrushName);
+
+	if (FolderBrush != FAppStyle::GetDefaultBrush())
+	{
+		return SNew(SHorizontalBox)
+			+ SHorizontalBox::Slot()
+			.AutoWidth()
+			.Padding(FMargin(2, 0, 6, 0))
+			[
+				SNew(SBox)
+				.WidthOverride(16.f)
+				.HeightOverride(16.f)
+				.HAlign(HAlign_Center)
+				.VAlign(VAlign_Center)
+				[
+					SNew(SImage)
+					.DesiredSizeOverride(FVector2D(16, 16))
+					.Image(FolderBrush)
+					.ColorAndOpacity(InFolderColor)
+				]
+			]
+
+			+ SHorizontalBox::Slot()
+			.FillWidth(1.f)
+			.Padding(FMargin(2, 0, 6, 0))
+			.VAlign( VAlign_Center )
+			[
+				SNew(STextBlock)
+				.TextStyle(FAppStyle::Get(), TEXT("Menu.Label"))
+				.Text(InFolderName)
+			];
+	}
+	return nullptr;
+}
+
 #undef LOCTEXT_NAMESPACE
diff --git a/Engine/Source/Editor/ContentBrowser/Private/ContentBrowserUtils.h b/Engine/Source/Editor/ContentBrowser/Private/ContentBrowserUtils.h
--- a/Engine/Source/Editor/ContentBrowser/Private/ContentBrowserUtils.h
+++ b/Engine/Source/Editor/ContentBrowser/Private/ContentBrowserUtils.h
@@ -22,6 +22,7 @@ class SAssetView;
 class SFilterList;
 class SPathView;
 class SWidget;
+class SNavigationBar;
 struct FARFilter;
 struct FAssetData;
 struct FContentBrowserDataFilter;
@@ -31,6 +32,7 @@ struct FContentBrowserItemPath;
 struct FPathViewConfig;
 enum class EContentBrowserIsFolderVisibleFlags : uint8;
 enum class EContentBrowserItemAttributeFilter : uint8;
+enum class EContentBrowserItemCategoryFilter : uint8;
 
 namespace ContentBrowserUtils
 {
@@ -196,4 +198,21 @@ namespace ContentBrowserUtils
 	 *  - A filesystem path to a uasset or umap file.
 	 */
 	FContentBrowserItem TryGetItemFromUserProvidedPath(FStringView RequestedPathView);
+
+	FString FormatCollectionCrumbData(const ICollectionContainer& CollectionContainer, const FCollectionNameType& Collection);
+	
+	/** Updates a navigation bar based on the current state of an asset view and path view */
+	void UpdateNavigationBar(const TSharedPtr<SNavigationBar>& NavigationBar, const TSharedPtr<SAssetView>& AssetView, const TSharedPtr<SPathView>& PathView);
+
+	/** Filter a list of items based on permissions or find a valid alias. */
+	TArray<FContentBrowserItem> FilterOrAliasItems(TArrayView<const FContentBrowserItem> Items);
+
+	/** Retrieves a named config. */
+	const FContentBrowserInstanceConfig* GetConstInstanceConfig(const FName& ForInstance);
+
+	/** Finds the content items that are children of the supplied path, while obeying the supplied filters. */
+	TArray<FContentBrowserItem> GetChildItemsFromVirtualPath(const FName& Path, EContentBrowserItemCategoryFilter ItemCategoryFilter, EContentBrowserItemAttributeFilter ItemAttributeFilter, const FName& ConfigInstanceName, const SPathView& PathViewForFiltering);
+
+	/** Creates a widget that represents a folder with the given name, brush, and color. */
+	TSharedPtr<SWidget> GetFolderWidgetForNavigationBar(const FText& InFolderName, const FName& InFolderBrushName, const FLinearColor& InFolderColor);
 }
diff --git a/Engine/Source/Editor/ContentBrowser/Private/Experimental/ContentBrowserExtensionUtils.cpp b/Engine/Source/Editor/ContentBrowser/Private/Experimental/ContentBrowserExtensionUtils.cpp
--- a/Engine/Source/Editor/ContentBrowser/Private/Experimental/ContentBrowserExtensionUtils.cpp
+++ b/Engine/Source/Editor/ContentBrowser/Private/Experimental/ContentBrowserExtensionUtils.cpp
@@ -9,27 +9,27 @@
 
 namespace UE::Editor::ContentBrowser::ExtensionUtils
 {
-	TOptional<FLinearColor> GetFolderColor(const FName& FolderPath)
+	TOptional<FLinearColor> GetFolderColor(const FContentBrowserItemPath& FolderPath)
 	{
-		const FName VirtualPath = IContentBrowserDataModule::Get().GetSubsystem()->ConvertInternalPathToVirtual(FolderPath);
-
 		TSharedPtr<ICollectionContainer> CollectionContainer;
 		FName CollectionName;
 		ECollectionShareType::Type CollectionFolderShareType = ECollectionShareType::CST_All;
 
-		if(ContentBrowserUtils::IsCollectionPath(VirtualPath.ToString(), &CollectionContainer , &CollectionName, &CollectionFolderShareType))
+		if (ContentBrowserUtils::IsCollectionPath(FolderPath.GetVirtualPathString(), &CollectionContainer , &CollectionName, &CollectionFolderShareType))
 		{
-			if(TOptional<FLinearColor> Color = CollectionViewUtils::GetCustomColor(CollectionContainer.Get(), CollectionName, CollectionFolderShareType))
+			if (TOptional<FLinearColor> Color = CollectionViewUtils::GetCustomColor(CollectionContainer.Get(), CollectionName, CollectionFolderShareType))
 			{
 				return Color;
 			}
 		}
-		if (TOptional<FLinearColor> Color = ContentBrowserUtils::GetPathColor(FolderPath.ToString()))
-		{
-			return Color.GetValue();
-		}
 
-		return TOptional<FLinearColor>();
+		return ContentBrowserUtils::GetPathColor(FolderPath.GetInternalPathString());
+	}
+	
+	TOptional<FLinearColor> GetFolderColor(const FName& FolderPath)
+	{
+		const FContentBrowserItemPath ItemPath(FolderPath, EContentBrowserPathType::Internal);
+		return GetFolderColor(ItemPath);
 	}
 
 	void SetFolderColor(const FName& FolderPath, const FLinearColor& FolderColor)
diff --git a/Engine/Source/Editor/ContentBrowser/Private/SAssetDialog.cpp b/Engine/Source/Editor/ContentBrowser/Private/SAssetDialog.cpp
--- a/Engine/Source/Editor/ContentBrowser/Private/SAssetDialog.cpp
+++ b/Engine/Source/Editor/ContentBrowser/Private/SAssetDialog.cpp
@@ -21,6 +21,7 @@
 #include "CoreGlobals.h"
 #include "Delegates/Delegate.h"
 #include "Editor.h"
+#include "ICollectionSource.h"
 #include "Editor/EditorEngine.h"
 #include "Framework/Application/SlateApplication.h"
 #include "Framework/Commands/GenericCommands.h"
@@ -51,7 +52,9 @@
 #include "SPrimaryButton.h"
 #include "SlateOptMacros.h"
 #include "SlotBase.h"
+#include "SNavigationBar.h"
 #include "SWarningOrErrorBox.h"
+#include "Experimental/ContentBrowserExtensionUtils.h"
 #include "Styling/AppStyle.h"
 #include "Templates/Tuple.h"
 #include "Templates/UnrealTemplate.h"
@@ -104,6 +107,11 @@ void SAssetDialog::Construct(const FArguments& InArgs, const FSharedAssetDialogC
 
 	AssetClassNames = InConfig.AssetClassNames;
 
+	HistoryManager.SetOnApplyHistoryData(FOnApplyHistoryData::CreateSP(this, &SAssetDialog::OnApplyHistoryData));
+	HistoryManager.SetOnUpdateHistoryData(FOnUpdateHistoryData::CreateSP(this, &SAssetDialog::OnUpdateHistoryData));
+
+	RecentDirectories.MaxItems = 30;
+
 	const FString DefaultPath = InConfig.DefaultPath;
 
 	RegisterActiveTimer( 0.f, FWidgetActiveTimerDelegate::CreateSP( this, &SAssetDialog::SetFocusPostConstruct ) );
@@ -111,7 +119,7 @@ void SAssetDialog::Construct(const FArguments& InArgs, const FSharedAssetDialogC
 	FPathPickerConfig PathPickerConfig;
 	PathPickerConfig.DefaultPath = DefaultPath;
 	PathPickerConfig.bFocusSearchBoxWhenOpened = false;
-	PathPickerConfig.OnPathSelected = FOnPathSelected::CreateSP(this, &SAssetDialog::HandlePathSelected);
+	PathPickerConfig.OnPathSelected = FOnPathSelected::CreateSP(this, &SAssetDialog::SetCurrentlySelectedPath, EContentBrowserPathType::Virtual);
 	PathPickerConfig.SetPathsDelegates.Add(&SetPathsDelegate);
 	PathPickerConfig.OnGetFolderContextMenu = FOnGetFolderContextMenu::CreateSP(this, &SAssetDialog::OnGetFolderContextMenu);	
 	PathPickerConfig.bOnPathSelectedPassesVirtualPaths = true;
@@ -127,14 +135,12 @@ void SAssetDialog::Construct(const FArguments& InArgs, const FSharedAssetDialogC
 	AssetPickerConfig.SaveSettingsName = TEXT("AssetDialog");
 	AssetPickerConfig.bCanShowFolders = true;
 	AssetPickerConfig.bCanShowDevelopersFolder = true;
-	AssetPickerConfig.OnFolderEntered = FOnPathSelected::CreateSP(this, &SAssetDialog::HandleAssetViewFolderEntered);
+	AssetPickerConfig.OnFolderEntered = FOnPathSelected::CreateSP(this, &SAssetDialog::SetCurrentlySelectedPath, EContentBrowserPathType::Virtual);
 	AssetPickerConfig.OnGetAssetContextMenu = FOnGetAssetContextMenu::CreateSP(this, &SAssetDialog::OnGetAssetContextMenu);
 	AssetPickerConfig.OnGetFolderContextMenu = FOnGetFolderContextMenu::CreateSP(this, &SAssetDialog::OnGetFolderContextMenu);
 
 	OnPathSelected = InConfig.OnPathSelected;
 
-	SetCurrentlySelectedPath(DefaultPath, EContentBrowserPathType::Internal);
-
 	// Open and save specific configuration
 	FText ConfirmButtonText;
 	bool bIncludeNameBox = false;
@@ -211,22 +217,73 @@ void SAssetDialog::Construct(const FArguments& InArgs, const FSharedAssetDialogC
 
 	AssetPicker = StaticCastSharedRef<SAssetPicker>(FContentBrowserSingleton::Get().CreateAssetPicker(AssetPickerConfig));
 
-	// Update AssetDialog's path to match PathPicker
-	if (SelectedVirtualPaths.Num() > 0)
-	{
-		const FName DefaultVirtualPath = IContentBrowserDataModule::Get().GetSubsystem()->ConvertInternalPathToVirtual(*DefaultPath);
-		if (DefaultVirtualPath != *SelectedVirtualPaths[0])
-		{
-			SetCurrentlySelectedPath(SelectedVirtualPaths[0], EContentBrowserPathType::Virtual);
-		}
-	}
-
 	FContentBrowserCommands::Register();
 	BindCommands();
 
 	// The root widget in this dialog.
 	TSharedRef<SVerticalBox> MainVerticalBox = SNew(SVerticalBox);
 
+	// Navigation & History
+	MainVerticalBox->AddSlot()
+		.AutoHeight()
+		[
+			SNew(SBorder)
+			.BorderImage(FAppStyle::GetBrush("ToolPanel.GroupBorder"))
+			[
+				SNew(SHorizontalBox)
+				// History Back Button
+				+SHorizontalBox::Slot()
+				.Padding(10, 0, 0, 0)
+				.AutoWidth()
+				[
+					SNew(SButton)
+					.VAlign(VAlign_Center)
+					.ButtonStyle(FAppStyle::Get(), "SimpleButton")
+					.ToolTipText(this, &SAssetDialog::GetBackTooltip )
+					.ContentPadding( FMargin(1, 0) )
+					.OnClicked(this, &SAssetDialog::OnBackClicked)
+					.IsEnabled(this, &SAssetDialog::IsBackEnabled)
+					[
+						SNew(SImage)
+						.Image(FAppStyle::Get().GetBrush("Icons.CircleArrowLeft"))
+						.ColorAndOpacity(FSlateColor::UseForeground())
+					]
+				]
+				// History Forward Button
+				+ SHorizontalBox::Slot()
+				.Padding(2, 0, 0, 0)
+				.AutoWidth()
+				[
+					SNew(SButton)
+					.VAlign(VAlign_Center)
+					.ButtonStyle(FAppStyle::Get(), "SimpleButton")
+					.ToolTipText( this, &SAssetDialog::GetForwardTooltip )
+					.ContentPadding( FMargin(1, 0) )
+					.OnClicked(this, &SAssetDialog::OnForwardClicked)
+					.IsEnabled(this, &SAssetDialog::IsForwardEnabled)
+					[
+						SNew(SImage)
+						.Image(FAppStyle::Get().GetBrush("Icons.CircleArrowRight"))
+						.ColorAndOpacity(FSlateColor::UseForeground())
+					]
+				]
+				+ SHorizontalBox::Slot()
+				.VAlign(VAlign_Center)
+				.HAlign(HAlign_Fill)
+				.FillWidth(1.0f)
+				.Padding(2, 0, 0, 0)
+				[
+					SAssignNew(NavigationBar, SNavigationBar)
+					.OnPathClicked(this, &SAssetDialog::SetCurrentlySelectedPath, EContentBrowserPathType::Virtual)
+					.OnNavigateToPath(this, &SAssetDialog::OnPathTextEdited)
+					.OnCanEditPathAsText(this, &SAssetDialog::OnCanEditPathAsText)
+					.OnCompletePrefix(this, &SAssetDialog::OnCompletePathPrefix)
+					.GetPathMenuContent(this, &SAssetDialog::OnGetCrumbDelimiterContent)
+					.GetComboOptions(this, &SAssetDialog::GetRecentPaths)
+				]
+			]
+		];
+	
 	// Path/Asset view
 	MainVerticalBox->AddSlot()
 		.FillHeight(1)
@@ -363,6 +420,16 @@ void SAssetDialog::Construct(const FArguments& InArgs, const FSharedAssetDialogC
 	[
 		MainVerticalBox
 	];
+
+	// Sync path selection with the path picker
+	if (SelectedVirtualPaths.Num() > 0)
+	{
+		SetCurrentlySelectedPath(SelectedVirtualPaths[0], EContentBrowserPathType::Virtual);
+	}
+	else
+	{
+		SetCurrentlySelectedPath(DefaultPath, EContentBrowserPathType::Internal);
+	}
 }
 END_SLATE_FUNCTION_BUILD_OPTIMIZATION
 
@@ -381,6 +448,40 @@ FReply SAssetDialog::OnKeyDown( const FGeometry& MyGeometry, const FKeyEvent& In
 	return SCompoundWidget::OnKeyDown(MyGeometry, InKeyEvent);
 }
 
+FReply SAssetDialog::OnMouseButtonDown(const FGeometry& MyGeometry, const FPointerEvent& MouseEvent)
+{
+	// Mouse back and forward buttons traverse history
+	if (MouseEvent.GetEffectingButton() == EKeys::ThumbMouseButton)
+	{
+		HistoryManager.GoBack();
+		return FReply::Handled();
+	}
+	if (MouseEvent.GetEffectingButton() == EKeys::ThumbMouseButton2)
+	{
+		HistoryManager.GoForward();
+		return FReply::Handled();
+	}
+
+	return FReply::Unhandled();
+}
+
+FReply SAssetDialog::OnMouseButtonDoubleClick(const FGeometry& InMyGeometry, const FPointerEvent& InMouseEvent)
+{
+	// Mouse back and forward buttons traverse history
+	if (InMouseEvent.GetEffectingButton() == EKeys::ThumbMouseButton)
+	{
+		HistoryManager.GoBack();
+		return FReply::Handled();
+	}
+	if (InMouseEvent.GetEffectingButton() == EKeys::ThumbMouseButton2)
+	{
+		HistoryManager.GoForward();
+		return FReply::Handled();
+	}
+
+	return FReply::Unhandled();
+}
+
 void SAssetDialog::BindCommands()
 {
 	Commands = TSharedPtr< FUICommandList >(new FUICommandList);
@@ -753,6 +854,71 @@ void SAssetDialog::FocusNameBox()
 	}
 }
 
+bool SAssetDialog::IsBackEnabled() const
+{
+	return HistoryManager.CanGoBack();
+}
+
+bool SAssetDialog::IsForwardEnabled() const
+{
+	return HistoryManager.CanGoForward();
+}
+
+FText SAssetDialog::GetBackTooltip() const
+{
+	if (HistoryManager.CanGoBack())
+	{
+		return FText::Format(LOCTEXT("HistoryBackTooltipFmt", "Back to {0}"), HistoryManager.GetBackDesc());
+	}
+	return FText::GetEmpty();
+}
+
+FText SAssetDialog::GetForwardTooltip() const
+{
+	if (HistoryManager.CanGoForward())
+	{
+		return FText::Format(LOCTEXT("HistoryForwardTooltipFmt", "Forward to {0}"), HistoryManager.GetForwardDesc());
+	}
+	return FText::GetEmpty();
+}
+
+FReply SAssetDialog::OnBackClicked()
+{
+	HistoryManager.GoBack();
+	return FReply::Handled();
+}
+
+FReply SAssetDialog::OnForwardClicked()
+{
+	HistoryManager.GoForward();
+	return FReply::Handled();
+}
+
+void SAssetDialog::OnApplyHistoryData(const FHistoryData& History)
+{
+	PathPicker->GetPathView()->ApplyHistoryData(History);
+	AssetPicker->GetAssetView()->ApplyHistoryData(History);
+
+	ContentBrowserUtils::UpdateNavigationBar(NavigationBar.ToSharedRef(), AssetPicker->GetAssetView().ToSharedRef(), PathPicker->GetPathView().ToSharedRef());
+}
+
+void SAssetDialog::OnUpdateHistoryData(FHistoryData& HistoryData) const
+{
+	const FContentSources& ContentSources = AssetPicker->GetAssetView()->GetContentSources();
+	const TArray<FContentBrowserItem> SelectedItems = AssetPicker->GetAssetView()->GetSelectedItems();
+
+	const FText NewSource = ContentSources.HasVirtualPaths() ? FText::FromName(ContentSources.GetVirtualPaths()[0]) : (ContentSources.HasCollections() ? FText::FromName(ContentSources.GetCollections()[0].Name) : LOCTEXT("AllAssets", "All Assets"));
+
+	HistoryData.HistoryDesc = NewSource;
+	HistoryData.ContentSources = ContentSources;
+
+	HistoryData.SelectionData.Reset();
+	for (const FContentBrowserItem& SelectedItem : SelectedItems)
+	{
+		HistoryData.SelectionData.SelectedVirtualPaths.Add(SelectedItem.GetVirtualPath());
+	}
+}
+
 FText SAssetDialog::GetAssetNameText() const
 {
 	return FText::FromString(CurrentlyEnteredAssetName);
@@ -788,28 +954,184 @@ FText SAssetDialog::GetNameErrorLabelText() const
 	return FText::GetEmpty();
 }
 
-void SAssetDialog::HandlePathSelected(const FString& NewPath)
+bool SAssetDialog::OnCanEditPathAsText(const FString& Text) const
+{
+	const FContentSources& ContentSources = AssetPicker->GetAssetView()->GetContentSources();
+	if (ContentSources.HasCollections())
+	{
+		return false;
+	}
+	return true;
+}
+
+void SAssetDialog::OnPathTextEdited(const FString& NewPath)
+{
+	FContentBrowserItem Item = ContentBrowserUtils::TryGetItemFromUserProvidedPath(NewPath);
+	if (Item.IsValid())
+	{
+		SyncToItems({ Item });
+	}
+}
+
+TArray<FString> SAssetDialog::OnCompletePathPrefix(const FString& Prefix) const
+{
+	FStringView PrefixView = Prefix;
+	
+	// Strip to last path separator
+	FName Parent;
+	if (int32 Index = UE::String::FindLastChar(PrefixView, '/'); Index != INDEX_NONE)
+	{	
+		PrefixView.LeftInline(Index);
+		Parent = FName(PrefixView);
+	}
+
+	// Find PrefixView in the available tree of data sources, get its direct children, and filter them by SuffixView
+	TArray<FContentBrowserItem> SubItems = ContentBrowserUtils::GetChildItemsFromVirtualPath(
+		Parent,
+		PathPicker->GetPathView()->GetContentBrowserItemCategoryFilter(),
+		PathPicker->GetPathView()->GetContentBrowserItemAttributeFilter(),
+		NAME_None,
+		*PathPicker->GetPathView()
+	);
+	
+	TArray<FString> Results;
+	for (const FContentBrowserItem& Item : SubItems)
+	{
+		FName Path = Item.GetVirtualPath();
+		FNameBuilder PathBuilder(Path);
+		if (PathBuilder.ToView().StartsWith(Prefix))
+		{
+			Results.Add(Item.GetVirtualPath().ToString());
+		}
+	}
+	return Results;
+}
+
+TSharedRef<SWidget> SAssetDialog::OnGetCrumbDelimiterContent(const FString& CrumbData) const
 {
-	FName ConvertedPath;
-	const EContentBrowserPathType ConvertedPathType = IContentBrowserDataModule::Get().GetSubsystem()->TryConvertVirtualPath(NewPath, ConvertedPath);
+	TArray<FContentBrowserItem> SubItems = ContentBrowserUtils::GetChildItemsFromVirtualPath(
+		*CrumbData,
+		PathPicker->GetPathView()->GetContentBrowserItemCategoryFilter(),
+		PathPicker->GetPathView()->GetContentBrowserItemAttributeFilter(),
+		NAME_None,
+		*PathPicker->GetPathView()
+	);
+
+	SubItems.Sort([](const FContentBrowserItem& ItemOne, const FContentBrowserItem& ItemTwo)
+	{
+		return ItemOne.GetDisplayName().CompareTo(ItemTwo.GetDisplayName()) < 0;
+	});
 
-	FARFilter NewFilter;
+	if (SubItems.Num() > 0)
+	{
+		FMenuBuilder MenuBuilder( true, nullptr );
 
-	NewFilter.ClassPaths.Append(AssetClassNames);
-	NewFilter.PackagePaths.Add(*NewPath);
+		for (const FContentBrowserItem& SubItem : SubItems)
+		{
+			FName FolderBrushName = NAME_None;
+			FName FolderShadowBrushName = NAME_None;
+			ContentBrowserUtils::TryGetFolderBrushAndShadowNameSmall(SubItem, FolderBrushName, FolderShadowBrushName);
 
-	SetCurrentlySelectedPath(ConvertedPath.ToString(), ConvertedPathType);
+			FText EntryName = SubItem.GetDisplayName();
+			FUIAction EntryAction = FUIAction(FExecuteAction::CreateSP(const_cast<SAssetDialog*>(this), &SAssetDialog::OnCrumbDelimiterItemClicked, SubItem.GetVirtualPath().ToString()));
 
-	SetFilterDelegate.ExecuteIfBound(NewFilter);
+			if (FolderBrushName != NAME_None)
+			{
+				const FContentBrowserItemPath FolderPath(SubItem.GetVirtualPath(), SubItem.GetInternalPath());
+				FLinearColor FolderColor = UE::Editor::ContentBrowser::ExtensionUtils::GetFolderColor(FolderPath).Get(ContentBrowserUtils::GetDefaultColor());
+				
+				FMenuEntryParams Params;
+				Params.EntryWidget = ContentBrowserUtils::GetFolderWidgetForNavigationBar(EntryName, FolderBrushName, FolderColor);
+				Params.DirectActions = EntryAction;
+				MenuBuilder.AddMenuEntry(Params);
+			}
+			else
+			{
+				MenuBuilder.AddMenuEntry(
+					EntryName,
+					FText::GetEmpty(),
+					FSlateIcon(FAppStyle::GetAppStyleSetName(), FolderBrushName),
+					EntryAction
+				);
+			}
+		}
+
+		return SNew(SVerticalBox)
+			+SVerticalBox::Slot()
+			.MaxHeight(400.0f)
+			[
+				MenuBuilder.MakeWidget()
+			];
+	}
+
+	return SNullWidget::NullWidget;
 }
 
-void SAssetDialog::HandleAssetViewFolderEntered(const FString& NewPath)
+void SAssetDialog::OnCrumbDelimiterItemClicked(FString ClickedPath)
 {
-	SetCurrentlySelectedPath(NewPath, EContentBrowserPathType::Virtual);
+	SetCurrentlySelectedPath(ClickedPath, EContentBrowserPathType::Virtual);
+}
 
-	TArray<FString> NewPaths;
-	NewPaths.Add(NewPath);
-	SetPathsDelegate.Execute(NewPaths);
+TArray<FString> SAssetDialog::GetRecentPaths() const
+{
+	return RecentDirectories;
+}
+
+void SAssetDialog::SetCurrentlySelectedPath(const FString& NewPath, const EContentBrowserPathType InPathType)
+{
+	CurrentlySelectedPath = NewPath;
+	CurrentlySelectedPathType = InPathType;
+
+	const FName NewVirtualPath = InPathType == EContentBrowserPathType::Virtual ? FName(NewPath) : IContentBrowserDataModule::Get().GetSubsystem()->ConvertInternalPathToVirtual(*NewPath);
+	
+	// Update Path View
+	if (PathPicker)
+	{
+		TArray<FString> SelectedVirtualPaths = PathPicker->GetPaths();
+		if (SelectedVirtualPaths.Num() == 0 || NewVirtualPath != *SelectedVirtualPaths[0])
+		{
+			SetPathsDelegate.Execute({ NewVirtualPath.ToString() });
+		}
+	}
+
+	if (AssetPicker)
+	{
+		// The asset picker uses the FARFilter::PackagePaths field to fill content sources.
+		// Thus, content sources must be checked rather than the filter itself.
+		const FContentSources& Sources = AssetPicker->GetAssetView()->GetContentSources();
+		if (!Sources.HasVirtualPaths() || Sources.GetVirtualPaths()[0] != NewVirtualPath)
+		{
+			FARFilter NewFilter;
+
+			NewFilter.ClassPaths.Append(AssetClassNames);
+			NewFilter.PackagePaths.Add(NewVirtualPath);
+
+			SetFilterDelegate.Execute(NewFilter);
+		}
+	}
+	
+	// Update Navigation Bar
+	if (NavigationBar && AssetPicker && PathPicker)
+	{
+		ContentBrowserUtils::UpdateNavigationBar(NavigationBar, AssetPicker->GetAssetView(), PathPicker->GetPathView());
+	}
+	
+	UpdateInputValidity();
+
+	HistoryManager.AddHistoryData();
+	RecentDirectories.AddUnique(InPathType == EContentBrowserPathType::Virtual ? NewPath : NewVirtualPath.ToString());
+
+	OnPathSelected.ExecuteIfBound(NewPath);
+}
+
+void SAssetDialog::SyncToItems(TArrayView<const FContentBrowserItem> ItemsToSync, const bool bAllowImplicitSync)
+{
+	TArray<FContentBrowserItem> NewItemsToSync = ContentBrowserUtils::FilterOrAliasItems(ItemsToSync);
+	ItemsToSync = NewItemsToSync;
+	
+	// Tell the sources view first so the asset view will be up to date by the time we request the sync
+	PathPicker->SyncToItems(ItemsToSync, bAllowImplicitSync);
+	AssetPicker->SyncToItems(ItemsToSync, bAllowImplicitSync, true);
 }
 
 bool SAssetDialog::IsConfirmButtonEnabled() const
@@ -907,16 +1229,6 @@ void SAssetDialog::CloseDialog()
 	}
 }
 
-void SAssetDialog::SetCurrentlySelectedPath(const FString& NewPath, const EContentBrowserPathType InPathType)
-{
-	CurrentlySelectedPath = NewPath;
-	CurrentlySelectedPathType = InPathType;
-
-	UpdateInputValidity();
-
-	OnPathSelected.ExecuteIfBound(NewPath);
-}
-
 void SAssetDialog::SetCurrentlyEnteredAssetName(const FString& NewName)
 {
 	CurrentlyEnteredAssetName = NewName;
diff --git a/Engine/Source/Editor/ContentBrowser/Private/SAssetDialog.h b/Engine/Source/Editor/ContentBrowser/Private/SAssetDialog.h
--- a/Engine/Source/Editor/ContentBrowser/Private/SAssetDialog.h
+++ b/Engine/Source/Editor/ContentBrowser/Private/SAssetDialog.h
@@ -6,9 +6,12 @@
 #include "Containers/UnrealString.h"
 #include "ContentBrowserDataSubsystem.h"
 #include "ContentBrowserDelegates.h"
+#include "HistoryManager.h"
 #include "HAL/Platform.h"
 #include "IAssetTypeActions.h"
 #include "IContentBrowserSingleton.h"
+#include "SNavigationBar.h"
+#include "Containers/MRUArray.h"
 #include "Input/Reply.h"
 #include "Internationalization/Text.h"
 #include "Layout/Visibility.h"
@@ -50,6 +53,8 @@ public:
 	virtual void Construct(const FArguments& InArgs, const FSharedAssetDialogConfig& InConfig);
 
 	virtual FReply OnKeyDown( const FGeometry& MyGeometry, const FKeyEvent& InKeyEvent ) override;
+	virtual FReply OnMouseButtonDown(const FGeometry& MyGeometry, const FPointerEvent& MouseEvent) override;
+	virtual FReply OnMouseButtonDoubleClick(const FGeometry& InMyGeometry, const FPointerEvent& InMouseEvent) override;
 
 	/** Sets the delegate handler for when an open operation is committed */
 	void SetOnAssetsChosenForOpen(const FOnAssetsChosenForOpen& InOnAssetsChosenForOpen);
@@ -71,6 +76,23 @@ private:
 	/** Moves keyboard focus to the name box if this is a save dialog */
 	void FocusNameBox();
 
+
+	// --- History Callbacks ---
+	bool IsBackEnabled() const;
+	bool IsForwardEnabled() const;
+
+	FText GetBackTooltip() const;
+	FText GetForwardTooltip() const;
+
+	FReply OnBackClicked();
+	FReply OnForwardClicked();
+
+	/** Handles when the history manager shifts to a different history state. */
+	void OnApplyHistoryData(const FHistoryData& History);
+	/** Fills in the history item with the current state. */
+	void OnUpdateHistoryData(FHistoryData& History) const;
+
+	
 	/** Gets the name to display in the asset name box */
 	FText GetAssetNameText() const;
 
@@ -86,11 +108,32 @@ private:
 	/** Gets the text to display in the name error label */
 	FText GetNameErrorLabelText() const;
 
-	/** Handler for when a path is selected in the path view */
-	void HandlePathSelected(const FString& NewPath);
+	
+	// --- Navigation Bar Behavior ---
+	/** Whether the navigation bar can be edited */
+	bool OnCanEditPathAsText(const FString& Text) const;
+
+	/** Respond to edited text from the navigation bar */
+	void OnPathTextEdited(const FString& NewPath);
+
+	/** Provides autocomplete for text-based path editing in the navigation bar */
+	TArray<FString> OnCompletePathPrefix(const FString& Prefix) const;
+	
+	/** Provide content for the navigation bar's breadcrumb menues. */
+	TSharedRef<SWidget> OnGetCrumbDelimiterContent(const FString& CrumbData) const;
+	void OnCrumbDelimiterItemClicked(FString ClickedPath);
+
+	/** Get a list of recently visited paths for the navigation bar. */
+	TArray<FString> GetRecentPaths() const;
+
+	
+	// --- Selection ---
+	/** Synchronizes all pickers & views to a new path state */
+	void SetCurrentlySelectedPath(const FString& NewPath, const EContentBrowserPathType InPathType);
+
+	/** Changes sources to show the specified items and selects them */
+	void SyncToItems(TArrayView<const FContentBrowserItem> ItemsToSync, const bool bAllowImplicitSync = false);
 
-	/** Handler for when a folder is entered in the asset view */
-	void HandleAssetViewFolderEntered(const FString& NewPath);
 
 	/** Determines if the confirm button (e.g. Open/Save) is enabled. */
 	bool IsConfirmButtonEnabled() const;
@@ -144,8 +187,6 @@ private:
 	/** Closes this dialog */
 	void CloseDialog();
 
-	void SetCurrentlySelectedPath(const FString& NewPath, const EContentBrowserPathType InPathType);
-
 	void SetCurrentlyEnteredAssetName(const FString& NewName);
 
 	FName GetCurrentSelectedVirtualPath() const;
@@ -165,6 +206,12 @@ private:
 	/** Whether this is an open or save dialog */
 	EAssetDialogType::Type DialogType;
 
+	/** Tracks the history data for the dialog */
+	FHistoryManager HistoryManager;
+
+	/** The most recently visited directories */
+	TMRUArray<FString> RecentDirectories;
+
 	/** Used to update the path view after it has been created */
 	FSetPathPickerPathsDelegate SetPathsDelegate;
 
@@ -224,6 +271,9 @@ private:
 	/** Path Picker used by the dialog */
 	TSharedPtr<SPathPicker> PathPicker;
 
+	/** Navigation Bar used by the dialog */
+	TSharedPtr<SNavigationBar> NavigationBar;
+	
 	/** Asset Picker used by the dialog */
 	TSharedPtr<SAssetPicker> AssetPicker;
 
diff --git a/Engine/Source/Editor/ContentBrowser/Private/SAssetPicker.cpp b/Engine/Source/Editor/ContentBrowser/Private/SAssetPicker.cpp
--- a/Engine/Source/Editor/ContentBrowser/Private/SAssetPicker.cpp
+++ b/Engine/Source/Editor/ContentBrowser/Private/SAssetPicker.cpp
@@ -658,6 +658,28 @@ void SAssetPicker::SyncToAssets(const TArray<FAssetData>& AssetDataList)
 	AssetViewPtr->SyncToLegacy(AssetDataList, TArray<FString>());
 }
 
+void SAssetPicker::SyncToItems(TArrayView<const FContentBrowserItem> ItemsToSync, const bool bAllowImplicitSync, const bool bDisableFiltersThatHideAssets)
+{
+	if (bDisableFiltersThatHideAssets)
+	{
+		if (FilterListPtr)
+		{
+			FilterListPtr->DisableFiltersThatHideItems(ItemsToSync);
+		}
+		
+		if (SearchBoxPtr)
+		{
+			SetSearchBoxText(FText::GetEmpty());
+            SearchBoxPtr->SetText(FText::GetEmpty());
+            SearchBoxPtr->SetError(FText::GetEmpty());	
+		}
+	}
+	if (AssetViewPtr)
+	{
+		AssetViewPtr->SyncToItems(ItemsToSync, bAllowImplicitSync);
+	}
+}
+
 TArray< FAssetData > SAssetPicker::GetCurrentSelection()
 {
 	return AssetViewPtr->GetSelectedAssets();
diff --git a/Engine/Source/Editor/ContentBrowser/Private/SAssetPicker.h b/Engine/Source/Editor/ContentBrowser/Private/SAssetPicker.h
--- a/Engine/Source/Editor/ContentBrowser/Private/SAssetPicker.h
+++ b/Engine/Source/Editor/ContentBrowser/Private/SAssetPicker.h
@@ -62,6 +62,9 @@ public:
 	void ExecuteRenameCommand();
 
 	TSharedPtr<SWidget> GetSearchBox() const;
+	
+	/** Syncs the AssetPicker to select the provided items. */
+	void SyncToItems(TArrayView<const FContentBrowserItem> ItemsToSync, const bool bAllowImplicitSync = false, const bool bDisableFiltersThatHideAssets = true);
 
 private:
 	/** Focuses the search box post-construct */
diff --git a/Engine/Source/Editor/ContentBrowser/Private/SContentBrowser.cpp b/Engine/Source/Editor/ContentBrowser/Private/SContentBrowser.cpp
--- a/Engine/Source/Editor/ContentBrowser/Private/SContentBrowser.cpp
+++ b/Engine/Source/Editor/ContentBrowser/Private/SContentBrowser.cpp
@@ -133,6 +133,7 @@
 #include "ToolMenuMisc.h"
 #include "ToolMenuSection.h"
 #include "ToolMenus.h"
+#include "Experimental/ContentBrowserExtensionUtils.h"
 #include "Toolkits/GlobalEditorCommonCommands.h"
 #include "Trace/Detail/Channel.h"
 #include "Types/ISlateMetaData.h"
@@ -173,195 +174,6 @@ const FString SContentBrowser::SettingsIniSection = TEXT("ContentBrowser");
 
 using namespace UE::Editor::ContentBrowser::Private;
 
-namespace ContentBrowser
-{
-	const FContentBrowserInstanceConfig* GetConstInstanceConfig(FName ForInstance) 
-	{
-		if (ForInstance.IsNone())
-		{
-			return nullptr;
-		}
-
-		UContentBrowserConfig* Config = UContentBrowserConfig::Get();
-		if (Config == nullptr)
-		{
-			return nullptr;
-		}
-
-		const FContentBrowserInstanceConfig* InstanceConfig = Config->Instances.Find(ForInstance);
-		return InstanceConfig;
-	}
-
-	// Returns child items that can be navigated to from Path e.g. subfolders
-	TArray<FContentBrowserItem> GetChildItemsFromVirtualPath(
-		FName Path,
-		EContentBrowserItemCategoryFilter ItemCategoryFilter, 
-		EContentBrowserItemAttributeFilter ItemAttributeFilter,
-		FName ConfigInstanceName,
-		SPathView& PathViewForFiltering 
-		)
-	{
-		const UContentBrowserSettings* ContentBrowserSettings = GetDefault<UContentBrowserSettings>();
-		bool bDisplayEmpty = ContentBrowserSettings->DisplayEmptyFolders;
-		// check to see if we have an instance config that overrides the default in UContentBrowserSettings
-		if (const FContentBrowserInstanceConfig* EditorConfig = GetConstInstanceConfig(ConfigInstanceName))
-		{
-			bDisplayEmpty = EditorConfig->bShowEmptyFolders;
-		}
-
-		EContentBrowserIsFolderVisibleFlags FolderFlags = ContentBrowserUtils::GetIsFolderVisibleFlags(bDisplayEmpty);
-
-		TOptional<FContentBrowserFolderContentsFilter> FolderFilter;
-		if (!bDisplayEmpty)
-		{
-			FolderFilter = FContentBrowserFolderContentsFilter{};
-			FolderFilter->ItemCategoryFilter = ItemCategoryFilter;
-		}
-
-		UContentBrowserDataSubsystem* ContentBrowserData = IContentBrowserDataModule::Get().GetSubsystem();
-
-		FContentBrowserDataFilter SubItemsFilter;
-		SubItemsFilter.ItemTypeFilter = EContentBrowserItemTypeFilter::IncludeFolders;
-		SubItemsFilter.bRecursivePaths = false;
-		SubItemsFilter.ItemCategoryFilter = ItemCategoryFilter;
-		SubItemsFilter.ItemAttributeFilter = ItemAttributeFilter;
-
-		TArray<FContentBrowserItem> SubItems = ContentBrowserData->GetItemsUnderPath(Path, SubItemsFilter);
-
-		for (auto It = SubItems.CreateIterator(); It; ++It)
-		{
-			const FContentBrowserItem& Item = *It;
-			if (!Item.GetInternalPath().IsNone())
-			{
-				if (!PathViewForFiltering.InternalPathPassesBlockLists(FNameBuilder(Item.GetInternalPath())))
-				{
-					It.RemoveCurrent();
-					continue;
-				}
-			}
-			else
-			{
-				// Test if any child internal paths pass for this fully virtual path
-				bool bPasses = false;
-				for (const FContentBrowserItemData& ItemData : Item.GetInternalItems())
-				{
-					UContentBrowserDataSource* ItemDataSource = ItemData.GetOwnerDataSource();
-					if (!ItemDataSource) 
-					{
-						continue;
-					}
-
-					auto OnSubPath = [&PathViewForFiltering, &bPasses, &SubItemsFilter](FName VirtualSubPath, FName InternalPath)
-					{
-						if (InternalPath.IsNone())
-						{
-							// Keep enumerating, this path only exist virtually e.g. /All/Plugins
-							return true;
-						}
-
-						const FNameBuilder InternalPathBuilder(InternalPath);
-						if (ContentBrowserDataUtils::PathPassesAttributeFilter(InternalPathBuilder, 0, SubItemsFilter.ItemAttributeFilter) 
-						&& PathViewForFiltering.InternalPathPassesBlockLists(InternalPathBuilder))
-						{
-							bPasses = true;
-							// Stop enumerating
-							return false;
-						}
-						return true;
-					}; 
-					ItemDataSource->GetRootPathVirtualTree().EnumerateSubPaths(Item.GetVirtualPath(), OnSubPath, /*bRecurse*/ true);
-
-					if (bPasses)
-					{
-						break;
-					}
-				}
-
-				if (!bPasses)
-				{
-					It.RemoveCurrent();
-					continue;
-				}
-			}
-
-			if (!ContentBrowserData->IsFolderVisible(Item.GetVirtualPath(), FolderFlags, FolderFilter))
-			{
-				It.RemoveCurrent();
-				continue;
-			}
-		}
-
-		return SubItems;
-	}
-
-	void RetrieveSmallFolderInformation(const FContentBrowserItem& InFolder, FName& OutFolderBrushName, FName& OutFolderShadowName, FLinearColor& OutFolderOverrideColor)
-	{
-		// We are just dragging folders so add information on the FolderBrush to use and the color
-		if (InFolder.IsValid() && InFolder.IsFolder())
-		{
-			ContentBrowserUtils::TryGetFolderBrushAndShadowNameSmall(InFolder, OutFolderBrushName, OutFolderShadowName);
-
-			const bool bCollectionFolder = EnumHasAnyFlags(InFolder.GetItemCategory(), EContentBrowserItemFlags::Category_Collection);
-			if (bCollectionFolder)
-			{
-				TSharedPtr<ICollectionContainer> CollectionContainer;
-				FName CollectionName;
-				ECollectionShareType::Type CollectionFolderShareType = ECollectionShareType::CST_All;
-				ContentBrowserUtils::IsCollectionPath(InFolder.GetVirtualPath().ToString(), &CollectionContainer , &CollectionName, &CollectionFolderShareType);
-			
-				if (TOptional<FLinearColor> Color = CollectionViewUtils::GetCustomColor(CollectionContainer.Get(), CollectionName, CollectionFolderShareType))
-				{
-					OutFolderOverrideColor = Color.GetValue();
-				}
-			}
-			else
-			{
-				if (TOptional<FLinearColor> Color = ContentBrowserUtils::GetPathColor(InFolder.GetInvariantPath().ToString()))
-				{
-					OutFolderOverrideColor = Color.GetValue();
-				}
-			}
-		}
-	}
-
-	TSharedPtr<SWidget> GetFolderWidgetForNavigationBar(const FText& InFolderName, const FName& InFolderBrushName, const FLinearColor& InFolderOverrideColor)
-	{
-		const FSlateBrush* FolderBrush = FAppStyle::GetBrush(InFolderBrushName);
-
-		if (FolderBrush != FAppStyle::GetDefaultBrush())
-		{
-			return SNew(SHorizontalBox)
-						+ SHorizontalBox::Slot()
-						.AutoWidth()
-						.Padding(FMargin(2, 0, 6, 0))
-						[
-							SNew(SBox)
-							.WidthOverride(16.f)
-							.HeightOverride(16.f)
-							.HAlign(HAlign_Center)
-							.VAlign(VAlign_Center)
-							[
-								SNew(SImage)
-								.DesiredSizeOverride(FVector2D(16, 16))
-								.Image(FolderBrush)
-								.ColorAndOpacity(InFolderOverrideColor)
-							]
-						]
-
-						+ SHorizontalBox::Slot()
-						.FillWidth(1.f)
-						.Padding(FMargin(2, 0, 6, 0))
-						.VAlign( VAlign_Center )
-						[
-							SNew(STextBlock)
-							.TextStyle(FAppStyle::Get(), TEXT("Menu.Label"))
-							.Text(InFolderName)
-						];
-		}
-		return nullptr;
-	}
-}
-
 class SContentBrowser::FCollectionSource
 {
 public:
@@ -2189,32 +2001,7 @@ void SContentBrowser::SyncToItems(TArrayView<const FContentBrowserItem> ItemsToS
 	UContentBrowserDataSubsystem* ContentBrowserData = IContentBrowserDataModule::Get().GetSubsystem();
 	const TSharedRef<FPathPermissionList>& FolderPermissions = FAssetToolsModule::GetModule().Get().GetFolderPermissionList();
 
-	// If any of the items to sync don't pass the permission filter, try to find an alias that might be visible
-	TArray<FContentBrowserItem> NewItemsToSync;
-	for (const FContentBrowserItem& Item : ItemsToSync)
-	{
-		if (FolderPermissions->PassesStartsWithFilter(Item.GetInternalPath()))
-		{
-			NewItemsToSync.Add(Item);
-		}
-		else
-		{
-			TArray<FContentBrowserItemPath> Aliases = ContentBrowserData->GetAliasesForPath(Item.GetInternalPath());
-			for (const FContentBrowserItemPath& Alias : Aliases)
-			{
-				if (FolderPermissions->PassesStartsWithFilter(Alias.GetInternalPathName()))
-				{
-					FContentBrowserItem AliasedItem = ContentBrowserData->GetItemAtPath(Alias.GetVirtualPathName(), EContentBrowserItemTypeFilter::IncludeFiles);
-					if (AliasedItem.IsValid())
-					{
-						NewItemsToSync.Add(MoveTemp(AliasedItem));
-						break;
-					}
-				}
-			}
-		}
-	}
-
+	TArray<FContentBrowserItem> NewItemsToSync = ContentBrowserUtils::FilterOrAliasItems(ItemsToSync);
 	ItemsToSync = NewItemsToSync;
 	PrepareToSyncItems(ItemsToSync, bDisableFiltersThatHideAssets);
 
@@ -3172,7 +2959,7 @@ TArray<FString> SContentBrowser::OnCompletePathPrefix(const FString& Prefix) con
 
 	// Find PrefixView in the available tree of data sources, get its direct children, and filter them by SuffixView 
 	UContentBrowserDataSubsystem* ContentBrowserData = IContentBrowserDataModule::Get().GetSubsystem();
-	TArray<FContentBrowserItem> SubItems = ContentBrowser::GetChildItemsFromVirtualPath(
+	TArray<FContentBrowserItem> SubItems = ContentBrowserUtils::GetChildItemsFromVirtualPath(
 		Parent,
 		PathViewPtr->GetContentBrowserItemCategoryFilter(),
 		PathViewPtr->GetContentBrowserItemAttributeFilter(),
@@ -3213,7 +3000,7 @@ TSharedRef<SWidget> SContentBrowser::OnGetCrumbDelimiterContent(const FString& C
 
 				for( const FCollectionNameType& ChildCollection : ChildCollections )
 				{
-					const FString ChildCollectionCrumbData = FormatCollectionCrumbData(*CollectionSource->GetCollectionContainer(), ChildCollection);
+					const FString ChildCollectionCrumbData = ContentBrowserUtils::FormatCollectionCrumbData(*CollectionSource->GetCollectionContainer(), ChildCollection);
 
 					MenuBuilder.AddMenuEntry(
 						FText::FromName(ChildCollection.Name),
@@ -3229,7 +3016,7 @@ TSharedRef<SWidget> SContentBrowser::OnGetCrumbDelimiterContent(const FString& C
 	}
 	else if( ContentSources.HasVirtualPaths() )
 	{
-		TArray<FContentBrowserItem> SubItems = ContentBrowser::GetChildItemsFromVirtualPath(
+		TArray<FContentBrowserItem> SubItems = ContentBrowserUtils::GetChildItemsFromVirtualPath(
 			*CrumbData,
 			PathViewPtr->GetContentBrowserItemCategoryFilter(),
 			PathViewPtr->GetContentBrowserItemAttributeFilter(),
@@ -3248,16 +3035,18 @@ TSharedRef<SWidget> SContentBrowser::OnGetCrumbDelimiterContent(const FString& C
 			{
 				FName FolderBrushName = NAME_None;
 				FName FolderShadowBrushName = NAME_None;
-				FLinearColor FolderColor = ContentBrowserUtils::GetDefaultColor();
-				ContentBrowser::RetrieveSmallFolderInformation(SubItem, FolderBrushName, FolderShadowBrushName, FolderColor);
-
+				ContentBrowserUtils::TryGetFolderBrushAndShadowNameSmall(SubItem, FolderBrushName, FolderShadowBrushName);
+				
 				FText EntryName = SubItem.GetDisplayName();
 				FUIAction EntryAction = FUIAction(FExecuteAction::CreateSP(const_cast<SContentBrowser*>(this), &SContentBrowser::OnPathMenuItemClicked, SubItem.GetVirtualPath().ToString()));
 
 				if (FolderBrushName != NAME_None)
 				{
+					const FContentBrowserItemPath FolderPath(SubItem.GetVirtualPath(), SubItem.GetInternalPath());
+					FLinearColor FolderColor = UE::Editor::ContentBrowser::ExtensionUtils::GetFolderColor(FolderPath).Get(ContentBrowserUtils::GetDefaultColor());
+						
 					FMenuEntryParams Params;
-					Params.EntryWidget = ContentBrowser::GetFolderWidgetForNavigationBar(EntryName, FolderBrushName, FolderColor);
+					Params.EntryWidget = ContentBrowserUtils::GetFolderWidgetForNavigationBar(EntryName, FolderBrushName, FolderColor);
 					Params.DirectActions = EntryAction;
 					MenuBuilder.AddMenuEntry(Params);
 				}
@@ -3292,15 +3081,6 @@ TSharedRef<SWidget> SContentBrowser::OnGetCrumbDelimiterContent(const FString& C
 	return Widget.ToSharedRef();
 }
 
-FString SContentBrowser::FormatCollectionCrumbData(const ICollectionContainer& CollectionContainer, const FCollectionNameType& Collection)
-{
-	return FString::Printf(
-		TEXT("%s?%s?%s"),
-		*CollectionContainer.GetCollectionSource()->GetName().ToString(),
-		*Collection.Name.ToString(),
-		*FString::FromInt(Collection.Type));
-}
-
 bool SContentBrowser::ParseCollectionCrumbData(const FString& CrumbData, FCollectionSource*& OutCollectionSource, FCollectionNameType& OutCollection) const
 {
 	OutCollectionSource = nullptr;
@@ -4259,77 +4039,8 @@ void SContentBrowser::SyncGlobalSelectionSet()
 
 void SContentBrowser::UpdatePath()
 {
-	FContentSources ContentSources = AssetViewPtr->GetContentSources();
-
-	NavigationBar->ClearPaths();
-
-	if (ContentSources.HasVirtualPaths() )
-	{
-		UContentBrowserDataSubsystem* ContentBrowserData = IContentBrowserDataModule::Get().GetSubsystem();
-
-		TArray<FString> Crumbs;
-		ContentSources.GetVirtualPaths()[0].ToString().ParseIntoArray(Crumbs, TEXT("/"), true);
-
-		FContentBrowserDataFilter SubItemsFilter;
-		SubItemsFilter.ItemTypeFilter = EContentBrowserItemTypeFilter::IncludeFolders;
-		SubItemsFilter.bRecursivePaths = false;
-		SubItemsFilter.ItemCategoryFilter = PathViewPtr->GetContentBrowserItemCategoryFilter();
-		SubItemsFilter.ItemAttributeFilter = PathViewPtr->GetContentBrowserItemAttributeFilter();
-
-		FString CrumbPath = TEXT("/");
-		for (const FString& Crumb : Crumbs)
-		{
-			CrumbPath += Crumb;
-
-			bool bHasSubItems = false;
-			ContentBrowserData->EnumerateItemsUnderPath(*CrumbPath, SubItemsFilter, [&bHasSubItems](FContentBrowserItemData&& InSubItem)
-				{
-					bHasSubItems = true;
-					return false;
-				});
-
-			const FContentBrowserItem CrumbFolderItem = ContentBrowserData->GetItemAtPath(*CrumbPath, EContentBrowserItemTypeFilter::IncludeFolders);
-			NavigationBar->PushPath(CrumbFolderItem.IsValid() ? CrumbFolderItem.GetDisplayName() : FText::FromString(Crumb), CrumbPath, bHasSubItems);
-
-			CrumbPath += TEXT("/");
-		}
-	}
-	else if ( ContentSources.HasCollections() )
-	{
-		const FCollectionRef& Collection = ContentSources.GetCollections()[0];
-
-		TArray<FCollectionNameType> CollectionPathItems;
-
-		// Walk up the parents of this collection so that we can generate a complete path (this loop also adds the child collection to the array)
-		for (TOptional<FCollectionNameType> CurrentCollection = FCollectionNameType(Collection.Name, Collection.Type);
-			CurrentCollection.IsSet(); 
-			CurrentCollection = Collection.Container->GetParentCollection(CurrentCollection->Name, CurrentCollection->Type)
-			)
-		{
-			CollectionPathItems.Insert(CurrentCollection.GetValue(), 0);
-		}
-
-		// Now add each part of the path to the breadcrumb trail
-		for (const FCollectionNameType& CollectionPathItem : CollectionPathItems)
-		{
-			const FString CrumbData = FormatCollectionCrumbData(*Collection.Container, CollectionPathItem);
-			
-			TArray<FCollectionNameType> ChildCollections;
-			Collection.Container->GetChildCollections(CollectionPathItem.Name, CollectionPathItem.Type, ChildCollections);
-			const bool bHasChildren = ChildCollections.Num() > 0;
-
-			FFormatNamedArguments Args;
-			Args.Add(TEXT("CollectionName"), FText::FromName(CollectionPathItem.Name));
-			const FText DisplayName = FText::Format(LOCTEXT("CollectionPathIndicator", "{CollectionName} (Collection)"), Args);
-
-			NavigationBar->PushPath(DisplayName, CrumbData, bHasChildren);
-		}
-	}
-	else
-	{
-		NavigationBar->PushPath(LOCTEXT("AllAssets", "All Assets"), TEXT(""), true); 
-	}
-
+	ContentBrowserUtils::UpdateNavigationBar(NavigationBar, AssetViewPtr, PathViewPtr);
+	
 	CachedCanWriteToCurrentPath.Reset();
 }
 
@@ -5205,7 +4916,7 @@ void SContentBrowser::OpenNewContentBrowser()
  
 const FContentBrowserInstanceConfig* SContentBrowser::GetConstInstanceConfig() const
 {
-	return ContentBrowser::GetConstInstanceConfig(InstanceName);
+	return ContentBrowserUtils::GetConstInstanceConfig(InstanceName);
 }
 
 FContentBrowserInstanceConfig* SContentBrowser::GetMutableInstanceConfig()
diff --git a/Engine/Source/Editor/ContentBrowser/Private/SContentBrowser.h b/Engine/Source/Editor/ContentBrowser/Private/SContentBrowser.h
--- a/Engine/Source/Editor/ContentBrowser/Private/SContentBrowser.h
+++ b/Engine/Source/Editor/ContentBrowser/Private/SContentBrowser.h
@@ -355,7 +355,6 @@ private:
 	 */
 	TSharedRef<SWidget> OnGetCrumbDelimiterContent(const FString& CrumbData) const;
 
-	static FString FormatCollectionCrumbData(const ICollectionContainer& CollectionContainer, const FCollectionNameType& Collection);
 	bool ParseCollectionCrumbData(const FString& CrumbData, FCollectionSource*& OutCollectionSource, FCollectionNameType& OutCollection) const;
 
 	/** Return a list of recently visited locations which can be used with OnNavigateToPath to return to those locations */
diff --git a/Engine/Source/Editor/ContentBrowser/Private/SPathPicker.cpp b/Engine/Source/Editor/ContentBrowser/Private/SPathPicker.cpp
--- a/Engine/Source/Editor/ContentBrowser/Private/SPathPicker.cpp
+++ b/Engine/Source/Editor/ContentBrowser/Private/SPathPicker.cpp
@@ -253,4 +253,13 @@ const TSharedPtr<SPathView>& SPathPicker::GetPathView() const
 	return PathViewPtr;
 }
 
+void SPathPicker::SyncToItems(TArrayView<const FContentBrowserItem> ItemsToSync, const bool bAllowImplicitSync)
+{
+	if (PathViewPtr->DisablePluginPathFiltersThatHideItems(ItemsToSync))
+	{
+		PathViewPtr->Populate();
+	}
+	PathViewPtr->SyncToItems(ItemsToSync, bAllowImplicitSync);
+}
+
 #undef LOCTEXT_NAMESPACE
diff --git a/Engine/Source/Editor/ContentBrowser/Private/SPathPicker.h b/Engine/Source/Editor/ContentBrowser/Private/SPathPicker.h
--- a/Engine/Source/Editor/ContentBrowser/Private/SPathPicker.h
+++ b/Engine/Source/Editor/ContentBrowser/Private/SPathPicker.h
@@ -40,6 +40,8 @@ public:
 
 	/** Return the associated SPathView */
 	const TSharedPtr<SPathView>& GetPathView() const;
+	
+	void SyncToItems(TArrayView<const FContentBrowserItem> ItemsToSync, const bool bAllowImplicitSync);
 
 	/** Handler for creating a new folder in the path picker */
 	void CreateNewFolder(FString FolderPath, FOnCreateNewFolder InOnCreateNewFolder);
diff --git a/Engine/Source/Editor/ContentBrowser/Public/Experimental/ContentBrowserExtensionUtils.h b/Engine/Source/Editor/ContentBrowser/Public/Experimental/ContentBrowserExtensionUtils.h
--- a/Engine/Source/Editor/ContentBrowser/Public/Experimental/ContentBrowserExtensionUtils.h
+++ b/Engine/Source/Editor/ContentBrowser/Public/Experimental/ContentBrowserExtensionUtils.h
@@ -5,9 +5,14 @@
 #include "Containers/UnrealString.h"
 #include "Math/Color.h"
 
+struct FContentBrowserItemPath;
+
 namespace UE::Editor::ContentBrowser::ExtensionUtils
 {
 	// Get the custom color of the given folder path (if any)
+	CONTENTBROWSER_API TOptional<FLinearColor> GetFolderColor(const FContentBrowserItemPath& FolderPath);
+
+	// Get the custom color of the given internal folder path (if any)
 	CONTENTBROWSER_API TOptional<FLinearColor> GetFolderColor(const FName& FolderPath);
 
 	// Set a custom color for the given folder path
