Fixing an issue where content browser commands could be added to a non-existent parent command context due to being run outside a slate application.

#rnx

diff --git a/Engine/Source/Editor/ContentBrowser/Private/ContentBrowserCommands.cpp b/Engine/Source/Editor/ContentBrowser/Private/ContentBrowserCommands.cpp
--- a/Engine/Source/Editor/ContentBrowser/Private/ContentBrowserCommands.cpp
+++ b/Engine/Source/Editor/ContentBrowser/Private/ContentBrowserCommands.cpp
@@ -10,6 +10,7 @@
 
 #define LOCTEXT_NAMESPACE "ContentBrowser"
 
+const FName FContentBrowserCommands::GenericCommandsName = TEXT("GenericCommands");
 
 void FContentBrowserCommands::RegisterCommands()
 {
diff --git a/Engine/Source/Editor/ContentBrowser/Private/ContentBrowserCommands.h b/Engine/Source/Editor/ContentBrowser/Private/ContentBrowserCommands.h
--- a/Engine/Source/Editor/ContentBrowser/Private/ContentBrowserCommands.h
+++ b/Engine/Source/Editor/ContentBrowser/Private/ContentBrowserCommands.h
@@ -17,9 +17,11 @@ class FContentBrowserCommands
 {
 public:
 
+	static const FName GenericCommandsName;
+
 	/** Default constructor. */
 	FContentBrowserCommands()
-		: TCommands<FContentBrowserCommands>(TEXT("ContentBrowser"), NSLOCTEXT( "ContentBrowser", "ContentBrowser", "Content Browser" ), "GenericCommands", FAppStyle::GetAppStyleSetName() )
+		: TCommands<FContentBrowserCommands>(TEXT("ContentBrowser"), NSLOCTEXT( "ContentBrowser", "ContentBrowser", "Content Browser" ), (FInputBindingManager::Get().GetContextByName(GenericCommandsName).IsValid() ? GenericCommandsName : NAME_None), FAppStyle::GetAppStyleSetName())
 	{ }
 
 public:
