Rename FContentSources to FAssetViewContentSources

diff --git a/Engine/Plugins/Animation/ControlRig/Source/ControlRigEditor/Private/EditMode/SControlRigBaseListWidget.cpp b/Engine/Plugins/Animation/ControlRig/Source/ControlRigEditor/Private/EditMode/SControlRigBaseListWidget.cpp
--- a/Engine/Plugins/Animation/ControlRig/Source/ControlRigEditor/Private/EditMode/SControlRigBaseListWidget.cpp
+++ b/Engine/Plugins/Animation/ControlRig/Source/ControlRigEditor/Private/EditMode/SControlRigBaseListWidget.cpp
@@ -1348,7 +1348,7 @@ FReply SControlRigBaseListWidget::ExecuteDeleteFolderConfirmed()
 				DefaultSelectedPaths.Add(TEXT("/Game"));
 				PathPicker->GetPathView()->SetSelectedPaths(DefaultSelectedPaths);
 
-				FContentSources DefaultContentSources(FName("/Game"));
+				FAssetViewContentSources DefaultContentSources(FName("/Game"));
 				AssetPicker->GetAssetView()->SetContentSources(DefaultContentSources);
 			}
 		}
diff --git a/Engine/Plugins/Experimental/MeshModelingToolsetExp/Source/ModelingEditorUI/Private/ModelingWidgets/SToolInputAssetPicker.cpp b/Engine/Plugins/Experimental/MeshModelingToolsetExp/Source/ModelingEditorUI/Private/ModelingWidgets/SToolInputAssetPicker.cpp
--- a/Engine/Plugins/Experimental/MeshModelingToolsetExp/Source/ModelingEditorUI/Private/ModelingWidgets/SToolInputAssetPicker.cpp
+++ b/Engine/Plugins/Experimental/MeshModelingToolsetExp/Source/ModelingEditorUI/Private/ModelingWidgets/SToolInputAssetPicker.cpp
@@ -138,7 +138,7 @@ void SToolInputAssetPicker::Construct( const FArguments& InArgs )
 	}
 	PRAGMA_ENABLE_DEPRECATION_WARNINGS
 
-	CurrentContentSources = FContentSources(InArgs._AssetPickerConfig.Filter.PackagePaths, MoveTemp(CollectionsFilter));
+	CurrentContentSources = FAssetViewContentSources(InArgs._AssetPickerConfig.Filter.PackagePaths, MoveTemp(CollectionsFilter));
 	CurrentBackendFilter = InArgs._AssetPickerConfig.Filter;
 	CurrentBackendFilter.PackagePaths.Reset();
 
@@ -231,7 +231,7 @@ void SToolInputAssetPicker::Construct( const FArguments& InArgs )
 
 void SToolInputAssetPicker::UpdateAssetSourceCollections(TArray<FCollectionRef> Collections)
 {
-	CurrentContentSources = FContentSources(MoveTemp(Collections));
+	CurrentContentSources = FAssetViewContentSources(MoveTemp(Collections));
 	if (AssetViewPtr.IsValid())
 	{
 		AssetViewPtr->SetContentSources(CurrentContentSources);
diff --git a/Engine/Plugins/Experimental/MeshModelingToolsetExp/Source/ModelingEditorUI/Public/ModelingWidgets/SToolInputAssetPicker.h b/Engine/Plugins/Experimental/MeshModelingToolsetExp/Source/ModelingEditorUI/Public/ModelingWidgets/SToolInputAssetPicker.h
--- a/Engine/Plugins/Experimental/MeshModelingToolsetExp/Source/ModelingEditorUI/Public/ModelingWidgets/SToolInputAssetPicker.h
+++ b/Engine/Plugins/Experimental/MeshModelingToolsetExp/Source/ModelingEditorUI/Public/ModelingWidgets/SToolInputAssetPicker.h
@@ -9,7 +9,7 @@
 #include "Widgets/SCompoundWidget.h"
 #include "AssetRegistry/AssetData.h"
 #include "AssetRegistry/ARFilter.h"
-#include "ContentSources.h"
+#include "AssetViewContentSources.h"
 #include "IContentBrowserSingleton.h"
 #include "SourcesData.h"
 
@@ -127,7 +127,7 @@ private:
 	EAssetTypeCategories::Type DefaultFilterMenuExpansion;
 
 	/** The sources data currently used by the picker */
-	FContentSources CurrentContentSources;
+	FAssetViewContentSources CurrentContentSources;
 
 	/** Current filter we are using, needed reset asset view after we have custom filtered */
 	FARFilter CurrentBackendFilter;
diff --git a/Engine/Source/Editor/ContentBrowser/Private/AssetContextMenu.cpp b/Engine/Source/Editor/ContentBrowser/Private/AssetContextMenu.cpp
--- a/Engine/Source/Editor/ContentBrowser/Private/AssetContextMenu.cpp
+++ b/Engine/Source/Editor/ContentBrowser/Private/AssetContextMenu.cpp
@@ -112,7 +112,7 @@ void FAssetContextMenu::BindCommands(TSharedPtr< FUICommandList >& Commands)
 	}
 }
 
-TSharedRef<SWidget> FAssetContextMenu::MakeContextMenu(TArrayView<const FContentBrowserItem> InSelectedItems, const FContentSources& InContentSources, TSharedPtr< FUICommandList > InCommandList)
+TSharedRef<SWidget> FAssetContextMenu::MakeContextMenu(TArrayView<const FContentBrowserItem> InSelectedItems, const FAssetViewContentSources& InContentSources, TSharedPtr< FUICommandList > InCommandList)
 {
 	FWarnIfAssetsLoadedInScope WarnIfAssetsLoaded;
 	
diff --git a/Engine/Source/Editor/ContentBrowser/Private/AssetContextMenu.h b/Engine/Source/Editor/ContentBrowser/Private/AssetContextMenu.h
--- a/Engine/Source/Editor/ContentBrowser/Private/AssetContextMenu.h
+++ b/Engine/Source/Editor/ContentBrowser/Private/AssetContextMenu.h
@@ -6,7 +6,7 @@
 #include "Containers/ArrayView.h"
 #include "ContentBrowserDataMenuContexts.h"
 #include "ContentBrowserItem.h"
-#include "ContentSources.h"
+#include "AssetViewContentSources.h"
 #include "Delegates/Delegate.h"
 #include "HAL/Platform.h"
 #include "Input/Reply.h"
@@ -34,7 +34,7 @@ public:
 	void BindCommands(TSharedPtr< FUICommandList >& Commands);
 
 	/** Makes the context menu widget */
-	TSharedRef<SWidget> MakeContextMenu(TArrayView<const FContentBrowserItem> InSelectedItems, const FContentSources& InContentSources, TSharedPtr< FUICommandList > InCommandList);
+	TSharedRef<SWidget> MakeContextMenu(TArrayView<const FContentBrowserItem> InSelectedItems, const FAssetViewContentSources& InContentSources, TSharedPtr< FUICommandList > InCommandList);
 
 	/** Updates the list of currently selected items to those passed in */
 	void SetSelectedItems(TArrayView<const FContentBrowserItem> InSelectedItems);
@@ -178,7 +178,7 @@ private:
 	TArray<FContentBrowserItem> SelectedFiles;
 	TArray<FContentBrowserItem> SelectedFolders;
 
-	FContentSources ContentSources;
+	FAssetViewContentSources ContentSources;
 
 	/** The asset view this context menu is a part of */
 	TWeakPtr<SAssetView> AssetView;
diff --git a/Engine/Source/Editor/ContentBrowser/Private/AssetViewContentSources.cpp b/Engine/Source/Editor/ContentBrowser/Private/AssetViewContentSources.cpp
new file mode 100644
--- /dev/null
+++ b/Engine/Source/Editor/ContentBrowser/Private/AssetViewContentSources.cpp
@@ -0,0 +1,93 @@
+// Copyright Epic Games, Inc. All Rights Reserved.
+
+#include "AssetViewContentSources.h"
+
+#include "ICollectionContainer.h"
+
+FAssetViewContentSources::FAssetViewContentSources(FName InVirtualPath)
+	: VirtualPaths({ InVirtualPath })
+{
+}
+
+FAssetViewContentSources::FAssetViewContentSources(TArray<FName> InVirtualPaths)
+	: VirtualPaths(MoveTemp(InVirtualPaths))
+{
+}
+
+FAssetViewContentSources::FAssetViewContentSources(FCollectionRef InCollection)
+	: Collections({ InCollection })
+{
+	SanitizeCollections();
+}
+
+FAssetViewContentSources::FAssetViewContentSources(TArray<FCollectionRef> InCollections)
+	: Collections(MoveTemp(InCollections))
+{
+	SanitizeCollections();
+}
+
+FAssetViewContentSources::FAssetViewContentSources(TArray<FName> InVirtualPaths, TArray<FCollectionRef> InCollections)
+	: VirtualPaths(MoveTemp(InVirtualPaths))
+	, Collections(MoveTemp(InCollections))
+{
+	// Only one of VirtualPaths or Collections should be set.
+	check(VirtualPaths.IsEmpty() || Collections.IsEmpty());
+
+	SanitizeCollections();
+}
+
+void FAssetViewContentSources::SetVirtualPath(FName InVirtualPath)
+{
+	Reset();
+
+	VirtualPaths.Add(InVirtualPath);
+}
+
+void FAssetViewContentSources::SetVirtualPaths(const TArray<FName>& InVirtualPaths)
+{
+	VirtualPaths = InVirtualPaths;
+
+	Collections.Reset();
+}
+
+void FAssetViewContentSources::SetCollection(const FCollectionRef& InCollection)
+{
+	Reset();
+
+	Collections.Add(InCollection);
+
+	SanitizeCollections();
+}
+
+void FAssetViewContentSources::SetCollections(const TArray<FCollectionRef>& InCollections)
+{
+	VirtualPaths.Reset();
+
+	Collections = InCollections;
+
+	SanitizeCollections();
+}
+
+bool FAssetViewContentSources::IsDynamicCollection() const
+{
+	if (Collections.Num() == 1)
+	{
+		const FCollectionRef& Collection = Collections[0];
+		ECollectionStorageMode::Type StorageMode = ECollectionStorageMode::Static;
+		return (Collection.Container->GetCollectionStorageMode(Collection.Name, Collection.Type, StorageMode) && StorageMode == ECollectionStorageMode::Dynamic);
+	}
+
+	return false;
+}
+
+void FAssetViewContentSources::Reset()
+{
+	VirtualPaths.Reset();
+
+	Collections.Reset();
+}
+
+void FAssetViewContentSources::SanitizeCollections()
+{
+	verify(Collections.RemoveAll([](const FCollectionRef& Collection) { return !Collection.Container.IsValid(); }) == 0);
+}
\ No newline at end of file
diff --git a/Engine/Source/Editor/ContentBrowser/Private/ContentBrowserUtils.cpp b/Engine/Source/Editor/ContentBrowser/Private/ContentBrowserUtils.cpp
--- a/Engine/Source/Editor/ContentBrowser/Private/ContentBrowserUtils.cpp
+++ b/Engine/Source/Editor/ContentBrowser/Private/ContentBrowserUtils.cpp
@@ -1282,7 +1282,7 @@ FString ContentBrowserUtils::FormatCollectionCrumbData(const ICollectionContaine
 
 void ContentBrowserUtils::UpdateNavigationBar(const TSharedPtr<SNavigationBar>& NavigationBar, const TSharedPtr<SAssetView>& AssetView, const TSharedPtr<SPathView>& PathView)
 {
-	const FContentSources& ContentSources = AssetView->GetContentSources();
+	const FAssetViewContentSources& ContentSources = AssetView->GetContentSources();
 
 	NavigationBar->ClearPaths();
 
diff --git a/Engine/Source/Editor/ContentBrowser/Private/ContentSources.cpp b/Engine/Source/Editor/ContentBrowser/Private/ContentSources.cpp
deleted file mode 100644
diff --git a/Engine/Source/Editor/ContentBrowser/Private/SAssetDialog.cpp b/Engine/Source/Editor/ContentBrowser/Private/SAssetDialog.cpp
--- a/Engine/Source/Editor/ContentBrowser/Private/SAssetDialog.cpp
+++ b/Engine/Source/Editor/ContentBrowser/Private/SAssetDialog.cpp
@@ -678,7 +678,7 @@ void SAssetDialog::SelectDefaultPaths()
 
 	PathPicker->GetPathView()->SetSelectedPaths(DefaultSelectedPaths);
 
-	FContentSources DefaultContentSources(DefaultVirtualPathsToSelect);
+	FAssetViewContentSources DefaultContentSources(DefaultVirtualPathsToSelect);
 	AssetPicker->GetAssetView()->SetContentSources(DefaultContentSources);
 }
 
@@ -904,7 +904,7 @@ void SAssetDialog::OnApplyHistoryData(const FHistoryData& History)
 
 void SAssetDialog::OnUpdateHistoryData(FHistoryData& HistoryData) const
 {
-	const FContentSources& ContentSources = AssetPicker->GetAssetView()->GetContentSources();
+	const FAssetViewContentSources& ContentSources = AssetPicker->GetAssetView()->GetContentSources();
 	const TArray<FContentBrowserItem> SelectedItems = AssetPicker->GetAssetView()->GetSelectedItems();
 
 	const FText NewSource = ContentSources.HasVirtualPaths() ? FText::FromName(ContentSources.GetVirtualPaths()[0]) : (ContentSources.HasCollections() ? FText::FromName(ContentSources.GetCollections()[0].Name) : LOCTEXT("AllAssets", "All Assets"));
@@ -956,7 +956,7 @@ FText SAssetDialog::GetNameErrorLabelText() const
 
 bool SAssetDialog::OnCanEditPathAsText(const FString& Text) const
 {
-	const FContentSources& ContentSources = AssetPicker->GetAssetView()->GetContentSources();
+	const FAssetViewContentSources& ContentSources = AssetPicker->GetAssetView()->GetContentSources();
 	if (ContentSources.HasCollections())
 	{
 		return false;
@@ -1098,7 +1098,7 @@ void SAssetDialog::SetCurrentlySelectedPath(const FString& NewPath, const EConte
 	{
 		// The asset picker uses the FARFilter::PackagePaths field to fill content sources.
 		// Thus, content sources must be checked rather than the filter itself.
-		const FContentSources& Sources = AssetPicker->GetAssetView()->GetContentSources();
+		const FAssetViewContentSources& Sources = AssetPicker->GetAssetView()->GetContentSources();
 		if (!Sources.HasVirtualPaths() || Sources.GetVirtualPaths()[0] != NewVirtualPath)
 		{
 			FARFilter NewFilter;
diff --git a/Engine/Source/Editor/ContentBrowser/Private/SAssetPicker.cpp b/Engine/Source/Editor/ContentBrowser/Private/SAssetPicker.cpp
--- a/Engine/Source/Editor/ContentBrowser/Private/SAssetPicker.cpp
+++ b/Engine/Source/Editor/ContentBrowser/Private/SAssetPicker.cpp
@@ -102,7 +102,7 @@ void SAssetPicker::Construct( const FArguments& InArgs )
 	}
 	PRAGMA_ENABLE_DEPRECATION_WARNINGS
 
-	CurrentContentSources = FContentSources(InArgs._AssetPickerConfig.Filter.PackagePaths, MoveTemp(CollectionsFilter));
+	CurrentContentSources = FAssetViewContentSources(InArgs._AssetPickerConfig.Filter.PackagePaths, MoveTemp(CollectionsFilter));
 	CurrentBackendFilter = InArgs._AssetPickerConfig.Filter;
 	CurrentBackendFilter.PackagePaths.Reset();
 
diff --git a/Engine/Source/Editor/ContentBrowser/Private/SAssetPicker.h b/Engine/Source/Editor/ContentBrowser/Private/SAssetPicker.h
--- a/Engine/Source/Editor/ContentBrowser/Private/SAssetPicker.h
+++ b/Engine/Source/Editor/ContentBrowser/Private/SAssetPicker.h
@@ -8,7 +8,7 @@
 #include "Containers/ArrayView.h"
 #include "Containers/UnrealString.h"
 #include "ContentBrowserDelegates.h"
-#include "ContentSources.h"
+#include "AssetViewContentSources.h"
 #include "HAL/Platform.h"
 #include "IAssetTypeActions.h"
 #include "IContentBrowserSingleton.h"
@@ -198,7 +198,7 @@ private:
 	EAssetTypeCategories::Type DefaultFilterMenuExpansion;
 
 	/** The sources data currently used by the picker */
-	FContentSources CurrentContentSources;
+	FAssetViewContentSources CurrentContentSources;
 
 	/** Current filter we are using, needed reset asset view after we have custom filtered */
 	FARFilter CurrentBackendFilter;
diff --git a/Engine/Source/Editor/ContentBrowser/Private/SAssetView.cpp b/Engine/Source/Editor/ContentBrowser/Private/SAssetView.cpp
--- a/Engine/Source/Editor/ContentBrowser/Private/SAssetView.cpp
+++ b/Engine/Source/Editor/ContentBrowser/Private/SAssetView.cpp
@@ -383,7 +383,7 @@ public:
 	 * @param bAllowItemRecycling Whether or not to allow reuse of items and therefore widgets. 
 	 * 	Setting to false can avoid lots of time firing modification delegates for recursive searches.
 	 */
-	void RefreshItemsFromBackend(const FContentSources& ContentSources, const FContentBrowserDataFilter& DataFilter, bool bAllowItemRecycling);
+	void RefreshItemsFromBackend(const FAssetViewContentSources& ContentSources, const FContentBrowserDataFilter& DataFilter, bool bAllowItemRecycling);
 	
 	/** 
 	 * Find an FAssetViewItem containing the given content browser data if one exists.
@@ -1200,7 +1200,7 @@ void SAssetView::Construct( const FArguments& InArgs )
 			Collections,
 			[](const FCollectionNameType& Collection) { return FCollectionRef(FCollectionManagerModule::GetModule().Get().GetProjectCollectionContainer(), Collection); });
 
-		ContentSources = FContentSources(InArgs._InitialSourcesData.VirtualPaths, MoveTemp(Collections));
+		ContentSources = FAssetViewContentSources(InArgs._InitialSourcesData.VirtualPaths, MoveTemp(Collections));
 
 		ContentSources.OnEnumerateCustomSourceItemDatas = InArgs._InitialSourcesData.OnEnumerateCustomSourceItemDatas;
 		ContentSources.bIncludeVirtualPaths = InArgs._InitialSourcesData.IsIncludingVirtualPaths();
@@ -1524,7 +1524,7 @@ TOptional< float > SAssetView::GetIsWorkingProgressBarState() const
 	return 0.0f;
 }
 
-void SAssetView::SetContentSources(const FContentSources& InContentSources)
+void SAssetView::SetContentSources(const FAssetViewContentSources& InContentSources)
 {
 	TSet<TSharedPtr<ICollectionContainer>> OldCollectionContainers;
 	for (const FCollectionRef& Collection : ContentSources.GetCollections())
@@ -1577,7 +1577,7 @@ void SAssetView::SetSourcesData(const FSourcesData& InSourcesData)
 		Collections,
 		[](const FCollectionNameType& Collection) { return FCollectionRef(FCollectionManagerModule::GetModule().Get().GetProjectCollectionContainer(), Collection); });
 
-	FContentSources NewContentSources(InSourcesData.VirtualPaths, MoveTemp(Collections));
+	FAssetViewContentSources NewContentSources(InSourcesData.VirtualPaths, MoveTemp(Collections));
 
 	NewContentSources.OnEnumerateCustomSourceItemDatas = InSourcesData.OnEnumerateCustomSourceItemDatas;
 	NewContentSources.bIncludeVirtualPaths = InSourcesData.IsIncludingVirtualPaths();
@@ -1587,7 +1587,7 @@ void SAssetView::SetSourcesData(const FSourcesData& InSourcesData)
 	PRAGMA_ENABLE_DEPRECATION_WARNINGS
 }
 
-const FContentSources& SAssetView::GetContentSources() const
+const FAssetViewContentSources& SAssetView::GetContentSources() const
 {
 	return ContentSources;
 }
@@ -3055,7 +3055,7 @@ void SAssetView::RefreshSourceItems()
 		FPlatformTime::Seconds() - RefreshSourceItemsStartTime);
 }
 
-void FAssetViewItemCollection::RefreshItemsFromBackend(const FContentSources& ContentSources, const FContentBrowserDataFilter& DataFilter, bool bAllowItemRecycling)
+void FAssetViewItemCollection::RefreshItemsFromBackend(const FAssetViewContentSources& ContentSources, const FContentBrowserDataFilter& DataFilter, bool bAllowItemRecycling)
 {
 	AbortTextFiltering();
 
diff --git a/Engine/Source/Editor/ContentBrowser/Private/SContentBrowser.cpp b/Engine/Source/Editor/ContentBrowser/Private/SContentBrowser.cpp
--- a/Engine/Source/Editor/ContentBrowser/Private/SContentBrowser.cpp
+++ b/Engine/Source/Editor/ContentBrowser/Private/SContentBrowser.cpp
@@ -40,7 +40,7 @@
 #include "ContentBrowserStyle.h"
 #include "ContentBrowserUtils.h"
 #include "ContentBrowserVirtualPathTree.h"
-#include "ContentSources.h"
+#include "AssetViewContentSources.h"
 #include "CoreGlobals.h"
 #include "Delegates/Delegate.h"
 #include "Editor.h"
@@ -842,7 +842,7 @@ void SContentBrowser::Construct( const FArguments& InArgs, const FName& InInstan
 	if (SelectedCollection.IsSet() && SelectedCollection.GetValue().IsValid())
 	{
 		// Select the specified collection by default
-		FContentSources DefaultContentSources(SelectedCollection.GetValue());
+		FAssetViewContentSources DefaultContentSources(SelectedCollection.GetValue());
 		AssetViewPtr->SetContentSources(DefaultContentSources);
 	}
 	else
@@ -852,7 +852,7 @@ void SContentBrowser::Construct( const FArguments& InArgs, const FName& InInstan
 		FName DefaultVirtualPath;
 		IContentBrowserDataModule::Get().GetSubsystem()->ConvertInternalPathToVirtual(DefaultInvariantPath, DefaultVirtualPath);
 
-		FContentSources DefaultContentSources(DefaultVirtualPath);
+		FAssetViewContentSources DefaultContentSources(DefaultVirtualPath);
 		AssetViewPtr->SetContentSources(DefaultContentSources);
 	}
 
@@ -2415,7 +2415,7 @@ void SContentBrowser::SourcesChanged(const TArray<FString>& SelectedPaths, const
 	FString NewSource = SelectedPaths.Num() > 0 ? SelectedPaths[0] : (SelectedCollections.Num() > 0 ? SelectedCollections[0].Name.ToString() : TEXT("None"));
 	UE_LOG(LogContentBrowser, VeryVerbose, TEXT("The content browser source was changed by the sources view to '%s'"), *NewSource);
 
-	FContentSources ContentSources;
+	FAssetViewContentSources ContentSources;
 	{
 		TArray<FName> SelectedPathNames;
 		SelectedPathNames.Reserve(SelectedPaths.Num());
@@ -2423,7 +2423,7 @@ void SContentBrowser::SourcesChanged(const TArray<FString>& SelectedPaths, const
 		{
 			SelectedPathNames.Add(FName(*SelectedPath));
 		}
-		ContentSources = FContentSources(MoveTemp(SelectedPathNames), SelectedCollections);
+		ContentSources = FAssetViewContentSources(MoveTemp(SelectedPathNames), SelectedCollections);
 	}
 
 	// A dynamic collection should apply its search query to the CB search, so we need to stash the current search so that we can restore it again later
@@ -2659,7 +2659,7 @@ void SContentBrowser::OnApplyHistoryData( const FHistoryData& History )
 
 void SContentBrowser::OnUpdateHistoryData(FHistoryData& HistoryData) const
 {
-	const FContentSources& ContentSources = AssetViewPtr->GetContentSources();
+	const FAssetViewContentSources& ContentSources = AssetViewPtr->GetContentSources();
 	const TArray<FContentBrowserItem> SelectedItems = AssetViewPtr->GetSelectedItems();
 
 	const FText NewSource = ContentSources.HasVirtualPaths() ? FText::FromName(ContentSources.GetVirtualPaths()[0]) : (ContentSources.HasCollections() ? FText::FromName(ContentSources.GetCollections()[0].Name) : LOCTEXT("AllAssets", "All Assets"));
@@ -2890,7 +2890,7 @@ void SContentBrowser::OnNavigateToPath(const FString& NewPath)
 
 void SContentBrowser::OnPathClicked( const FString& CrumbData )
 {
-	FContentSources ContentSources = AssetViewPtr->GetContentSources();
+	FAssetViewContentSources ContentSources = AssetViewPtr->GetContentSources();
 
 	if (ContentSources.HasCollections() )
 	{
@@ -2936,7 +2936,7 @@ void SContentBrowser::OnPathMenuItemClicked(FString ClickedPath)
 
 bool SContentBrowser::OnCanEditPathAsText(const FString& Text) const
 {
-	const FContentSources& ContentSources = AssetViewPtr->GetContentSources();
+	const FAssetViewContentSources& ContentSources = AssetViewPtr->GetContentSources();
 	if (ContentSources.HasCollections())
 	{
 		// Do not present collections as text because their names are not very user friendly right now.
@@ -2980,7 +2980,7 @@ TArray<FString> SContentBrowser::OnCompletePathPrefix(const FString& Prefix) con
 
 TSharedRef<SWidget> SContentBrowser::OnGetCrumbDelimiterContent(const FString& CrumbData) const
 {
-	const FContentSources& ContentSources = AssetViewPtr->GetContentSources();
+	const FAssetViewContentSources& ContentSources = AssetViewPtr->GetContentSources();
 
 	TSharedPtr<SWidget> Widget = SNullWidget::NullWidget;
 	TSharedPtr<SWidget> MenuWidget;
@@ -3119,7 +3119,7 @@ bool SContentBrowser::ParseCollectionCrumbData(const FString& CrumbData, FCollec
 FString SContentBrowser::GetCurrentPath(const EContentBrowserPathType PathType) const
 {
 	FString CurrentPath;
-	const FContentSources& ContentSources = AssetViewPtr->GetContentSources();
+	const FAssetViewContentSources& ContentSources = AssetViewPtr->GetContentSources();
 	if ( ContentSources.HasVirtualPaths() && ContentSources.GetVirtualPaths()[0] != NAME_None )
 	{
 		if (PathType == EContentBrowserPathType::Virtual)
@@ -3180,7 +3180,7 @@ void SContentBrowser::AppendNewMenuContextObjects(const EContentBrowserDataMenuC
 
 TSharedRef<SWidget> SContentBrowser::MakeAddNewContextMenu(const EContentBrowserDataMenuContext_AddNewMenuDomain InDomain, UContentBrowserToolbarMenuContext* CommonContext)
 {
-	const FContentSources& ContentSources = AssetViewPtr->GetContentSources();
+	const FAssetViewContentSources& ContentSources = AssetViewPtr->GetContentSources();
 
 	bool bCanBeModified = false;
 
@@ -3267,7 +3267,7 @@ bool SContentBrowser::CanWriteToCurrentPath() const
 {
 	if (AssetViewPtr.IsValid())
 	{
-		const FContentSources& ContentSources = AssetViewPtr->GetContentSources();
+		const FAssetViewContentSources& ContentSources = AssetViewPtr->GetContentSources();
 		if (ContentSources.GetVirtualPaths().Num() == 1)
 		{
 			FName CurrentPath = ContentSources.GetVirtualPaths()[0];
@@ -3358,7 +3358,7 @@ bool SContentBrowser::IsAddNewEnabled() const
 
 FText SContentBrowser::GetAddNewToolTipText() const
 {
-	const FContentSources& ContentSources = AssetViewPtr->GetContentSources();
+	const FAssetViewContentSources& ContentSources = AssetViewPtr->GetContentSources();
 
 	if ( ContentSources.GetVirtualPaths().Num() == 1 )
 	{
@@ -4061,7 +4061,7 @@ FText SContentBrowser::GetPathText() const
 
 	if ( IsFilteredBySource() )
 	{
-		const FContentSources& ContentSources = AssetViewPtr->GetContentSources();
+		const FAssetViewContentSources& ContentSources = AssetViewPtr->GetContentSources();
 
 		// At least one source is selected
 		const int32 NumSources = ContentSources.GetVirtualPaths().Num() + ContentSources.GetCollections().Num();
@@ -4086,7 +4086,7 @@ FText SContentBrowser::GetPathText() const
 
 bool SContentBrowser::IsFilteredBySource() const
 {
-	const FContentSources& ContentSources = AssetViewPtr->GetContentSources();
+	const FAssetViewContentSources& ContentSources = AssetViewPtr->GetContentSources();
 	return !ContentSources.IsEmpty();
 }
 
@@ -4270,7 +4270,7 @@ void SContentBrowser::HideCollectionContainer(const TSharedRef<ICollectionContai
 {
 	RemoveSlotForCollectionContainer(CollectionContainer);
 
-	auto UpdateContentSources = [&CollectionContainer](const FContentSources& ContentSources, FContentSources& OutNewContentSources)
+	auto UpdateContentSources = [&CollectionContainer](const FAssetViewContentSources& ContentSources, FAssetViewContentSources& OutNewContentSources)
 		{
 			auto Predicate = [&CollectionContainer](const FCollectionRef& Collection)
 				{
@@ -4290,7 +4290,7 @@ void SContentBrowser::HideCollectionContainer(const TSharedRef<ICollectionContai
 		};
 
 	{
-		FContentSources NewContentSources;
+		FAssetViewContentSources NewContentSources;
 		if (UpdateContentSources(AssetViewPtr->GetContentSources(), NewContentSources))
 		{
 			AssetViewPtr->SetContentSources(NewContentSources);
@@ -4299,7 +4299,7 @@ void SContentBrowser::HideCollectionContainer(const TSharedRef<ICollectionContai
 
 	HistoryManager.RemoveHistoryData([&UpdateContentSources](FHistoryData& HistoryData)
 		{
-			FContentSources NewContentSources;
+			FAssetViewContentSources NewContentSources;
 			if (UpdateContentSources(HistoryData.ContentSources, NewContentSources))
 			{
 				if (!NewContentSources.HasCollections())
@@ -4329,7 +4329,7 @@ void SContentBrowser::HandleIsHiddenChanged(ICollectionContainer& CollectionCont
 void SContentBrowser::HandleCollectionRemoved(ICollectionContainer& CollectionContainer, const FCollectionNameType& Collection)
 {
 	// Remove Collection from ContentSources.
-	auto UpdateContentSources = [&CollectionContainer, &Collection](const FContentSources& ContentSources, FContentSources& OutNewContentSources)
+	auto UpdateContentSources = [&CollectionContainer, &Collection](const FAssetViewContentSources& ContentSources, FAssetViewContentSources& OutNewContentSources)
 		{
 			const int32 FoundIndex = ContentSources.GetCollections().IndexOfByPredicate([&CollectionContainer, &Collection](const FCollectionRef& CollectionRef)
 				{
@@ -4350,7 +4350,7 @@ void SContentBrowser::HandleCollectionRemoved(ICollectionContainer& CollectionCo
 		};
 
 	{
-		FContentSources NewContentSources;
+		FAssetViewContentSources NewContentSources;
 		if (UpdateContentSources(AssetViewPtr->GetContentSources(), NewContentSources))
 		{
 			AssetViewPtr->SetContentSources(NewContentSources);
@@ -4359,7 +4359,7 @@ void SContentBrowser::HandleCollectionRemoved(ICollectionContainer& CollectionCo
 
 	HistoryManager.RemoveHistoryData([&UpdateContentSources](FHistoryData& HistoryData)
 	{
-		FContentSources NewContentSources;
+		FAssetViewContentSources NewContentSources;
 		if (UpdateContentSources(HistoryData.ContentSources, NewContentSources))
 		{
 			if (!NewContentSources.HasCollections())
@@ -4377,7 +4377,7 @@ void SContentBrowser::HandleCollectionRemoved(ICollectionContainer& CollectionCo
 void SContentBrowser::HandleCollectionRenamed(ICollectionContainer& CollectionContainer, const FCollectionNameType& OriginalCollection, const FCollectionNameType& NewCollection)
 {
 	// Replaces OriginalCollection with NewCollection in ContentSources.
-	auto UpdateContentSources = [&CollectionContainer, &OriginalCollection, &NewCollection](const FContentSources& ContentSources, FContentSources& OutNewContentSources)
+	auto UpdateContentSources = [&CollectionContainer, &OriginalCollection, &NewCollection](const FAssetViewContentSources& ContentSources, FAssetViewContentSources& OutNewContentSources)
 		{
 			const int32 FoundIndex = ContentSources.GetCollections().IndexOfByPredicate([&CollectionContainer, &OriginalCollection](const FCollectionRef& Collection)
 				{
@@ -4398,7 +4398,7 @@ void SContentBrowser::HandleCollectionRenamed(ICollectionContainer& CollectionCo
 		};
 
 	{
-		FContentSources NewContentSources;
+		FAssetViewContentSources NewContentSources;
 		if (UpdateContentSources(AssetViewPtr->GetContentSources(), NewContentSources))
 		{
 			AssetViewPtr->SetContentSources(NewContentSources);
@@ -4407,7 +4407,7 @@ void SContentBrowser::HandleCollectionRenamed(ICollectionContainer& CollectionCo
 
 	HistoryManager.RewriteHistoryData([&UpdateContentSources](FHistoryData& HistoryData)
 		{
-			FContentSources NewContentSources;
+			FAssetViewContentSources NewContentSources;
 			if (UpdateContentSources(HistoryData.ContentSources, NewContentSources))
 			{
 				HistoryData.ContentSources = MoveTemp(NewContentSources);
@@ -4417,7 +4417,7 @@ void SContentBrowser::HandleCollectionRenamed(ICollectionContainer& CollectionCo
 
 void SContentBrowser::HandleCollectionUpdated(ICollectionContainer& CollectionContainer, const FCollectionNameType& Collection)
 {
-	const FContentSources& ContentSources = AssetViewPtr->GetContentSources();
+	const FAssetViewContentSources& ContentSources = AssetViewPtr->GetContentSources();
 
 	// If we're currently viewing the dynamic collection that was updated, make sure our active filter text is up-to-date
 	if (ContentSources.IsDynamicCollection())
diff --git a/Engine/Source/Editor/ContentBrowser/Public/AssetViewContentSources.h b/Engine/Source/Editor/ContentBrowser/Public/AssetViewContentSources.h
new file mode 100644
--- /dev/null
+++ b/Engine/Source/Editor/ContentBrowser/Public/AssetViewContentSources.h
@@ -0,0 +1,76 @@
+// Copyright Epic Games, Inc. All Rights Reserved.
+
+#pragma once
+
+#include "CoreMinimal.h"
+#include "CollectionManagerTypes.h"
+
+struct FAssetViewContentSources
+{
+	DECLARE_DELEGATE_RetVal_OneParam(bool, FOnEnumerateCustomSourceItemDatas, TFunctionRef<bool(class FContentBrowserItemData&&)>)
+
+	FOnEnumerateCustomSourceItemDatas OnEnumerateCustomSourceItemDatas;
+	bool bIncludeVirtualPaths = true;
+
+	FAssetViewContentSources() = default;
+
+	CONTENTBROWSER_API explicit FAssetViewContentSources(FName InVirtualPath);
+	CONTENTBROWSER_API explicit FAssetViewContentSources(TArray<FName> InVirtualPaths);
+
+	CONTENTBROWSER_API explicit FAssetViewContentSources(FCollectionRef InCollection);
+	CONTENTBROWSER_API explicit FAssetViewContentSources(TArray<FCollectionRef> InCollections);
+
+	CONTENTBROWSER_API FAssetViewContentSources(TArray<FName> InVirtualPaths, TArray<FCollectionRef> InCollections);
+
+	FAssetViewContentSources(const FAssetViewContentSources&) = default;
+	FAssetViewContentSources(FAssetViewContentSources&&) = default;
+
+	FAssetViewContentSources& operator=(const FAssetViewContentSources&) = default;
+	FAssetViewContentSources& operator=(FAssetViewContentSources&&) = default;
+
+	bool IsEmpty() const
+	{
+		return VirtualPaths.IsEmpty() && Collections.IsEmpty();
+	}
+
+	bool HasVirtualPaths() const
+	{
+		return !VirtualPaths.IsEmpty();
+	}
+
+	const TArray<FName>& GetVirtualPaths() const
+	{
+		return VirtualPaths;
+	}
+
+	CONTENTBROWSER_API void SetVirtualPath(FName InVirtualPath);
+	CONTENTBROWSER_API void SetVirtualPaths(const TArray<FName>& InVirtualPaths);
+
+	bool HasCollections() const
+	{
+		return !Collections.IsEmpty();
+	}
+
+	const TArray<FCollectionRef>& GetCollections() const
+	{
+		return Collections;
+	}
+
+	CONTENTBROWSER_API void SetCollection(const FCollectionRef& InCollection);
+	CONTENTBROWSER_API void SetCollections(const TArray<FCollectionRef>& InCollections);
+
+	bool IsIncludingVirtualPaths() const
+	{
+		return bIncludeVirtualPaths;
+	}
+
+	CONTENTBROWSER_API bool IsDynamicCollection() const;
+
+	CONTENTBROWSER_API void Reset();
+
+private:
+	void SanitizeCollections();
+
+	TArray<FName> VirtualPaths;
+	TArray<FCollectionRef> Collections;
+};
diff --git a/Engine/Source/Editor/ContentBrowser/Public/ContentSources.h b/Engine/Source/Editor/ContentBrowser/Public/ContentSources.h
deleted file mode 100644
diff --git a/Engine/Source/Editor/ContentBrowser/Public/HistoryManager.h b/Engine/Source/Editor/ContentBrowser/Public/HistoryManager.h
--- a/Engine/Source/Editor/ContentBrowser/Public/HistoryManager.h
+++ b/Engine/Source/Editor/ContentBrowser/Public/HistoryManager.h
@@ -5,7 +5,7 @@
 #include "Containers/Array.h"
 #include "Containers/Set.h"
 #include "Containers/SparseArray.h"
-#include "ContentSources.h"
+#include "AssetViewContentSources.h"
 #include "CoreMinimal.h"
 #include "Delegates/Delegate.h"
 #include "Internationalization/Text.h"
@@ -63,7 +63,7 @@ struct FHistoryData
 	FText HistoryDesc;
 
 	/** The base set of filters on the asset view which includes selected paths and collections */
-	FContentSources ContentSources;
+	FAssetViewContentSources ContentSources;
 
 	UE_DEPRECATED(5.6, "Use ContentSources instead.")
 	FSourcesData SourcesData;
diff --git a/Engine/Source/Editor/ContentBrowser/Public/SAssetView.h b/Engine/Source/Editor/ContentBrowser/Public/SAssetView.h
--- a/Engine/Source/Editor/ContentBrowser/Public/SAssetView.h
+++ b/Engine/Source/Editor/ContentBrowser/Public/SAssetView.h
@@ -17,7 +17,7 @@
 #include "ContentBrowserDelegates.h"
 #include "ContentBrowserItem.h"
 #include "ContentBrowserTelemetry.h"
-#include "ContentSources.h"
+#include "AssetViewContentSources.h"
 #include "CoreMinimal.h"
 #include "Delegates/Delegate.h"
 #include "Experimental/ContentBrowserViewExtender.h"
@@ -209,7 +209,7 @@ public:
 		SLATE_ARGUMENT( bool, bShowPathViewFilters )
 
 		/** The initial base sources filter */
-		SLATE_ARGUMENT( FContentSources, InitialContentSources )
+		SLATE_ARGUMENT( FAssetViewContentSources, InitialContentSources )
 
 		/** The initial base sources filter */
 		SLATE_ARGUMENT_DEPRECATED( FSourcesData, InitialSourcesData, 5.6, "Use InitialContentSources instead." )
@@ -330,13 +330,13 @@ public:
 	void Construct( const FArguments& InArgs );
 
 	/** Changes the base sources for this view */
-	void SetContentSources(const FContentSources& InContentSources);
+	void SetContentSources(const FAssetViewContentSources& InContentSources);
 
 	UE_DEPRECATED(5.6, "Use SetContentSources() instead.")
 	void SetSourcesData(const FSourcesData& InSourcesData);
 
 	/** Returns the sources filter applied to this asset view */
-	const FContentSources& GetContentSources() const;
+	const FAssetViewContentSources& GetContentSources() const;
 
 	UE_DEPRECATED(5.6, "Use GetContentSources() instead.")
 	FSourcesData GetSourcesData() const;
@@ -1007,7 +1007,7 @@ private:
 	TWeakPtr<SFilterList> FilterBar;
 
 	/** The current base source filter for the view */
-	FContentSources ContentSources;
+	FAssetViewContentSources ContentSources;
 	FARFilter BackendFilter;
 	TSharedPtr<FPathPermissionList> AssetClassPermissionList;
 	TSharedPtr<FPathPermissionList> FolderPermissionList;
diff --git a/Engine/Source/Editor/ContentBrowser/Public/SourcesData.h b/Engine/Source/Editor/ContentBrowser/Public/SourcesData.h
--- a/Engine/Source/Editor/ContentBrowser/Public/SourcesData.h
+++ b/Engine/Source/Editor/ContentBrowser/Public/SourcesData.h
@@ -8,7 +8,7 @@
 #include "ICollectionManager.h"
 #include "CollectionManagerModule.h"
 
-UE_DEPRECATED(5.6, "FSourcesData has been deprecated in favor of FContentSources.");
+UE_DEPRECATED(5.6, "FSourcesData has been deprecated in favor of FAssetViewContentSources");
 struct FSourcesData
 {
 	DECLARE_DELEGATE_RetVal_OneParam(bool, FOnEnumerateCustomSourceItemDatas, TFunctionRef<bool(class FContentBrowserItemData&&)>)
diff --git a/Engine/Source/Editor/UMGEditor/Private/Library/SLibraryViewModel.cpp b/Engine/Source/Editor/UMGEditor/Private/Library/SLibraryViewModel.cpp
--- a/Engine/Source/Editor/UMGEditor/Private/Library/SLibraryViewModel.cpp
+++ b/Engine/Source/Editor/UMGEditor/Private/Library/SLibraryViewModel.cpp
@@ -110,7 +110,7 @@ void FWidgetTemplateListViewModel::ConstructListView(TArray<TSharedPtr<FWidgetTe
 		}
 
 		// Provide a dummy invalid virtual path to make sure nothing tries to enumerate root "/"
-		LibraryContentSources = MakeShared<FContentSources>(FName(TEXT("/UMGWidgetTemplateListViewModel")));
+		LibraryContentSources = MakeShared<FAssetViewContentSources>(FName(TEXT("/UMGWidgetTemplateListViewModel")));
 		// Disable any enumerate of virtual path folders
 		LibraryContentSources->bIncludeVirtualPaths = false;
 		// Supply a custom list of source items to display
diff --git a/Engine/Source/Editor/UMGEditor/Private/Library/SLibraryViewModel.h b/Engine/Source/Editor/UMGEditor/Private/Library/SLibraryViewModel.h
--- a/Engine/Source/Editor/UMGEditor/Private/Library/SLibraryViewModel.h
+++ b/Engine/Source/Editor/UMGEditor/Private/Library/SLibraryViewModel.h
@@ -53,7 +53,7 @@ private:
 	TSharedPtr<class FAssetTextFilter> SearchFilter;
 
 	/** Source data for the asset view */
-	TSharedPtr<struct FContentSources> LibraryContentSources;
+	TSharedPtr<struct FAssetViewContentSources> LibraryContentSources;
 
 	/** Source item path storage for asset view, accessed via delegate */
 	TArray<FContentBrowserItemPath> SourceItemPaths;
