[ContentBrowser] Content Browser Source Panel Add/Import Content Submenu fails to display buttons correctly

diff --git a/Engine/Source/Developer/ToolMenus/Private/ToolMenus.cpp b/Engine/Source/Developer/ToolMenus/Private/ToolMenus.cpp
--- a/Engine/Source/Developer/ToolMenus/Private/ToolMenus.cpp
+++ b/Engine/Source/Developer/ToolMenus/Private/ToolMenus.cpp
@@ -549,7 +549,9 @@ TArray<UToolMenu*> UToolMenus::CollectHierarchy(const FName InName, const TMap<F
 void UToolMenus::ApplyStyleToBuilder(FMenuBuilder& InBuilder, const FToolMenuEntryStyle& InStyle)
 {
 	const ISlateStyle* StyleSet = InBuilder.GetStyleSet();
-	if (!StyleSet && InStyle.StyleSet.IsSet())
+
+	// If specified, the provided StyleSet always overrides that specified in the builder
+	if (InStyle.StyleSet.IsSet())
 	{
 		StyleSet = InStyle.StyleSet.GetValue();
 	}
diff --git a/Engine/Source/Developer/ToolMenus/Public/ToolMenuEntry.h b/Engine/Source/Developer/ToolMenus/Public/ToolMenuEntry.h
--- a/Engine/Source/Developer/ToolMenus/Public/ToolMenuEntry.h
+++ b/Engine/Source/Developer/ToolMenus/Public/ToolMenuEntry.h
@@ -21,7 +21,7 @@ struct FKeyEvent;
 struct FToolMenuEntryStyle
 {
 	/** The StyleSet used to create the calling widget. Default is auto-calculated or derived from elsewhere */
-	TOptional<const class ISlateStyle*> StyleSet = nullptr;
+	TOptional<const class ISlateStyle*> StyleSet;
 
 	/** The name of the style (within the StyleSet) to use to create the calling widget. Default is auto-calculated or derived from elsewhere */
 	TOptional<FName> StyleName;
diff --git a/Engine/Source/Editor/ContentBrowser/Private/PathContextMenu.cpp b/Engine/Source/Editor/ContentBrowser/Private/PathContextMenu.cpp
--- a/Engine/Source/Editor/ContentBrowser/Private/PathContextMenu.cpp
+++ b/Engine/Source/Editor/ContentBrowser/Private/PathContextMenu.cpp
@@ -137,7 +137,7 @@ void FPathContextMenu::MakePathViewContextMenu(UToolMenu* Menu)
 				// New Asset (submenu)
 				if (Context->bCanBeModified)
 				{
-					Section.AddSubMenu(
+					FToolMenuEntry& NewAssetEntry = Section.AddSubMenu(
 						"NewAsset",
 						LOCTEXT("AddImportLabel", "Add/Import Content"),
 						NewAssetToolTip,
@@ -147,6 +147,8 @@ void FPathContextMenu::MakePathViewContextMenu(UToolMenu* Menu)
 						false,
 						FSlateIcon(FAppStyle::GetAppStyleSetName(), "Icons.Import")
 					);
+
+					NewAssetEntry.SubMenuData.Style.StyleSet = &UE::ContentBrowser::Private::FContentBrowserStyle::Get();
 				}
 			}
 
