Overload accessing content browser colors with a FContentBrowserItem so that GetInvariantPath is used.

diff --git a/Engine/Source/Editor/ContentBrowser/Private/Experimental/ContentBrowserExtensionUtils.cpp b/Engine/Source/Editor/ContentBrowser/Private/Experimental/ContentBrowserExtensionUtils.cpp
--- a/Engine/Source/Editor/ContentBrowser/Private/Experimental/ContentBrowserExtensionUtils.cpp
+++ b/Engine/Source/Editor/ContentBrowser/Private/Experimental/ContentBrowserExtensionUtils.cpp
@@ -9,18 +9,35 @@
 
 namespace UE::Editor::ContentBrowser::ExtensionUtils
 {
-	TOptional<FLinearColor> GetFolderColor(const FContentBrowserItemPath& FolderPath)
+	TOptional<FLinearColor> CheckAndGetCollectionColor(const FString& InVirtualPath)
 	{
 		TSharedPtr<ICollectionContainer> CollectionContainer;
 		FName CollectionName;
 		ECollectionShareType::Type CollectionFolderShareType = ECollectionShareType::CST_All;
 
-		if (ContentBrowserUtils::IsCollectionPath(FolderPath.GetVirtualPathString(), &CollectionContainer , &CollectionName, &CollectionFolderShareType))
+		if (ContentBrowserUtils::IsCollectionPath(InVirtualPath, &CollectionContainer , &CollectionName, &CollectionFolderShareType))
+		{
+			return CollectionViewUtils::GetCustomColor(CollectionContainer.Get(), CollectionName, CollectionFolderShareType);
+		}
+		
+		return TOptional<FLinearColor>();
+	} 
+
+	TOptional<FLinearColor> GetFolderColor(const FContentBrowserItem& FolderItem)
+	{
+		if (TOptional<FLinearColor> CollectionColor = CheckAndGetCollectionColor(FolderItem.GetVirtualPath().ToString()))
+		{
+			return CollectionColor;
+		}
+		
+		return ContentBrowserUtils::GetPathColor(FolderItem.GetInvariantPath().ToString());
+	}
+
+	TOptional<FLinearColor> GetFolderColor(const FContentBrowserItemPath& FolderPath)
+	{
+		if (TOptional<FLinearColor> CollectionColor = CheckAndGetCollectionColor(FolderPath.GetVirtualPathString()))
 		{
-			if (TOptional<FLinearColor> Color = CollectionViewUtils::GetCustomColor(CollectionContainer.Get(), CollectionName, CollectionFolderShareType))
-			{
-				return Color;
-			}
+			return CollectionColor;
 		}
 
 		if (FolderPath.HasInternalPath())
diff --git a/Engine/Source/Editor/ContentBrowser/Private/SAssetDialog.cpp b/Engine/Source/Editor/ContentBrowser/Private/SAssetDialog.cpp
--- a/Engine/Source/Editor/ContentBrowser/Private/SAssetDialog.cpp
+++ b/Engine/Source/Editor/ContentBrowser/Private/SAssetDialog.cpp
@@ -1037,8 +1037,7 @@ TSharedRef<SWidget> SAssetDialog::OnGetCrumbDelimiterContent(const FString& Crum
 
 			if (FolderBrushName != NAME_None)
 			{
-				const FContentBrowserItemPath FolderPath = SubItem.GetItemPath();
-				FLinearColor FolderColor = UE::Editor::ContentBrowser::ExtensionUtils::GetFolderColor(FolderPath).Get(ContentBrowserUtils::GetDefaultColor());
+				FLinearColor FolderColor = UE::Editor::ContentBrowser::ExtensionUtils::GetFolderColor(SubItem).Get(ContentBrowserUtils::GetDefaultColor());
 				
 				FMenuEntryParams Params;
 				Params.EntryWidget = ContentBrowserUtils::GetFolderWidgetForNavigationBar(EntryName, FolderBrushName, FolderColor);
diff --git a/Engine/Source/Editor/ContentBrowser/Private/SContentBrowser.cpp b/Engine/Source/Editor/ContentBrowser/Private/SContentBrowser.cpp
--- a/Engine/Source/Editor/ContentBrowser/Private/SContentBrowser.cpp
+++ b/Engine/Source/Editor/ContentBrowser/Private/SContentBrowser.cpp
@@ -3042,8 +3042,7 @@ TSharedRef<SWidget> SContentBrowser::OnGetCrumbDelimiterContent(const FString& C
 
 				if (FolderBrushName != NAME_None)
 				{
-					const FContentBrowserItemPath FolderPath = SubItem.GetItemPath();
-					FLinearColor FolderColor = UE::Editor::ContentBrowser::ExtensionUtils::GetFolderColor(FolderPath).Get(ContentBrowserUtils::GetDefaultColor());
+					FLinearColor FolderColor = UE::Editor::ContentBrowser::ExtensionUtils::GetFolderColor(SubItem).Get(ContentBrowserUtils::GetDefaultColor());
 						
 					FMenuEntryParams Params;
 					Params.EntryWidget = ContentBrowserUtils::GetFolderWidgetForNavigationBar(EntryName, FolderBrushName, FolderColor);
diff --git a/Engine/Source/Editor/ContentBrowser/Public/Experimental/ContentBrowserExtensionUtils.h b/Engine/Source/Editor/ContentBrowser/Public/Experimental/ContentBrowserExtensionUtils.h
--- a/Engine/Source/Editor/ContentBrowser/Public/Experimental/ContentBrowserExtensionUtils.h
+++ b/Engine/Source/Editor/ContentBrowser/Public/Experimental/ContentBrowserExtensionUtils.h
@@ -5,10 +5,14 @@
 #include "Containers/UnrealString.h"
 #include "Math/Color.h"
 
+struct FContentBrowserItem;
 struct FContentBrowserItemPath;
 
 namespace UE::Editor::ContentBrowser::ExtensionUtils
 {
+	// Get the custom color of the given folder (if any). Prefer using this over raw FName variants.
+	CONTENTBROWSER_API TOptional<FLinearColor> GetFolderColor(const FContentBrowserItem& FolderItem);
+
 	// Get the custom color of the given folder path (if any)
 	CONTENTBROWSER_API TOptional<FLinearColor> GetFolderColor(const FContentBrowserItemPath& FolderPath);
 
