Content Browser - Common Folder Operations

#rnx

diff --git a/Engine/Source/Editor/ContentBrowser/Private/ContentBrowserSingleton.cpp b/Engine/Source/Editor/ContentBrowser/Private/ContentBrowserSingleton.cpp
--- a/Engine/Source/Editor/ContentBrowser/Private/ContentBrowserSingleton.cpp
+++ b/Engine/Source/Editor/ContentBrowser/Private/ContentBrowserSingleton.cpp
@@ -1446,4 +1446,17 @@ void FContentBrowserSingleton::SetSearchText(const FText& NewText)
 	}
 }
 
+void FContentBrowserSingleton::ToggleFolderFavorite(const TArray<FString>& FolderPaths)
+{
+	if (!HasPrimaryContentBrowser())
+	{
+		ChooseNewPrimaryBrowser();
+	}
+
+	if (PrimaryContentBrowser.IsValid())
+	{
+		PrimaryContentBrowser.Pin()->ToggleFolderFavorite(FolderPaths);
+	}
+}
+
 #undef LOCTEXT_NAMESPACE
diff --git a/Engine/Source/Editor/ContentBrowser/Private/ContentBrowserSingleton.h b/Engine/Source/Editor/ContentBrowser/Private/ContentBrowserSingleton.h
--- a/Engine/Source/Editor/ContentBrowser/Private/ContentBrowserSingleton.h
+++ b/Engine/Source/Editor/ContentBrowser/Private/ContentBrowserSingleton.h
@@ -123,6 +123,7 @@ PRAGMA_ENABLE_DEPRECATION_WARNINGS
 	virtual void ModifyCustomTextFilterByLabel(const FCustomTextFilterData& NewFilterData, const FText& FilterLabel) override;
 	virtual bool IsAssetViewDoneFiltering() override;
 	virtual void SetSearchText(const FText& NewText) override;
+	virtual void ToggleFolderFavorite(const TArray<FString>& FolderPaths) override;
 
 	/** Broadcast that the favorites have changed. */
 	void BroadcastFavoritesChanged() const;
diff --git a/Engine/Source/Editor/ContentBrowser/Private/Experimental/ContentBrowserExtensionUtils.cpp b/Engine/Source/Editor/ContentBrowser/Private/Experimental/ContentBrowserExtensionUtils.cpp
--- a/Engine/Source/Editor/ContentBrowser/Private/Experimental/ContentBrowserExtensionUtils.cpp
+++ b/Engine/Source/Editor/ContentBrowser/Private/Experimental/ContentBrowserExtensionUtils.cpp
@@ -58,4 +58,9 @@ namespace UE::Editor::ContentBrowser::ExtensionUtils
 	{
 		ContentBrowserUtils::SetPathColor(FolderPath.ToString(), FolderColor);
 	}
+
+	bool IsFolderFavorite(const FString& FolderPath)
+	{
+		return ContentBrowserUtils::IsFavoriteFolder(FContentBrowserItemPath(FolderPath, EContentBrowserPathType::Virtual));
+	}
 }
diff --git a/Engine/Source/Editor/ContentBrowser/Private/SAssetView.cpp b/Engine/Source/Editor/ContentBrowser/Private/SAssetView.cpp
--- a/Engine/Source/Editor/ContentBrowser/Private/SAssetView.cpp
+++ b/Engine/Source/Editor/ContentBrowser/Private/SAssetView.cpp
@@ -2403,6 +2403,11 @@ bool SAssetView::HasThumbnailsPendingUpdate() const
 	return bPendingUpdateThumbnails;
 }
 
+bool SAssetView::HasDeferredItemToCreate() const
+{
+	return DeferredItemToCreate.IsValid();
+}
+
 void SAssetView::ProcessItemsPendingFilter(const double TickStartTime)
 {
 	const double ProcessItemsPendingFilterStartTime = FPlatformTime::Seconds();
diff --git a/Engine/Source/Editor/ContentBrowser/Public/Experimental/ContentBrowserExtensionUtils.h b/Engine/Source/Editor/ContentBrowser/Public/Experimental/ContentBrowserExtensionUtils.h
--- a/Engine/Source/Editor/ContentBrowser/Public/Experimental/ContentBrowserExtensionUtils.h
+++ b/Engine/Source/Editor/ContentBrowser/Public/Experimental/ContentBrowserExtensionUtils.h
@@ -21,4 +21,7 @@ namespace UE::Editor::ContentBrowser::ExtensionUtils
 
 	// Set a custom color for the given folder path
 	CONTENTBROWSER_API void SetFolderColor(const FName& FolderPath, const FLinearColor& FolderColor);
+
+	// Returns if this folder has been marked as a favorite folder
+	CONTENTBROWSER_API bool IsFolderFavorite(const FString& FolderPath);
 }
diff --git a/Engine/Source/Editor/ContentBrowser/Public/IContentBrowserSingleton.h b/Engine/Source/Editor/ContentBrowser/Public/IContentBrowserSingleton.h
--- a/Engine/Source/Editor/ContentBrowser/Public/IContentBrowserSingleton.h
+++ b/Engine/Source/Editor/ContentBrowser/Public/IContentBrowserSingleton.h
@@ -808,4 +808,7 @@ public:
 
 	/** Sets the AssetSearchBox text */
 	virtual void SetSearchText(const FText& NewText) = 0;
+
+	/** Toggle the Favorite status of the given folder list */
+	virtual void ToggleFolderFavorite(const TArray<FString>& FolderPaths) = 0;
 };
diff --git a/Engine/Source/Editor/ContentBrowser/Public/SAssetView.h b/Engine/Source/Editor/ContentBrowser/Public/SAssetView.h
--- a/Engine/Source/Editor/ContentBrowser/Public/SAssetView.h
+++ b/Engine/Source/Editor/ContentBrowser/Public/SAssetView.h
@@ -502,6 +502,9 @@ public:
 	/** True if we have thumbnails that are pending update */
 	bool HasThumbnailsPendingUpdate() const;
 
+	/** True if we have deferred item to create */
+	bool HasDeferredItemToCreate() const;
+
 	/** Get the underlying sort manager */
 	TWeakPtr<FAssetViewSortManager> GetSortManager() const;
 
