Expose DockContentBrowserDrawer through the public Content Browser API

diff --git a/Engine/Source/Editor/ContentBrowser/Private/ContentBrowserSingleton.h b/Engine/Source/Editor/ContentBrowser/Private/ContentBrowserSingleton.h
--- a/Engine/Source/Editor/ContentBrowser/Private/ContentBrowserSingleton.h
+++ b/Engine/Source/Editor/ContentBrowser/Private/ContentBrowserSingleton.h
@@ -124,6 +124,7 @@ PRAGMA_ENABLE_DEPRECATION_WARNINGS
 	virtual bool IsAssetViewDoneFiltering() override;
 	virtual void SetSearchText(const FText& NewText) override;
 	virtual void ToggleFolderFavorite(const TArray<FString>& FolderPaths) override;
+	virtual void DockContentBrowserDrawer() override;
 
 	/** Broadcast that the favorites have changed. */
 	void BroadcastFavoritesChanged() const;
@@ -143,9 +144,6 @@ PRAGMA_ENABLE_DEPRECATION_WARNINGS
 	/** Single storage location for content browser favorites */
 	TArray<FString> FavoriteFolderPaths;
 
-	/** Docks the current content browser drawer as a tabbed content browser in a layout */
-	void DockContentBrowserDrawer();
-
 private:
 
 	/** Util to get or create the content browser that should be used by the various Sync functions */
diff --git a/Engine/Source/Editor/ContentBrowser/Public/IContentBrowserSingleton.h b/Engine/Source/Editor/ContentBrowser/Public/IContentBrowserSingleton.h
--- a/Engine/Source/Editor/ContentBrowser/Public/IContentBrowserSingleton.h
+++ b/Engine/Source/Editor/ContentBrowser/Public/IContentBrowserSingleton.h
@@ -811,4 +811,7 @@ public:
 
 	/** Toggle the Favorite status of the given folder list */
 	virtual void ToggleFolderFavorite(const TArray<FString>& FolderPaths) = 0;
+
+	/** Docks the current content browser drawer as a tabbed content browser in a layout */
+	virtual void DockContentBrowserDrawer() = 0;
 };
