Content Browser:

Add a new Content Sources API that can override the look (widgets) and behavior of the content browser depending on the currently active content source
Add a new TEDS UI widget to display a table viewer when a content source is enabled (with the query and other init params driven by the content source)
Add an SLegacyContentSource widget to support displaying the current content browser widgets (SPathView, SAssetView etc)
Add a cvar to display a vertical toolbar to swap between content sources (ContentBrowser.UI.ShowContentSourcesBar)
Add a cvar to register a test content source that displays the same query as the Scene Outliner (TEDS.UI.EnableTedsContentBrowser)

diff --git a/Engine/Plugins/Experimental/EditorDataStorageFeatures/Source/TedsContentBrowser/Private/TedsContentBrowserModule.cpp b/Engine/Plugins/Experimental/EditorDataStorageFeatures/Source/TedsContentBrowser/Private/TedsContentBrowserModule.cpp
--- a/Engine/Plugins/Experimental/EditorDataStorageFeatures/Source/TedsContentBrowser/Private/TedsContentBrowserModule.cpp
+++ b/Engine/Plugins/Experimental/EditorDataStorageFeatures/Source/TedsContentBrowser/Private/TedsContentBrowserModule.cpp
@@ -4,7 +4,10 @@
 
 #include "Framework/Application/SlateApplication.h"
 #include "ContentBrowserModule.h"
+#include "IContentBrowserSingleton.h"
 #include "Columns/SlateDelegateColumns.h"
+#include "Elements/Columns/TypedElementFolderColumns.h"
+#include "Elements/Columns/TypedElementMiscColumns.h"
 #include "Elements/Common/TypedElementQueryTypes.h"
 #include "Elements/Common/EditorDataStorageFeatures.h"
 #include "Elements/Framework/TypedElementIndexHasher.h"
@@ -14,8 +17,10 @@
 #include "TedsAlertColumns.h"
 #include "TedsAssetDataColumns.h"
 #include "TedsRowViewNode.h"
-#include "Elements/Columns/TypedElementFolderColumns.h"
-#include "Elements/Columns/TypedElementMiscColumns.h"
+#include "Elements/Columns/TypedElementLabelColumns.h"
+#include "Elements/Columns/TypedElementTypeInfoColumns.h"
+#include "Elements/Columns/TypedElementVisibilityColumns.h"
+#include "Elements/Framework/TypedElementQueryBuilder.h"
 #include "Widgets/STedsTableViewer.h"
 #include "Widgets/STedsTileViewer.h"
 #include "Widgets/Views/SListView.h"
@@ -27,7 +32,6 @@ namespace UE::Editor::ContentBrowser
 	using namespace DataStorage;
 
 	static bool bEnableTedsContentBrowser = false;
-	
 	static FAutoConsoleVariableRef CVarUseTEDSOutliner(
 		TEXT("TEDS.UI.EnableTedsContentBrowser"),
 		bEnableTedsContentBrowser,
@@ -64,6 +68,41 @@ namespace UE::Editor::ContentBrowser
 	{
 		return MakeShared<FTedsContentBrowserViewExtender>();
 	}
+	
+	static bool bEnableTestContentSource = false;
+	static FAutoConsoleVariableRef CVarEnableTestContentSource(
+		TEXT("TEDS.UI.EnableTestContentSource"),
+		bEnableTestContentSource,
+		TEXT("Add a test content source to the Content Browser")
+		, FConsoleVariableDelegate::CreateLambda([](IConsoleVariable* /*CVar*/)
+		{
+			FTedsContentBrowserModule& TedsContentBrowserModule = FModuleManager::Get().GetModuleChecked<FTedsContentBrowserModule>(TEXT("TedsContentBrowser"));
+			
+			if(bEnableTestContentSource)
+			{
+				TedsContentBrowserModule.RegisterTestContentSource();
+			}
+			else
+			{
+				TedsContentBrowserModule.UnregisterTestContentSource();
+			}
+		}));
+
+	static const FName TestContentSourceName("TestContentSource");
+
+	void FTedsContentBrowserModule::RegisterTestContentSource()
+	{
+		IContentBrowserSingleton::Get().RegisterContentSourceFactory(TestContentSourceName, IContentBrowserSingleton::FContentSourceFactory::CreateLambda(
+			[]() -> TSharedRef<IContentSource>
+		{
+			return MakeShared<FTestContentSource>();
+		}));
+	}
+
+	void FTedsContentBrowserModule::UnregisterTestContentSource()
+	{
+		IContentBrowserSingleton::Get().UnregisterContentSourceFactory(TestContentSourceName);
+	}
 
 	void FTedsContentBrowserModule::StartupModule()
 	{
@@ -94,7 +133,7 @@ namespace UE::Editor::ContentBrowser
 
 			AddRow(Item);
 		}
-		
+
 		RowQueryStack->ResetView(FRowHandleArrayView(Rows.GetData(), Rows.Num(),
 			FRowHandleArrayView::EFlags::IsSorted | FRowHandleArrayView::EFlags::IsUnique));
 	}
@@ -235,6 +274,40 @@ namespace UE::Editor::ContentBrowser
 		}
 	}
 
+	FName FTestContentSource::GetName()
+	{
+		return FName("TestContentSource");
+	}
+
+	FText FTestContentSource::GetDisplayName()
+	{
+		return LOCTEXT("TestContentSource", "Outliner");
+	}
+
+	FSlateIcon FTestContentSource::GetIcon()
+	{
+		return FSlateIcon(FAppStyle::GetAppStyleSetName(), "LevelEditor.Tabs.Outliner");
+	}
+
+	void FTestContentSource::GetAssetViewInitParams(FTableViewerInitParams& OutInitParams)
+	{
+		using namespace UE::Editor::DataStorage::Queries;
+
+		
+		OutInitParams.QueryDescription = Select()
+			.Where()
+				.All<FTypedElementClassTypeInfoColumn>() // The test content source is simply looking at all rows with type info for now
+			.Compile();
+
+		// A few columns shown in the Outliner as a sample
+		OutInitParams.Columns = { FTypedElementLabelColumn::StaticStruct(), FTypedElementClassTypeInfoColumn::StaticStruct(),
+			FTedsAlertColumn::StaticStruct(), FTedsChildAlertColumn::StaticStruct() };
+
+		// Same widget purposes as the Outliner for now
+		OutInitParams.WidgetPurposes = { FName("SceneOutliner.Cell"), FName("General.Cell") };
+
+	}
+
 	FTedsContentBrowserViewExtender::FTedsContentBrowserViewExtender()
 	{
 		DataStorage = GetMutableDataStorageFeature<IEditorDataStorageProvider>(StorageFeatureName);
diff --git a/Engine/Plugins/Experimental/EditorDataStorageFeatures/Source/TedsContentBrowser/Private/TedsContentBrowserModule.h b/Engine/Plugins/Experimental/EditorDataStorageFeatures/Source/TedsContentBrowser/Private/TedsContentBrowserModule.h
--- a/Engine/Plugins/Experimental/EditorDataStorageFeatures/Source/TedsContentBrowser/Private/TedsContentBrowserModule.h
+++ b/Engine/Plugins/Experimental/EditorDataStorageFeatures/Source/TedsContentBrowser/Private/TedsContentBrowserModule.h
@@ -2,6 +2,7 @@
 
 #pragma once
 
+#include "ContentSources/IContentSource.h"
 #include "Elements/Common/TypedElementHandles.h"
 #include "Experimental/ContentBrowserViewExtender.h"
 #include "Modules/ModuleInterface.h"
@@ -13,6 +14,7 @@ namespace UE::Editor::DataStorage
 {
 	class FQueryStackNode_RowView;
 	class ITableViewer;
+
 	namespace QueryStack
 	{
 		class FRowViewNode;
@@ -23,6 +25,18 @@ class FAssetViewItem;
 
 namespace UE::Editor::ContentBrowser
 {
+
+	// A test content source that currently displays the query used by the Outliner
+	class FTestContentSource : public IContentSource
+	{
+	public:
+		virtual ~FTestContentSource() override = default;
+
+		virtual FName GetName() override;
+		virtual FText GetDisplayName() override;
+		virtual FSlateIcon GetIcon() override;
+		virtual void GetAssetViewInitParams(FTableViewerInitParams& OutInitParams) override;
+	};
 	
 	// A custom view for the content browser that uses the TEDS asset registry integration to display rows with widgets using TEDS UI
 	class FTedsContentBrowserViewExtender : public IContentBrowserViewExtender
@@ -106,6 +120,9 @@ namespace UE::Editor::ContentBrowser
 		
 		static TSharedPtr<IContentBrowserViewExtender> CreateContentBrowserViewExtender();
 
+		void RegisterTestContentSource();
+		void UnregisterTestContentSource();
+
 		// IModuleInterface interface
 		virtual void StartupModule() override;
 		virtual void ShutdownModule() override;
diff --git a/Engine/Plugins/Experimental/EditorDataStorageFeatures/Source/TedsContentBrowser/Private/Widgets/TedsContentBrowserAssetViewWidget.cpp b/Engine/Plugins/Experimental/EditorDataStorageFeatures/Source/TedsContentBrowser/Private/Widgets/TedsContentBrowserAssetViewWidget.cpp
new file mode 100644
--- /dev/null
+++ b/Engine/Plugins/Experimental/EditorDataStorageFeatures/Source/TedsContentBrowser/Private/Widgets/TedsContentBrowserAssetViewWidget.cpp
@@ -0,0 +1,140 @@
+// Copyright Epic Games, Inc. All Rights Reserved.
+
+#include "TedsContentBrowserAssetViewWidget.h"
+
+#include "TedsRowViewNode.h"
+#include "ContentSources/IContentSource.h"
+#include "ContentSources/Columns/ContentSourcesColumns.h"
+#include "Elements/Columns/TypedElementLabelColumns.h"
+#include "Elements/Columns/TypedElementTypeInfoColumns.h"
+#include "Widgets/STedsTableViewer.h"
+
+#define LOCTEXT_NAMESPACE "ContentBrowserAssetViewWidget"
+
+// Wrapper widget around the table viewer so we can manage the lifetime of the query and query stack manually for now
+class STableViewerWrapper : public SCompoundWidget
+{
+public:
+	SLATE_BEGIN_ARGS(STableViewerWrapper)
+		: _InitParams()
+		{}
+		SLATE_ARGUMENT(UE::Editor::ContentBrowser::FTableViewerInitParams, InitParams)
+		SLATE_ARGUMENT(IEditorDataStorageProvider*, Storage)
+	SLATE_END_ARGS()
+
+	void Construct(const FArguments& Arguments)
+	{
+		using namespace UE::Editor::DataStorage;
+		Storage = Arguments._Storage;
+
+		FQueryDescription QueryDescription = Arguments._InitParams.QueryDescription;
+		
+		QueryHandle = Storage->RegisterQuery(MoveTemp(QueryDescription));
+		
+		RowView = MakeShared<QueryStack::FRowViewNode>(FRowHandleArrayView(Rows.GetData(), Rows.Num(),
+			FRowHandleArrayView::EFlags::IsSorted | FRowHandleArrayView::EFlags::IsUnique));
+
+		ChildSlot
+		[
+			SNew(UE::Editor::DataStorage::STedsTableViewer)
+					.QueryStack(RowView)
+					.CellWidgetPurposes(Arguments._InitParams.WidgetPurposes)
+					.Columns(Arguments._InitParams.Columns)
+					.ListSelectionMode(ESelectionMode::Multi)
+		];
+	}
+
+	virtual void Tick(const FGeometry& AllottedGeometry, const double InCurrentTime, const float InDeltaTime) override
+	{
+		// On tick, we run the query again and update the query stack if the rows changed
+		
+		TSet<UE::Editor::DataStorage::RowHandle> NewRows;
+
+		using namespace UE::Editor::DataStorage;
+
+		Storage->RunQuery(QueryHandle,
+			Queries::CreateDirectQueryCallbackBinding([&NewRows]
+				(const IEditorDataStorageProvider::IDirectQueryContext& Context, const RowHandle*)
+				{
+					NewRows.Append(Context.GetRowHandles());
+				}));
+
+		// Check if the two sets are equal, i.e no changes and no need to update the table viewer
+		const bool bSetsEqual = (Rows_Set.Num() == NewRows.Num()) && Rows_Set.Includes(NewRows);
+
+		if (!bSetsEqual)
+		{
+			Rows_Set = NewRows;
+			Rows = Rows_Set.Array();
+			
+			RowView->ResetView(FRowHandleArrayView(Rows.GetData(), Rows.Num(),
+				FRowHandleArrayView::EFlags::IsSorted | FRowHandleArrayView::EFlags::IsUnique));
+		}
+	}
+
+	virtual ~STableViewerWrapper() override
+	{
+		Storage->UnregisterQuery(QueryHandle);
+	}
+
+private:
+	UE::Editor::DataStorage::QueryHandle QueryHandle = UE::Editor::DataStorage::InvalidQueryHandle;
+	TArray<UE::Editor::DataStorage::RowHandle> Rows;
+	TSet<UE::Editor::DataStorage::RowHandle> Rows_Set;
+	TSharedPtr<UE::Editor::DataStorage::QueryStack::FRowViewNode> RowView;
+	IEditorDataStorageProvider* Storage = nullptr;
+	
+};
+
+namespace TedsContentBrowser::Private
+{
+	static const FName WidgetPurpose(TEXT("ContentBrowser.AssetView"));
+}
+
+void UContentBrowserAssetViewWidgetFactory::RegisterWidgetPurposes(IEditorDataStorageUiProvider& DataStorageUi) const
+{
+	DataStorageUi.RegisterWidgetPurpose(TedsContentBrowser::Private::WidgetPurpose, IEditorDataStorageUiProvider::EPurposeType::UniqueByName,
+		LOCTEXT("ContentBrowserAssetView_PurposeDescription", "Widget that displays a table viewer in the content browser"));
+}
+
+void UContentBrowserAssetViewWidgetFactory::RegisterWidgetConstructors(IEditorDataStorageProvider& DataStorage,
+	IEditorDataStorageUiProvider& DataStorageUi) const
+{
+	using namespace UE::Editor::DataStorage::Queries;
+	DataStorageUi.RegisterWidgetFactory<FContentBrowserAssetViewWidgetConstructor>(TedsContentBrowser::Private::WidgetPurpose);
+}
+
+FContentBrowserAssetViewWidgetConstructor::FContentBrowserAssetViewWidgetConstructor()
+	: FSimpleWidgetConstructor(StaticStruct())
+{
+}
+
+FContentBrowserAssetViewWidgetConstructor::FContentBrowserAssetViewWidgetConstructor(const UScriptStruct* TypeInfo)
+	: FSimpleWidgetConstructor(TypeInfo)
+{
+}
+
+TSharedPtr<SWidget> FContentBrowserAssetViewWidgetConstructor::CreateWidget(IEditorDataStorageProvider* DataStorage,
+	IEditorDataStorageUiProvider* DataStorageUi, UE::Editor::DataStorage::RowHandle TargetRow, UE::Editor::DataStorage::RowHandle WidgetRow,
+	const UE::Editor::DataStorage::FMetaDataView& Arguments)
+{
+	using namespace UE::Editor::DataStorage::Queries;
+
+	if (FContentSourceColumn* ContentSourceColumn = DataStorage->GetColumn<FContentSourceColumn>(WidgetRow))
+	{
+		if (TSharedPtr<UE::Editor::ContentBrowser::IContentSource> ContentSourcePtr = ContentSourceColumn->ContentSource.Pin())
+		{
+			UE::Editor::ContentBrowser::FTableViewerInitParams InitParams;
+			ContentSourcePtr->GetAssetViewInitParams(InitParams);
+			
+			return SNew(STableViewerWrapper)
+			.InitParams(InitParams)
+			.Storage(DataStorage);
+		}
+		
+	}
+
+	return SNullWidget::NullWidget;
+}
+
+#undef LOCTEXT_NAMESPACE
diff --git a/Engine/Plugins/Experimental/EditorDataStorageFeatures/Source/TedsContentBrowser/Private/Widgets/TedsContentBrowserAssetViewWidget.h b/Engine/Plugins/Experimental/EditorDataStorageFeatures/Source/TedsContentBrowser/Private/Widgets/TedsContentBrowserAssetViewWidget.h
new file mode 100644
--- /dev/null
+++ b/Engine/Plugins/Experimental/EditorDataStorageFeatures/Source/TedsContentBrowser/Private/Widgets/TedsContentBrowserAssetViewWidget.h
@@ -0,0 +1,48 @@
+// Copyright Epic Games, Inc. All Rights Reserved.
+
+#pragma once
+
+#include "Elements/Interfaces/TypedElementDataStorageFactory.h"
+#include "Elements/Interfaces/TypedElementDataStorageUiInterface.h"
+#include "Internationalization/Text.h"
+#include "UObject/ObjectMacros.h"
+
+#include "TedsContentBrowserAssetViewWidget.generated.h"
+
+class IEditorDataStorageProvider;
+class UScriptStruct;
+
+UCLASS()
+class UContentBrowserAssetViewWidgetFactory : public UEditorDataStorageFactory
+{
+	GENERATED_BODY()
+
+public:
+	~UContentBrowserAssetViewWidgetFactory() override = default;
+
+	TEDSCONTENTBROWSER_API virtual void RegisterWidgetPurposes(IEditorDataStorageUiProvider& DataStorageUi) const override;
+
+	TEDSCONTENTBROWSER_API void RegisterWidgetConstructors(IEditorDataStorageProvider& DataStorage,
+		IEditorDataStorageUiProvider& DataStorageUi) const override;
+};
+
+// Default asset view widget shown by a content source - the query and other init params are populated by the content source itself
+USTRUCT()
+struct FContentBrowserAssetViewWidgetConstructor : public FSimpleWidgetConstructor
+{
+	GENERATED_BODY()
+
+public:
+	TEDSCONTENTBROWSER_API FContentBrowserAssetViewWidgetConstructor();
+	TEDSCONTENTBROWSER_API explicit FContentBrowserAssetViewWidgetConstructor(const UScriptStruct* TypeInfo);
+	~FContentBrowserAssetViewWidgetConstructor() override = default;
+
+	TEDSCONTENTBROWSER_API virtual TSharedPtr<SWidget> CreateWidget(
+		IEditorDataStorageProvider* DataStorage,
+		IEditorDataStorageUiProvider* DataStorageUi,
+		UE::Editor::DataStorage::RowHandle TargetRow,
+		UE::Editor::DataStorage::RowHandle WidgetRow, 
+		const UE::Editor::DataStorage::FMetaDataView& Arguments) override;
+
+private:
+};
diff --git a/Engine/Source/Editor/ContentBrowser/ContentBrowser.Build.cs b/Engine/Source/Editor/ContentBrowser/ContentBrowser.Build.cs
--- a/Engine/Source/Editor/ContentBrowser/ContentBrowser.Build.cs
+++ b/Engine/Source/Editor/ContentBrowser/ContentBrowser.Build.cs
@@ -45,7 +45,8 @@ public class ContentBrowser : ModuleRules
 				"ToolMenus",
 				"StatusBar",
 				"ToolWidgets",
-				"TelemetryUtils"
+				"TelemetryUtils",
+				"TypedElementFramework"
 			}
 		);
 
diff --git a/Engine/Source/Editor/ContentBrowser/Private/ContentBrowserSingleton.cpp b/Engine/Source/Editor/ContentBrowser/Private/ContentBrowserSingleton.cpp
--- a/Engine/Source/Editor/ContentBrowser/Private/ContentBrowserSingleton.cpp
+++ b/Engine/Source/Editor/ContentBrowser/Private/ContentBrowserSingleton.cpp
@@ -16,6 +16,7 @@
 #include "ContentBrowserModule.h"
 #include "ContentBrowserStyle.h"
 #include "ContentBrowserUtils.h"
+#include "ContentSources/IContentSource.h"
 #include "CoreGlobals.h"
 #include "Delegates/Delegate.h"
 #include "Editor.h"
@@ -1459,4 +1460,29 @@ void FContentBrowserSingleton::ToggleFolderFavorite(const TArray<FString>& Folde
 	}
 }
 
+void FContentBrowserSingleton::RegisterContentSourceFactory(const FName& InContentSourceName, const FContentSourceFactory& InContentSourceFactory)
+{
+	ContentSourceFactories.Emplace(InContentSourceName, InContentSourceFactory);
+	ContentSourceFactoriesChangedEvent.Broadcast();
+}
+
+void FContentBrowserSingleton::UnregisterContentSourceFactory(const FName& InContentSourceName)
+{
+	ContentSourceFactories.Remove(InContentSourceName);
+	ContentSourceFactoriesChangedEvent.Broadcast();
+}
+
+void FContentBrowserSingleton::ForEachContentSourceFactory(const TFunctionRef<void(const FName&, const FContentSourceFactory&)>& InFunction)
+{
+	for(const TPair<FName, FContentSourceFactory>& Pair : ContentSourceFactories)
+	{
+		InFunction(Pair.Key, Pair.Value);
+	}
+}
+
+IContentBrowserSingleton::FOnContentSourceFactoriesChanged& FContentBrowserSingleton::OnContentSourceFactoriesChanged()
+{
+	return ContentSourceFactoriesChangedEvent;
+}
+
 #undef LOCTEXT_NAMESPACE
diff --git a/Engine/Source/Editor/ContentBrowser/Private/ContentBrowserSingleton.h b/Engine/Source/Editor/ContentBrowser/Private/ContentBrowserSingleton.h
--- a/Engine/Source/Editor/ContentBrowser/Private/ContentBrowserSingleton.h
+++ b/Engine/Source/Editor/ContentBrowser/Private/ContentBrowserSingleton.h
@@ -125,6 +125,10 @@ PRAGMA_ENABLE_DEPRECATION_WARNINGS
 	virtual void SetSearchText(const FText& NewText) override;
 	virtual void ToggleFolderFavorite(const TArray<FString>& FolderPaths) override;
 	virtual void DockContentBrowserDrawer() override;
+	virtual void RegisterContentSourceFactory(const FName& InContentSourceName, const FContentSourceFactory& InContentSourceFactory) override;
+	virtual void UnregisterContentSourceFactory(const FName& InContentSourceName) override;
+	virtual void ForEachContentSourceFactory(const TFunctionRef<void(const FName&, const FContentSourceFactory&)>& InFunction) override;
+	virtual FOnContentSourceFactoriesChanged& OnContentSourceFactoriesChanged() override;
 
 	/** Broadcast that the favorites have changed. */
 	void BroadcastFavoritesChanged() const;
@@ -220,4 +224,7 @@ private:
 
 	DECLARE_MULTICAST_DELEGATE(FOnFavoritesChangedDelegate);
 	FOnFavoritesChangedDelegate OnFavoritesChanged;
+
+	TMap<FName, FContentSourceFactory> ContentSourceFactories;
+	FOnContentSourceFactoriesChanged ContentSourceFactoriesChangedEvent;
 };
diff --git a/Engine/Source/Editor/ContentBrowser/Private/ContentBrowserStyle.cpp b/Engine/Source/Editor/ContentBrowser/Private/ContentBrowserStyle.cpp
--- a/Engine/Source/Editor/ContentBrowser/Private/ContentBrowserStyle.cpp
+++ b/Engine/Source/Editor/ContentBrowser/Private/ContentBrowserStyle.cpp
@@ -386,6 +386,26 @@ namespace UE::ContentBrowser::Private
 		Set("ContentBrowser.AssetViewOptions.Block.Padding", MenuBlockPadding);
 		Set("ContentBrowser.AssetViewOptions.Block.IndentedPadding", IndentedMenuBlockPadding);
 
+		// Content Sources
+		{
+			static const FVector2D DefaultIconSize = CoreStyleConstants::Icon24x24;
+			static constexpr float DefaultButtonSize = 56.0f;
+			
+			FToolBarStyle ContentBrowserSourceBarStyle = FAppStyle::Get().GetWidgetStyle<FToolBarStyle>("FVerticalToolBar");
+			
+			ContentBrowserSourceBarStyle.SetIconSize(DefaultIconSize);
+			ContentBrowserSourceBarStyle.SetUniformBlockWidth(DefaultButtonSize);
+			ContentBrowserSourceBarStyle.SetButtonPadding(FMargin(2.0f));
+			ContentBrowserSourceBarStyle.ButtonStyle.SetNormalPadding(FMargin(0));
+			ContentBrowserSourceBarStyle.ButtonStyle.SetPressedPadding(FMargin(0));
+			
+			Set("ContentBrowser.SourceBar", ContentBrowserSourceBarStyle);
+			Set("ContentBrowser.Sources.PanelIcon.Opened", new CORE_IMAGE_BRUSH_SVG("Starship/Common/SidePanelLeft", DefaultIconSize));
+			Set("ContentBrowser.Sources.PanelIcon.Closed", new CORE_IMAGE_BRUSH_SVG("Starship/Common/SidePanelLeftClosed", DefaultIconSize));
+			Set("ContentBrowser.Sources.ProjectIcon", new CORE_IMAGE_BRUSH_SVG("Starship/Common/folder-closed", DefaultIconSize));
+			Set("ContentBrowser.Sources.CreateIcon", new IMAGE_BRUSH_SVG("Starship/Common/PlaceActors", DefaultIconSize));
+		}
+
 		FSlateStyleRegistry::RegisterSlateStyle(*this);
 	}
 
diff --git a/Engine/Source/Editor/ContentBrowser/Private/ContentSources/Widgets/SContentSource.cpp b/Engine/Source/Editor/ContentBrowser/Private/ContentSources/Widgets/SContentSource.cpp
new file mode 100644
--- /dev/null
+++ b/Engine/Source/Editor/ContentBrowser/Private/ContentSources/Widgets/SContentSource.cpp
@@ -0,0 +1,81 @@
+// Copyright Epic Games, Inc. All Rights Reserved.
+
+#include "SContentSource.h"
+
+#include "ContentSources/Columns/ContentSourcesColumns.h"
+#include "Elements/Common/EditorDataStorageFeatures.h"
+#include "Elements/Interfaces/TypedElementDataStorageInterface.h"
+#include "Elements/Interfaces/TypedElementDataStorageUiInterface.h"
+
+namespace UE::Editor::ContentBrowser
+{
+	namespace Private
+	{
+		static const FName Purpose("ContentBrowser.AssetView");
+	}
+	
+	void SContentSource::Construct(const FArguments& InArgs)
+    {
+		using namespace UE::Editor::DataStorage;
+		
+		DataStorage = GetMutableDataStorageFeature<IEditorDataStorageProvider>(StorageFeatureName);
+		DataStorageUi = GetMutableDataStorageFeature<IEditorDataStorageUiProvider>(UiFeatureName);
+
+		ChildSlot
+		[
+			CreateWidget()
+		];
+    }
+
+	void SContentSource::SetContentSource(const TSharedPtr<IContentSource>& InContentSource)
+	{
+		ContentSource = InContentSource;
+
+		ChildSlot
+		[
+			CreateWidget()
+		];
+	}
+
+	TSharedRef<SWidget> SContentSource::CreateWidget()
+	{
+		// For now, if there was no valid content source provided we simply display nothing
+		if (!ContentSource)
+		{
+			return SNullWidget::NullWidget;
+		}
+		
+		using namespace UE::Editor::DataStorage;
+
+		// Otherwise we display any widget constructors registered for the "ContentBrowser.AssetView" purpose currently
+		TSharedPtr<FTypedElementWidgetConstructor> AssetViewWidgetConstructor;
+
+		auto AssignWidgetToColumn = [&AssetViewWidgetConstructor] (TUniquePtr<FTypedElementWidgetConstructor> Constructor, TConstArrayView<TWeakObjectPtr<const UScriptStruct>>)
+		{
+			AssetViewWidgetConstructor = TSharedPtr<FTypedElementWidgetConstructor>(Constructor.Release());
+			return false;
+		};
+
+		DataStorageUi->CreateWidgetConstructors(Private::Purpose, FMetaDataView(), AssignWidgetToColumn);
+
+		TSharedPtr<SWidget> AssetViewWidget;
+		
+		if (AssetViewWidgetConstructor)
+		{
+			RowHandle WidgetRow = DataStorage->AddRow(DataStorage->FindTable(FName("Editor_WidgetTable")));
+
+			// Add FContentSourceColumn so the widget can use it to know which content source is active
+			DataStorage->AddColumn(WidgetRow, FContentSourceColumn{.ContentSource = ContentSource});
+			
+			AssetViewWidget = DataStorageUi->ConstructWidget(WidgetRow, *AssetViewWidgetConstructor, FMetaDataView());
+		}
+
+		if (!AssetViewWidget)
+		{
+			return SNullWidget::NullWidget;
+		}
+
+		return AssetViewWidget.ToSharedRef();
+	}
+}
+
diff --git a/Engine/Source/Editor/ContentBrowser/Private/ContentSources/Widgets/SContentSource.h b/Engine/Source/Editor/ContentBrowser/Private/ContentSources/Widgets/SContentSource.h
new file mode 100644
--- /dev/null
+++ b/Engine/Source/Editor/ContentBrowser/Private/ContentSources/Widgets/SContentSource.h
@@ -0,0 +1,34 @@
+// Copyright Epic Games, Inc. All Rights Reserved.
+
+#pragma once
+
+#include "ContentSources/IContentSource.h"
+#include "Widgets/SCompoundWidget.h"
+
+class IEditorDataStorageProvider;
+class IEditorDataStorageUiProvider;
+
+namespace UE::Editor::ContentBrowser
+{
+	// Displays the widgets for a single content source that is currently active
+	class SContentSource : public SCompoundWidget
+	{
+	public:
+		SLATE_BEGIN_ARGS(SContentSource) {}
+		
+		SLATE_END_ARGS()
+
+		void Construct(const FArguments& InArgs);
+		void SetContentSource(const TSharedPtr<IContentSource>& ContentSource);
+
+	private:
+		TSharedRef<SWidget> CreateWidget();
+
+	private:
+		// The currently active content source
+		TSharedPtr<IContentSource> ContentSource;
+		IEditorDataStorageProvider* DataStorage = nullptr; 
+		IEditorDataStorageUiProvider* DataStorageUi = nullptr;
+	};
+
+}
\ No newline at end of file
diff --git a/Engine/Source/Editor/ContentBrowser/Private/ContentSources/Widgets/SContentSourcesView.cpp b/Engine/Source/Editor/ContentBrowser/Private/ContentSources/Widgets/SContentSourcesView.cpp
new file mode 100644
--- /dev/null
+++ b/Engine/Source/Editor/ContentBrowser/Private/ContentSources/Widgets/SContentSourcesView.cpp
@@ -0,0 +1,291 @@
+// Copyright Epic Games, Inc. All Rights Reserved.
+
+#pragma once
+
+#include "SContentSourcesView.h"
+
+#include "ContentBrowserStyle.h"
+#include "IContentBrowserSingleton.h"
+#include "SLegacyContentSource.h"
+#include "ToolMenus.h"
+#include "Widgets/Layout/SBox.h"
+#include "Widgets/Layout/SWidgetSwitcher.h"
+
+#define LOCTEXT_NAMESPACE "SContentSourcesView"
+
+namespace UE::Editor::ContentBrowser
+{
+	static bool bShowContentSourcesBar = false;
+	static FAutoConsoleVariableRef CVarShowContentSourcesBar(
+		TEXT("ContentBrowser.UI.ShowContentSourcesBar"),
+		bShowContentSourcesBar,
+		TEXT("Show the UI to swap between content sources (experimental)"));
+	
+	FDelayedAutoRegisterHelper SContentSourcesView::SourceBarMenuRegistration(
+	EDelayedRegisterRunPhase::EndOfEngineInit,
+	[]
+	{
+		static const FName SourceBarMenuName("ContentBrowser.SourceBar");
+		if (!ensureAlways(!UToolMenus::Get()->IsMenuRegistered(SourceBarMenuName)))
+		{
+			return;
+		}
+		
+		FToolMenuOwnerScoped OwnerScoped(UE_MODULE_NAME);
+		UToolMenu* ToolBar = UToolMenus::Get()->RegisterMenu(SourceBarMenuName, NAME_None, EMultiBoxType::VerticalToolBar);
+		ToolBar->SetStyleSet(&UE::ContentBrowser::Private::FContentBrowserStyle::Get());
+		ToolBar->StyleName = SourceBarMenuName; // Style name is the same as the menu entry name
+		
+		ToolBar->AddSection("Sources")
+		.AddDynamicEntry(NAME_None,
+			FNewToolMenuSectionDelegate::CreateLambda([](FToolMenuSection& InSection)
+			{
+				if (const UContentSourcesViewMenuContext* Context = InSection.FindContext<UContentSourcesViewMenuContext>())
+				{
+					if (const TSharedPtr<SContentSourcesView> ContentSourceWidget = Context->ContentSourcesWidget.Pin())
+					{
+						// Legacy Content Source entry
+						InSection.AddEntry(
+							FToolMenuEntry::InitToolBarButton(
+								FName("LegacyContentSource"),
+								FUIAction(
+									FExecuteAction::CreateLambda([Context]()
+									{
+										if (TSharedPtr<SContentSourcesView> ContentSourceWidget = Context->ContentSourcesWidget.Pin())
+										{
+											ContentSourceWidget->ActivateLegacyContentSource();
+										}
+									}),
+									FCanExecuteAction(),
+									FIsActionChecked::CreateLambda([Context]()
+									{
+										if (TSharedPtr<SContentSourcesView> ContentSourceWidget = Context->ContentSourcesWidget.Pin())
+										{
+											return ContentSourceWidget->IsLegacyContentSourceActive();
+										}
+										return false;
+									}),
+									FIsActionButtonVisible::CreateLambda([Context]()
+									{
+										if (TSharedPtr<SContentSourcesView> ContentSourceWidget = Context->ContentSourcesWidget.Pin())
+										{
+											return ContentSourceWidget->HasLegacyContentSource();
+										}
+										return false;
+									})),
+									LOCTEXT("LegacyContentSourceName", "Project"),
+									FText(),
+									FSlateIcon(UE::ContentBrowser::Private::FContentBrowserStyle::Get().GetStyleSetName(),
+										"ContentBrowser.Sources.ProjectIcon"),
+								EUserInterfaceActionType::ToggleButton
+							));
+
+						// Regular Content Sources
+						auto MakeSourceEntry = [Context, &InSection](const TSharedRef<IContentSource>& ContentSource)
+						{
+							return InSection.AddEntry(
+							FToolMenuEntry::InitToolBarButton(
+								ContentSource->GetName(),
+								FUIAction(
+									FExecuteAction::CreateLambda([ContentSource, Context]()
+									{
+										if (TSharedPtr<SContentSourcesView> ContentSourceWidget = Context->ContentSourcesWidget.Pin())
+										{
+											ContentSourceWidget->ActivateContentSource(ContentSource);
+										}
+									}),
+									FCanExecuteAction(),
+									FIsActionChecked::CreateLambda([ContentSource, Context]()
+									{
+										if (TSharedPtr<SContentSourcesView> ContentSourceWidget = Context->ContentSourcesWidget.Pin())
+										{
+											return ContentSourceWidget->IsContentSourceActive(ContentSource);
+										}
+										return false;
+									})),
+									ContentSource->GetDisplayName(),
+									FText(),
+									ContentSource->GetIcon(),
+								EUserInterfaceActionType::ToggleButton
+							));
+						};
+
+						ContentSourceWidget->ForEachContentSource(MakeSourceEntry);
+					}
+				}
+			}));
+	});
+	
+	void SContentSourcesView::Construct(const FArguments& InArgs)
+	{
+		LegacyContentSource = InArgs._LegacyContentSource;
+		OnLegacyContentSourceEnabledEvent = InArgs._OnLegacyContentSourceEnabled;
+		OnLegacyContentSourceDisabledEvent = InArgs._OnLegacyContentSourceDisabled;
+		
+		constexpr float ContentSourceBarWidth = 64.0f;
+
+		UpdateContentSourcesList();
+		IContentBrowserSingleton::Get().OnContentSourceFactoriesChanged().AddSP(this, &SContentSourcesView::OnContentSourcesChanged);
+
+		TSharedRef<SWidget> LegacyContentSourceWidget = LegacyContentSource ? LegacyContentSource.ToSharedRef() : SNullWidget::NullWidget;
+		
+		ChildSlot
+		[
+			SNew(SHorizontalBox)
+			+ SHorizontalBox::Slot()
+			.Padding(0.0f)
+			.AutoWidth()
+			[
+				SNew(SBorder)
+				.Visibility_Lambda([]()
+				{
+					return bShowContentSourcesBar ? EVisibility::Visible : EVisibility::Collapsed;
+				})
+				.BorderImage(FAppStyle::Get().GetBrush("Brushes.Recessed"))
+				[
+					SAssignNew(SourcesBarContainer, SBox)
+					.WidthOverride(ContentSourceBarWidth)
+					[
+						CreateSourceBar()
+					]
+				]
+			]
+			+ SHorizontalBox::Slot()
+			[
+				SAssignNew(LegacyWidgetSwitcher, SWidgetSwitcher)
+				+SWidgetSwitcher::Slot()
+				[
+					LegacyContentSourceWidget
+				]
+				+SWidgetSwitcher::Slot()
+				[
+					SAssignNew(ContentSourceWidget, SContentSource)
+				]
+			]
+		];
+		
+		ChooseActiveContentSource();
+	}
+
+	void SContentSourcesView::ForEachContentSource(const TFunctionRef<void(const TSharedRef<IContentSource>&)>& InFunctor)
+	{
+		for (const TSharedRef<IContentSource>& ContentSource : ContentSources)
+		{
+			InFunctor(ContentSource);
+		}
+	}
+
+	void SContentSourcesView::ActivateContentSource(const TSharedRef<IContentSource>& InContentSource)
+	{
+		// If this content source is already active, do nothing
+		if (ActiveContentSource == InContentSource)
+		{
+			return;
+		}
+		
+		ActiveContentSource = InContentSource;
+
+		// If the legacy content source was previously active, disable it and notify
+		if (bIsLegacyContentSourceActive)
+		{
+			bIsLegacyContentSourceActive = false;
+			OnLegacyContentSourceDisabledEvent.ExecuteIfBound();
+		}
+		
+		LegacyWidgetSwitcher->SetActiveWidgetIndex(1);
+		ContentSourceWidget->SetContentSource(ActiveContentSource);
+	}
+
+	bool SContentSourcesView::IsContentSourceActive(const TSharedRef<IContentSource>& InContentSource) const
+	{
+		return ActiveContentSource == InContentSource;
+	}
+
+	bool SContentSourcesView::HasLegacyContentSource() const
+	{
+		return LegacyContentSource.IsValid();
+	}
+
+	bool SContentSourcesView::IsLegacyContentSourceActive() const
+	{
+		return HasLegacyContentSource() && bIsLegacyContentSourceActive;
+	}
+
+	void SContentSourcesView::ActivateLegacyContentSource()
+	{
+		// If the legacy content source doesn't exist, or it is already active - do nothing
+		if (HasLegacyContentSource() && !IsLegacyContentSourceActive())
+		{
+			bIsLegacyContentSourceActive = true;
+			OnLegacyContentSourceEnabledEvent.ExecuteIfBound();
+			
+			LegacyWidgetSwitcher->SetActiveWidgetIndex(0);
+			ActiveContentSource = nullptr;
+
+			// Set the currently active content source widget as nullptr to empty the widget contents 
+			ContentSourceWidget->SetContentSource(ActiveContentSource);
+		}
+	}
+
+	void SContentSourcesView::ChooseActiveContentSource()
+	{
+		ActiveContentSource = nullptr;
+
+		// If the legacy content source is available, it gets priority
+		if (HasLegacyContentSource())
+		{
+			ActivateLegacyContentSource();
+		}
+		else if (ContentSources.Num())
+		{
+			ActivateContentSource(ContentSources[0]);
+		}
+		// If there is no legacy content source and no known content sources - just set our contents to SNullWidget for now
+		else
+		{
+			ChildSlot
+			[
+				SNullWidget::NullWidget
+			];
+		}
+	}
+
+	TSharedRef<SWidget> SContentSourcesView::CreateSourceBar()
+	{
+		FToolMenuContext MenuContext;
+		UContentSourcesViewMenuContext* CommonContextObject = NewObject<UContentSourcesViewMenuContext>();
+		CommonContextObject->ContentSourcesWidget = SharedThis(this);
+		MenuContext.AddObject(CommonContextObject);
+		return UToolMenus::Get()->GenerateWidget("ContentBrowser.SourceBar", MenuContext);
+	}
+
+	void SContentSourcesView::OnContentSourcesChanged()
+	{
+		UpdateContentSourcesList();
+		SourcesBarContainer->SetContent(CreateSourceBar());
+
+		if (ActiveContentSource)
+		{
+			// If the active content source does not exist in the list of known content sources anymore, choose a new one
+			if (ContentSources.FindByPredicate([this](const TSharedPtr<IContentSource>& InContentSource)
+			{
+				return InContentSource->GetName() == ActiveContentSource->GetName();
+			}) == nullptr)
+			{
+				ChooseActiveContentSource();
+			}
+		}
+	}
+
+	void SContentSourcesView::UpdateContentSourcesList()
+	{
+		ContentSources.Empty();
+		IContentBrowserSingleton::Get().ForEachContentSourceFactory([this]
+			(const FName& Name, const IContentBrowserSingleton::FContentSourceFactory& Factory)
+		{
+			ContentSources.Add(Factory.Execute());
+		});
+	}
+}
+
+#undef LOCTEXT_NAMESPACE
\ No newline at end of file
diff --git a/Engine/Source/Editor/ContentBrowser/Private/ContentSources/Widgets/SContentSourcesView.h b/Engine/Source/Editor/ContentBrowser/Private/ContentSources/Widgets/SContentSourcesView.h
new file mode 100644
--- /dev/null
+++ b/Engine/Source/Editor/ContentBrowser/Private/ContentSources/Widgets/SContentSourcesView.h
@@ -0,0 +1,97 @@
+// Copyright Epic Games, Inc. All Rights Reserved.
+
+#pragma once
+
+#include "SContentSource.h"
+#include "ContentSources/IContentSource.h"
+#include "Widgets/SCompoundWidget.h"
+
+#include "SContentSourcesView.generated.h"
+
+namespace UE::Editor::ContentBrowser
+{
+	class SLegacyContentSource;
+	class SContentSourcesView;
+}
+
+class SBox;
+class SWidgetSwitcher;
+
+UCLASS()
+class UContentSourcesViewMenuContext : public UObject
+{
+public:
+	GENERATED_BODY()
+		
+	TWeakPtr<UE::Editor::ContentBrowser::SContentSourcesView> ContentSourcesWidget;
+};
+
+
+namespace UE::Editor::ContentBrowser
+{
+	DECLARE_DELEGATE(FOnLegacyContentSourceEnabled)
+	DECLARE_DELEGATE(FOnLegacyContentSourceDisabled)
+
+	// Widget that displays a vertical toolbar to swap between known content sources and a child widget that contains the contents of the
+	// currently active content source
+	class SContentSourcesView : public SCompoundWidget
+	{
+	public:
+		SLATE_BEGIN_ARGS(SContentSourcesView) {}
+
+		// A legacy content source can be used to displays the widgets of the old asset based content browser
+		SLATE_ARGUMENT(TSharedPtr<SLegacyContentSource>, LegacyContentSource)
+		SLATE_EVENT(FOnLegacyContentSourceEnabled, OnLegacyContentSourceEnabled)
+		SLATE_EVENT(FOnLegacyContentSourceDisabled, OnLegacyContentSourceDisabled)
+			
+		SLATE_END_ARGS()
+
+		void Construct(const FArguments& InArgs);
+
+		void ForEachContentSource(const TFunctionRef<void(const TSharedRef<IContentSource>&)>& InFunctor);
+		void ActivateContentSource(const TSharedRef<IContentSource>& InContentSource);
+		bool IsContentSourceActive(const TSharedRef<IContentSource>& InContentSource) const;
+
+		// Whether this widget currently supports showing the legacy content source
+		bool HasLegacyContentSource() const;
+
+		// Whether the legacy content source is supported and currently active
+		bool IsLegacyContentSourceActive() const;
+		void ActivateLegacyContentSource();
+
+		// Pick the first valid content source that is found
+		void ChooseActiveContentSource();
+
+	private:
+		TSharedRef<SWidget> CreateSourceBar();
+		void OnContentSourcesChanged();
+		void UpdateContentSourcesList();
+		
+	private:
+
+		// List of currently known content sources
+		TArray<TSharedRef<IContentSource>> ContentSources;
+
+		// Currently active content source
+		TSharedPtr<IContentSource> ActiveContentSource;
+
+		// Widget that displays the contents of the currently active content source
+		TSharedPtr<SContentSource> ContentSourceWidget;
+
+		// Container widget for the vertical toolbar to display all known content source
+		TSharedPtr<SBox> SourcesBarContainer;
+
+		// The legacy content source
+		TSharedPtr<SLegacyContentSource> LegacyContentSource;
+		
+		bool bIsLegacyContentSourceActive = false;
+		FOnLegacyContentSourceEnabled OnLegacyContentSourceEnabledEvent;
+		FOnLegacyContentSourceDisabled OnLegacyContentSourceDisabledEvent;
+
+		// Widget switcher to display either the legacy content source or the currently active content source
+		TSharedPtr<SWidgetSwitcher> LegacyWidgetSwitcher;
+
+		/** Registers the content source bar */
+		static FDelayedAutoRegisterHelper SourceBarMenuRegistration;
+	};
+}
\ No newline at end of file
diff --git a/Engine/Source/Editor/ContentBrowser/Private/ContentSources/Widgets/SLegacyContentSource.cpp b/Engine/Source/Editor/ContentBrowser/Private/ContentSources/Widgets/SLegacyContentSource.cpp
new file mode 100644
--- /dev/null
+++ b/Engine/Source/Editor/ContentBrowser/Private/ContentSources/Widgets/SLegacyContentSource.cpp
@@ -0,0 +1,24 @@
+// Copyright Epic Games, Inc. All Rights Reserved.
+
+#pragma once
+
+#include "SLegacyContentSource.h"
+
+namespace UE::Editor::ContentBrowser
+{
+	void SLegacyContentSource::Construct(const FArguments& InArgs)
+	{
+		ChildSlot
+		[
+			SNullWidget::NullWidget
+		];
+	}
+
+	void SLegacyContentSource::SetContent(const TSharedRef<SWidget>& InWidget)
+	{
+		ChildSlot
+		[
+			InWidget
+		];
+	}
+}
diff --git a/Engine/Source/Editor/ContentBrowser/Private/ContentSources/Widgets/SLegacyContentSource.h b/Engine/Source/Editor/ContentBrowser/Private/ContentSources/Widgets/SLegacyContentSource.h
new file mode 100644
--- /dev/null
+++ b/Engine/Source/Editor/ContentBrowser/Private/ContentSources/Widgets/SLegacyContentSource.h
@@ -0,0 +1,26 @@
+// Copyright Epic Games, Inc. All Rights Reserved.
+
+#pragma once
+
+#include "IContentBrowserSingleton.h"
+#include "Widgets/SCompoundWidget.h"
+
+class SAssetView;
+
+namespace UE::Editor::ContentBrowser
+{
+	// Widget to display the old asset based content browser
+	// Currently does not contain any logic and simply displays a pre-existing widget that is supplied by SContentBrowser
+	class SLegacyContentSource : public SCompoundWidget
+	{
+	public:
+		SLATE_BEGIN_ARGS(SLegacyContentSource) {}
+		
+		SLATE_END_ARGS()
+		
+		void Construct(const FArguments& InArgs);
+		void SetContent(const TSharedRef<SWidget>& InWidget);
+	private:
+	};
+
+}
diff --git a/Engine/Source/Editor/ContentBrowser/Private/Menus/ContentBrowserMenus.cpp b/Engine/Source/Editor/ContentBrowser/Private/Menus/ContentBrowserMenus.cpp
--- a/Engine/Source/Editor/ContentBrowser/Private/Menus/ContentBrowserMenus.cpp
+++ b/Engine/Source/Editor/ContentBrowser/Private/Menus/ContentBrowserMenus.cpp
@@ -97,7 +97,7 @@ FDelayedAutoRegisterHelper SContentBrowser::NavigationBarMenuRegistration(
 						const TSharedRef<SContentBrowser> ContentBrowser = Context->ContentBrowser.Pin().ToSharedRef();
 
 						// Should have been initialized in SContentBrowser::CreateAssetView
-						if (!ensure(ContentBrowser->NavigationBar.IsValid()))
+						if (!ensure(ContentBrowser->LegacyContentSourceWidgets->NavigationBar.IsValid()))
 						{
 							return;
 						}
@@ -106,7 +106,7 @@ FDelayedAutoRegisterHelper SContentBrowser::NavigationBarMenuRegistration(
 							InSection.AddEntry(
 								FToolMenuEntry::InitWidget(
 									"PathBox",
-									ContentBrowser->NavigationBar.ToSharedRef(),
+									ContentBrowser->LegacyContentSourceWidgets->NavigationBar.ToSharedRef(),
 									FText::GetEmpty()));
 
 						PathBoxEntry.WidgetData.StyleParams.SizeRule = FSizeParam::SizeRule_StretchContent;
@@ -370,7 +370,7 @@ FDelayedAutoRegisterHelper SContentBrowser::ToolBarMenuRegistration(
                             const TSharedRef<SContentBrowser> ContentBrowser = Context->ContentBrowser.Pin().ToSharedRef();
 
                             // Should have been initialized in SContentBrowser::CreateAssetView
-                            if (!ensure(ContentBrowser->SearchBoxPtr.IsValid()))
+                            if (!ensure(ContentBrowser->LegacyContentSourceWidgets->SearchBoxPtr.IsValid()))
                             {
                                 return;
                             }
@@ -379,7 +379,7 @@ FDelayedAutoRegisterHelper SContentBrowser::ToolBarMenuRegistration(
                                 InSection.AddEntry(
                                 FToolMenuEntry::InitWidget(
                                     "SearchBox",
-                                    ContentBrowser->SearchBoxPtr.ToSharedRef(),
+                                    ContentBrowser->LegacyContentSourceWidgets->SearchBoxPtr.ToSharedRef(),
                                     FText::GetEmpty()));
 
                             SearchBoxEntry.WidgetData.StyleParams.SizeRule = FSizeParam::SizeRule_StretchContent;
diff --git a/Engine/Source/Editor/ContentBrowser/Private/SContentBrowser.cpp b/Engine/Source/Editor/ContentBrowser/Private/SContentBrowser.cpp
--- a/Engine/Source/Editor/ContentBrowser/Private/SContentBrowser.cpp
+++ b/Engine/Source/Editor/ContentBrowser/Private/SContentBrowser.cpp
@@ -133,6 +133,8 @@
 #include "ToolMenuMisc.h"
 #include "ToolMenuSection.h"
 #include "ToolMenus.h"
+#include "ContentSources/Widgets/SContentSourcesView.h"
+#include "ContentSources/Widgets/SLegacyContentSource.h"
 #include "Experimental/ContentBrowserExtensionUtils.h"
 #include "Toolkits/GlobalEditorCommonCommands.h"
 #include "Trace/Detail/Channel.h"
@@ -307,7 +309,7 @@ SContentBrowser::~SContentBrowser()
 
 bool SContentBrowser::ShouldShowRedirectors() const
 {
-	return FilterListPtr.IsValid() ? ContentBrowserUtils::ShouldShowRedirectors(FilterListPtr) : false;
+	return LegacyContentSourceWidgets->FilterListPtr.IsValid() ? ContentBrowserUtils::ShouldShowRedirectors(LegacyContentSourceWidgets->FilterListPtr) : false;
 }
 
 BEGIN_SLATE_FUNCTION_BUILD_OPTIMIZATION
@@ -317,6 +319,13 @@ void SContentBrowser::Construct( const FArguments& InArgs, const FName& InInstan
 
 	InstanceName = InInstanceName;
 
+	// Store a copy of the init config if specified so we can re-create the asset view widgets dynamically
+	if (Config)
+	{
+		InitConfig = *Config;
+		bHasInitConfig = true;
+	}
+
 	JumpMRU.MaxItems = 30;
 
 	UContentBrowserConfig::Initialize();
@@ -331,6 +340,8 @@ void SContentBrowser::Construct( const FArguments& InArgs, const FName& InInstan
 		InArgs._ContainingTab->SetOnTabClosed( SDockTab::FOnTabClosedCallback::CreateSP( this, &SContentBrowser::OnContainingTabClosed ) );
 		InArgs._ContainingTab->SetOnTabActivated( SDockTab::FOnTabActivatedCallback::CreateSP( this, &SContentBrowser::OnContainingTabActivated ) );
 	}
+	
+	LegacyContentSource = SNew(UE::Editor::ContentBrowser::SLegacyContentSource);
 
  	bIsLocked = InArgs._InitiallyLocked;
 	bCanSetAsPrimaryBrowser = Config != nullptr ? Config->bCanSetAsPrimaryBrowser : true;
@@ -339,11 +350,6 @@ void SContentBrowser::Construct( const FArguments& InArgs, const FName& InInstan
 	HistoryManager.SetOnApplyHistoryData(FOnApplyHistoryData::CreateSP(this, &SContentBrowser::OnApplyHistoryData));
 	HistoryManager.SetOnUpdateHistoryData(FOnUpdateHistoryData::CreateSP(this, &SContentBrowser::OnUpdateHistoryData));
 
-	PathContextMenu = MakeShareable(new FPathContextMenu( AsShared() ));
-	PathContextMenu->SetOnRenameFolderRequested(FPathContextMenu::FOnRenameFolderRequested::CreateSP(this, &SContentBrowser::OnRenameRequested));
-	PathContextMenu->SetOnFolderDeleted(FPathContextMenu::FOnFolderDeleted::CreateSP(this, &SContentBrowser::OnOpenedFolderDeleted));
-	PathContextMenu->SetOnFolderFavoriteToggled(FPathContextMenu::FOnFolderFavoriteToggled::CreateSP(this, &SContentBrowser::ToggleFolderFavorite));
-	PathContextMenu->SetOnPrivateContentEditToggled(FPathContextMenu::FOnPrivateContentEditToggled::CreateSP(this, &SContentBrowser::TogglePrivateContentEdit));
 	FrontendFilters = MakeShareable(new FAssetFilterCollectionType());
 	TextFilter = MakeShared<FAssetTextFilter>();
 
@@ -358,757 +364,244 @@ void SContentBrowser::Construct( const FArguments& InArgs, const FName& InInstan
 	SourcesSearch->SetHintText(LOCTEXT("SearchPathsHint", "Search Paths"));
 
 	static const FName DefaultForegroundName("DefaultForeground");
-
-	BindCommands();
+	
 	UContentBrowserSettings::OnSettingChanged().AddSP(this, &SContentBrowser::OnContentBrowserSettingsChanged);
-
-	// Currently this controls the asset count
-	const bool bShowBottomToolbar = Config != nullptr ? Config->bShowBottomToolbar : true;
-
+	
 	// Register console variable sink for private content setting changing
 	CVarSinkHandle = IConsoleManager::Get().RegisterConsoleVariableSink_Handle(FConsoleCommandDelegate::CreateSP(this, &SContentBrowser::OnConsoleVariableChanged));
 	UpdatePrivateContentFeatureEnabled(false /* bUpdateFilterIfChanged */);
 
-	AssetViewPtr = SNew(SAssetView)
-			.ThumbnailLabel(Config != nullptr ? Config->ThumbnailLabel : EThumbnailLabel::ClassName)
-			//.ThumbnailScale(Config != nullptr ? Config->ThumbnailScale : 0.18f)
-			.InitialViewType(Config != nullptr ? Config->InitialAssetViewType : EAssetViewType::Tile)
-			.OnNewItemRequested(this, &SContentBrowser::OnNewItemRequested)
-			.OnItemSelectionChanged(this, &SContentBrowser::OnItemSelectionChanged, EContentBrowserViewContext::AssetView)
-			.OnItemsActivated(this, &SContentBrowser::OnItemsActivated)
-			.OnGetItemContextMenu(this, &SContentBrowser::GetItemContextMenu, EContentBrowserViewContext::AssetView)
-			.OnItemRenameCommitted(this, &SContentBrowser::OnItemRenameCommitted)
-			.FrontendFilters(FrontendFilters)
-			.TextFilter(TextFilter)
-			.ShowRedirectors(this, &SContentBrowser::ShouldShowRedirectors)
-			.HighlightedText(this, &SContentBrowser::GetHighlightedText)
-			.ShowBottomToolbar(bShowBottomToolbar)
-			.ShowViewOptions(false) // We control this for the main content browser
-			.AllowThumbnailEditMode(true)
-			.AllowThumbnailHintLabel(false)
-			.CanShowFolders(Config != nullptr ? Config->bCanShowFolders : true)
-			.CanShowClasses(Config != nullptr ? Config->bCanShowClasses : true)
-			.CanShowRealTimeThumbnails(Config != nullptr ? Config->bCanShowRealTimeThumbnails : true)
-			.CanShowDevelopersFolder(Config != nullptr ? Config->bCanShowDevelopersFolder : true)
-			.CanShowFavorites(true)
-			.AddMetaData<FTagMetaData>(FTagMetaData(TEXT("ContentBrowserAssets")))
-			.OwningContentBrowser(SharedThis(this))
-			.OnSearchOptionsChanged(this, &SContentBrowser::HandleAssetViewSearchOptionsChanged)
-			.bShowPathViewFilters(true)
-			.FillEmptySpaceInTileView(true)
-			.ShowDisallowedAssetClassAsUnsupportedItems(true)
-			.AllowCustomView(true);
-
-	TSharedRef<SWidget> ViewOptions = SNullWidget::NullWidget;
+	ChildSlot
+	[
+		// The legacy content source will be activated by default, which will call OnLegacyContentSourceEnabled which ends up creating the asset view
+		// widgets and initiatializing all settings and so on
+		SAssignNew(ContentSourcesContainer, UE::Editor::ContentBrowser::SContentSourcesView)
+			.LegacyContentSource(LegacyContentSource)
+			.OnLegacyContentSourceEnabled(this, &SContentBrowser::OnLegacyContentSourceEnabled)
+			.OnLegacyContentSourceDisabled(this, &SContentBrowser::OnLegacyContentSourceDisabled)
+		
+	];
 
-	// Note, for backwards compatibility ShowBottomToolbar controls the visibility of view options so we respect that here
-	if (bShowBottomToolbar)
-	{
-		if (UE::Editor::ContentBrowser::IsNewStyleEnabled())
-		{
-			ViewOptions =
-				SNew(SActionButton)
-				.ActionButtonType(EActionButtonType::Simple)
-				.OnGetMenuContent(AssetViewPtr.ToSharedRef(), &SAssetView::GetViewButtonContent)
-				.Icon(FAppStyle::Get().GetBrush("Icons.Settings"));
-		}
-		else
-		{
-			ViewOptions =
-				SNew(SComboButton)
-				.ComboButtonStyle(&FAppStyle::Get().GetWidgetStyle<FComboButtonStyle>("SimpleComboButton"))
-				.OnGetMenuContent(AssetViewPtr.ToSharedRef(), &SAssetView::GetViewButtonContent)
-				.HasDownArrow(false)
-				.ButtonContent()
-				[
-					SNew(SHorizontalBox)
-					+ SHorizontalBox::Slot()
-					.AutoWidth()
-					.VAlign(VAlign_Center)
-					.Padding(4.0, 0.0f)
-					[
-						SNew(SImage)
-						.ColorAndOpacity(FSlateColor::UseForeground())
-						.Image(FAppStyle::Get().GetBrush("Icons.Settings"))
-					]
+	ExtendViewOptionsMenu(Config);
+	
+	// Set the initial history data
+	HistoryManager.AddHistoryData();
 
-					+ SHorizontalBox::Slot()
-					.VAlign(VAlign_Center)
-					.Padding(4.0, 0.0f)
-					[
-						SNew(STextBlock)
-						.Text(LOCTEXT("Settings", "Settings"))
-						.ColorAndOpacity(FSlateColor::UseForeground())
-					]
-				];
-		}
-	}
+	// Load settings if they were specified
+	LoadSettings(InInstanceName);
+	
+	// We want to be able to search the feature packs in the super search so we need the module loaded 
+	IAddContentDialogModule& AddContentDialogModule = FModuleManager::LoadModuleChecked<IAddContentDialogModule>("AddContentDialog");
 
-	if (UE::Editor::ContentBrowser::IsNewStyleEnabled())
-	{
-		constexpr float ToolBarVerticalPadding = 4.0f; // ToolBar total height should be 36.0f
-		constexpr float PanelInsetPadding = 2.0f;
-		constexpr float SourceTreeSectionPadding = 2.0f;
+	// Update the breadcrumb trail path
+	OnContentBrowserSettingsChanged(NAME_None);
+}
+END_SLATE_FUNCTION_BUILD_OPTIMIZATION
 
-		const TSharedRef<SWidget> AssetView = CreateAssetView(Config);
+void SContentBrowser::BindCommands()
+{
+	Commands = TSharedPtr< FUICommandList >(new FUICommandList);
 
-		CreateFavoritesView(Config);
-		CreatePathView(Config);
+	Commands->MapAction(FGenericCommands::Get().Rename, FUIAction(
+		FExecuteAction::CreateSP(this, &SContentBrowser::HandleRenameCommand),
+		FCanExecuteAction::CreateSP(this, &SContentBrowser::HandleRenameCommandCanExecute)
+	));
 
-		{
-			SAssignNew(SourceTreePtr, SContentBrowserSourceTree)
+	Commands->MapAction(FGenericCommands::Get().Delete, FUIAction(
+		FExecuteAction::CreateSP(this, &SContentBrowser::HandleDeleteCommandExecute),
+		FCanExecuteAction::CreateSP(this, &SContentBrowser::HandleDeleteCommandCanExecute)
+	));
 
-			+ SContentBrowserSourceTree::Slot()
-			.AreaWidget(FavoritesArea)
-			.Size(0.2f)
-			.Visibility(this, &SContentBrowser::GetFavoriteFolderVisibility)
+	Commands->MapAction(FContentBrowserCommands::Get().OpenAssetsOrFolders, FUIAction(
+		FExecuteAction::CreateSP(this, &SContentBrowser::HandleOpenAssetsOrFoldersCommandExecute)
+	));
 
-			+ SContentBrowserSourceTree::Slot()
-			.AreaWidget(PathArea)
-			.Size(0.8f);
-		}
+	Commands->MapAction(FContentBrowserCommands::Get().PreviewAssets, FUIAction(
+		FExecuteAction::CreateSP(this, &SContentBrowser::HandlePreviewAssetsCommandExecute)
+	));
 
-		NavigationToolBarWidget = CreateNavigationToolBar(Config);
+	Commands->MapAction(FContentBrowserCommands::Get().CreateNewFolder, FUIAction(
+		FExecuteAction::CreateSP(this, &SContentBrowser::HandleCreateNewFolderCommandExecute)
+	));
 
-		TSharedPtr<SBox> AssetViewNavigationToolBarContainer = nullptr;
-		TSharedPtr<SBox> SourceTreeAndAssetViewNavigationToolBarContainer = nullptr;
+	Commands->MapAction(FContentBrowserCommands::Get().GoUpToParentFolder, FUIAction(
+		FExecuteAction::CreateSP(this, &SContentBrowser::HandleGoUpToParentFolder),
+		FCanExecuteAction::CreateSP(this, &SContentBrowser::HandleCanGoUpToParentFolder)
+	));
 
-		ChildSlot
-		[
-			SNew(SVerticalBox)
+	Commands->MapAction(FContentBrowserCommands::Get().SaveSelectedAsset, FUIAction(
+		FExecuteAction::CreateSP(this, &SContentBrowser::HandleSaveAssetCommand),
+		FCanExecuteAction::CreateSP(this, &SContentBrowser::HandleSaveAssetCommandCanExecute)
+	));
 
-			+ SVerticalBox::Slot()
-			.AutoHeight()
-			[
-				SNew(SSeparator)
-				.Thickness(2.0f)
-			]
+	Commands->MapAction(FContentBrowserCommands::Get().SaveAllCurrentFolder, FUIAction(
+		FExecuteAction::CreateSP(this, &SContentBrowser::HandleSaveAllCurrentFolderCommand)
+	));
 
-			// Source / Tree + Assets
-			+ SVerticalBox::Slot()
-			.FillHeight(1.0f)
-			.Padding(0.0f)
-			[
-				SNew(SHorizontalBox)
+	Commands->MapAction(FContentBrowserCommands::Get().ResaveAllCurrentFolder, FUIAction(
+		FExecuteAction::CreateSP(this, &SContentBrowser::HandleResaveAllCurrentFolderCommand)
+	));
 
-				// Tree + Assets + Navigation Bar
-				+ SHorizontalBox::Slot()
-				.FillWidth(1.0f)
-				.Padding(FMargin(0.0f))
-				[
-					SNew(SVerticalBox)
+	Commands->MapAction(FContentBrowserCommands::Get().EditPath, FUIAction(
+		FExecuteAction::CreateSP(this, &SContentBrowser::EditPathCommand)
+	));
 
-					// Assets/tree
-					+ SVerticalBox::Slot()
-					.FillHeight(1.0f)
-					.Padding(FMargin(0.0f))
-					[
-						// The tree/assets splitter
-						SAssignNew(PathAssetSplitterPtr, SSplitter)
-						.PhysicalSplitterHandleSize(PanelInsetPadding)
-
-						// Sources View
-						+ SSplitter::Slot()
-						.Resizable(true)
-						.SizeRule(SSplitter::SizeToContent)
-						.OnSlotResized(this, &SContentBrowser::OnPathViewBoxColumnResized)
-						[
-							SNew(SBox)
-							.Padding(0.0f)
-							.Visibility(this, &SContentBrowser::GetSourcesViewVisibility)
-							.WidthOverride(this, &SContentBrowser::GetPathViewBoxWidthOverride)
-							[
-								SNew(SBorder)
-								.Padding(SourceTreeSectionPadding)
-								.BorderImage(FAppStyle::GetBrush("Brushes.Panel"))
-								[
-									// Panel background, seen between items
-									SNew(SBorder)
-									.BorderImage(FAppStyle::GetBrush("WhiteBrush"))
-									.BorderBackgroundColor(FStyleColors::Panel)
-									.HAlign(HAlign_Fill)
-									.VAlign(VAlign_Fill)
-									.Padding(0.0f)
-									[
-										SourceTreePtr.ToSharedRef()
-									]
-								]
-							]
-						]
+	// Allow extenders to add commands
+	FContentBrowserModule& ContentBrowserModule = FModuleManager::GetModuleChecked<FContentBrowserModule>(TEXT("ContentBrowser"));
+	TArray<FContentBrowserCommandExtender> CommmandExtenderDelegates = ContentBrowserModule.GetAllContentBrowserCommandExtenders();
 
-						// Asset View
-						+ SSplitter::Slot()
-						.Value(0.75f)
-						[
-							SNew(SBox)
-							.Padding(0.0f)
-							[
-								SNew(SVerticalBox)
+	for (int32 i = 0; i < CommmandExtenderDelegates.Num(); ++i)
+	{
+		if (CommmandExtenderDelegates[i].IsBound())
+		{
+			CommmandExtenderDelegates[i].Execute(Commands.ToSharedRef(), FOnContentBrowserGetSelection::CreateSP(this, &SContentBrowser::GetSelectionState));
+		}
+	}
 
-								+ SVerticalBox::Slot()
-								.AutoHeight()
-								[
-									SNew(SBorder)
-									.Padding(FMargin(3, ToolBarVerticalPadding))
-									.BorderImage(bIsDrawer ? FStyleDefaults::GetNoBrush() : FAppStyle::Get().GetBrush("Brushes.Panel"))
-									[
-										SNew(SHorizontalBox)
-
-										+ SHorizontalBox::Slot()
-										.FillWidth(1.0f)
-										.HAlign(HAlign_Fill)
-										.VAlign(VAlign_Center)
-										.Padding(5, 0, 0, 0)
-										[
-											CreateToolBar(Config)
-										]
-
-										+ SHorizontalBox::Slot()
-										.AutoWidth()
-										.HAlign(HAlign_Right)
-										.VAlign(VAlign_Center)
-										[
-											CreateDrawerDockButton(Config)
-										]
-
-										+ SHorizontalBox::Slot()
-										.AutoWidth()
-										.HAlign(HAlign_Right)
-										.VAlign(VAlign_Center)
-										.Padding(5.0f, 0.0f, 0.0f, 0.0f)
-										[
-											ViewOptions
-										]
-									]
-								]
+	FInputBindingManager::Get().RegisterCommandList(FContentBrowserCommands::Get().GetContextName(), Commands.ToSharedRef());
+}
 
-								+ SVerticalBox::Slot()
-								[
-									AssetView
-								]
+void SContentBrowser::UnbindCommands()
+{
+	Commands.Reset();
+}
 
-								// Navigation ToolBar
-								+ SVerticalBox::Slot()
-								.AutoHeight()
-								.HAlign(HAlign_Fill)
-								.Padding(FMargin(0, PanelInsetPadding, 0, 0))
-								[
-									SNew(SBorder)
-									.Padding(FMargin(3))
-									.BorderImage(bIsDrawer ? FStyleDefaults::GetNoBrush() : FAppStyle::Get().GetBrush("Brushes.Panel"))
-									[
-										SAssignNew(AssetViewNavigationToolBarContainer, SBox)
-										.Padding(0)
-										[
-											NavigationToolBarWidget.ToSharedRef()
-										]
-									]
-								]
-							]
-						]
-					]
+EVisibility SContentBrowser::GetFavoriteFolderVisibility() const
+{
+	if (const FContentBrowserInstanceConfig* Config = GetConstInstanceConfig())
+	{
+		return Config->bShowFavorites ? EVisibility::Visible : EVisibility::Collapsed;
+	}
 
-					+ SVerticalBox::Slot()
-					.AutoHeight()
-					.Padding(0)
-					[
-						SAssignNew(SourceTreeAndAssetViewNavigationToolBarContainer, SBox)
-						.Padding(0)
-					]
-				]
-			]
-		];
+	return GetDefault<UContentBrowserSettings>()->GetDisplayFavorites() ? EVisibility::Visible : EVisibility::Collapsed;
+}
 
-		NavigationToolBarParentSwitcher = MakeShared<UE::Editor::ContentBrowser::Private::TWidgetSizeParentSwitcher<Orient_Horizontal>>(
-			NavigationToolBarWidget,
-			TArray<TSharedPtr<SBox>>{ SourceTreeAndAssetViewNavigationToolBarContainer, AssetViewNavigationToolBarContainer },
-			TArray<FInt16Range>{ FInt16Range(0, 550), FInt16Range(550, 10000) },
-			[this]()
-			{
-				const float Size = this->GetCachedGeometry().GetLocalSize().X;
-				return Size;
-			}
-		);
+EVisibility SContentBrowser::GetLockButtonVisibility() const
+{
+	return IsLocked() ? EVisibility::Visible : EVisibility::Collapsed;
+}
 
-		SourceTreeSplitterNumFixedSlots = SourceTreePtr->GetSplitter()->GetChildren()->NumSlot();
+void SContentBrowser::AddFolderFavorite(const TArray<FString>& FolderPaths)
+{
+	for (const FString& FolderPath : FolderPaths)
+	{
+		const FContentBrowserItemPath ItemPath(FolderPath, EContentBrowserPathType::Virtual);
+		if (!ContentBrowserUtils::IsFavoriteFolder(ItemPath))
+		{
+			ContentBrowserUtils::AddFavoriteFolder(ItemPath);
+		}
 	}
-	else
+
+	SaveAndShowNewFolderFavorites(FolderPaths);
+}
+
+void SContentBrowser::ToggleFolderFavorite(const TArray<FString>& FolderPaths)
+{
+	TArray<FString> FolderPathsAdded;
+	for (const FString& FolderPath : FolderPaths)
 	{
-		ChildSlot
-		[
-			SNew(SVerticalBox)
-			+SVerticalBox::Slot()
-			.AutoHeight()
-			.Padding( 0, 0, 0, 0 )
-			[
-				SNew( SBorder )
-				.Padding( FMargin( 3 ) )
-				.BorderImage(bIsDrawer ? FStyleDefaults::GetNoBrush() : FAppStyle::Get().GetBrush("Brushes.Panel"))
-				[
-					SNew(SHorizontalBox)
-					+ SHorizontalBox::Slot()
-					.AutoWidth()
-					.HAlign(HAlign_Left)
-					.Padding(5, 0, 0, 0)
-					[
-						CreateToolBar(Config)
-					]
-					// History Back Button
-					+SHorizontalBox::Slot()
-					.Padding(10, 0, 0, 0)
-					.AutoWidth()
-					[
-						SNew(SButton)
-						.VAlign(EVerticalAlignment::VAlign_Center)
-						.ButtonStyle(FAppStyle::Get(), "SimpleButton")
-						.ToolTipText( this, &SContentBrowser::GetHistoryBackTooltip )
-						.ContentPadding( FMargin(1, 0) )
-						.OnClicked(this, &SContentBrowser::BackClicked)
-						.IsEnabled(this, &SContentBrowser::IsBackEnabled)
-						.AddMetaData<FTagMetaData>(FTagMetaData(TEXT("ContentBrowserHistoryBack")))
-						[
-							SNew(SImage)
-							.Image(FAppStyle::Get().GetBrush("Icons.CircleArrowLeft"))
-							.ColorAndOpacity(FSlateColor::UseForeground())
-						]
-					]
-					// History Forward Button
-					+ SHorizontalBox::Slot()
-					.Padding(2, 0, 0, 0)
-					.AutoWidth()
-					[
-						SNew(SButton)
-						.VAlign(EVerticalAlignment::VAlign_Center)
-						.ButtonStyle(FAppStyle::Get(), "SimpleButton")
-						.ToolTipText( this, &SContentBrowser::GetHistoryForwardTooltip )
-						.ContentPadding( FMargin(1, 0) )
-						.OnClicked(this, &SContentBrowser::ForwardClicked)
-						.IsEnabled(this, &SContentBrowser::IsForwardEnabled)
-						.AddMetaData<FTagMetaData>(FTagMetaData(TEXT("ContentBrowserHistoryForward")))
-						[
-							SNew(SImage)
-							.Image(FAppStyle::Get().GetBrush("Icons.CircleArrowRight"))
-							.ColorAndOpacity(FSlateColor::UseForeground())
-						]
-					]
-
-					// Path
-					+ SHorizontalBox::Slot()
-					.VAlign(VAlign_Center)
-					.HAlign(HAlign_Fill)
-					.FillWidth(1.0f)
-					.Padding(2, 0, 0, 0)
-					[
-						SAssignNew(NavigationBar, SNavigationBar)
-						.OnPathClicked(this, &SContentBrowser::OnPathClicked)
-						.GetPathMenuContent(this, &SContentBrowser::OnGetCrumbDelimiterContent)
-						.GetComboOptions(this, &SContentBrowser::GetRecentPaths)
-						.OnNavigateToPath(this, &SContentBrowser::OnNavigateToPath)
-						.OnCompletePrefix(this, &SContentBrowser::OnCompletePathPrefix)
-						.OnCanEditPathAsText(this, &SContentBrowser::OnCanEditPathAsText)
-						.AddMetaData<FTagMetaData>(FTagMetaData(TEXT("ContentBrowserPath")))
-					]
-
-					+ SHorizontalBox::Slot()
-					.AutoWidth()
-					.HAlign(HAlign_Right)
-					.VAlign(VAlign_Center)
-					[
-						CreateLockButton(Config)
-					]
-					+ SHorizontalBox::Slot()
-					.AutoWidth()
-					.HAlign(HAlign_Right)
-					.VAlign(VAlign_Center)
-					[
-						CreateDrawerDockButton(Config)
-					]
-					+ SHorizontalBox::Slot()
-					.AutoWidth()
-					.Padding(5.0f, 0.0f, 0.0f, 0.0f)
-					.HAlign(HAlign_Right)
-					.VAlign(VAlign_Center)
-					[
-						ViewOptions
-					]
-				]
-			]
-
-			+SVerticalBox::Slot()
-			.AutoHeight()
-			[
-				SNew(SSeparator)
-				.Thickness(2.0f)
-			]
-
-			// Assets/tree
-			+ SVerticalBox::Slot()
-			.FillHeight(1.0f)
-			.Padding(0.0f)
-			[
-				// The tree/assets splitter
-				SAssignNew(PathAssetSplitterPtr, SSplitter)
-				.PhysicalSplitterHandleSize(2.0f)
-
-				// Sources View
-				+ SSplitter::Slot()
-				.Resizable(true)
-				.SizeRule(SSplitter::SizeToContent)
-				.OnSlotResized(this, &SContentBrowser::OnPathViewBoxColumnResized)
-				[
-					SNew(SBox)
-					.Padding(FMargin(4.f))
-					.Visibility(this, &SContentBrowser::GetSourcesViewVisibility)
-					.WidthOverride(this, &SContentBrowser::GetPathViewBoxWidthOverride)
-					[
-						SNew(SBorder)
-						.Padding(FMargin(0))
-						.BorderImage(FAppStyle::GetBrush("Brushes.Recessed"))
-
-						[
-							SAssignNew(PathFavoriteSplitterPtr, SSplitter)
-							.Clipping(EWidgetClipping::ClipToBounds)
-							.PhysicalSplitterHandleSize(2.0f)
-							.HitDetectionSplitterHandleSize(8.0f)
-							.Orientation(EOrientation::Orient_Vertical)
-							.MinimumSlotHeight(26.0f)
-							.Visibility( this, &SContentBrowser::GetSourcesViewVisibility )
-							+SSplitter::Slot()
-							.SizeRule(TAttribute<SSplitter::ESizeRule>(this, &SContentBrowser::GetFavoritesAreaSizeRule))
-							.MinSize(TAttribute<float>(this, &SContentBrowser::GetFavoritesAreaMinSize))
-							.Value(0.2f)
-							[
-								SNew(SBorder)
-								.BorderImage(FAppStyle::Get().GetBrush("Brushes.Header"))
-								.Padding(0.0f, 2.0f, 0.0f, 0.0f)
-								[
-									CreateFavoritesView(Config)
-								]
-							]
-
-							+SSplitter::Slot()
-							.SizeRule(TAttribute<SSplitter::ESizeRule>(this, &SContentBrowser::GetPathAreaSizeRule))
-							.MinSize(29.0f)
-							.Value(0.8f)
-							[
-								SNew(SBorder)
-								.BorderImage(FAppStyle::Get().GetBrush("Brushes.Header"))
-								.Padding(0.0f, 2.0f, 0.0f, 0.0f)
-								[
-									CreatePathView(Config)
-								]
-							]
-						]
-					]
-				]
-
-				// Asset View
-				+ SSplitter::Slot()
-				.Value(0.75f)
-				[
-					CreateAssetView(Config)
-				]
-			]
-		];
-
-		SourceTreeSplitterNumFixedSlots = PathFavoriteSplitterPtr->GetChildren()->NumSlot();
+		const FContentBrowserItemPath ItemPath(FolderPath, EContentBrowserPathType::Virtual);
+		if (ContentBrowserUtils::IsFavoriteFolder(ItemPath))
+		{
+			ContentBrowserUtils::RemoveFavoriteFolder(ItemPath);
+		}
+		else
+		{
+			ContentBrowserUtils::AddFavoriteFolder(ItemPath);
+			FolderPathsAdded.Add(FolderPath);
+		}
 	}
 
-	ExtendViewOptionsMenu(Config);
-
-	AssetContextMenu = MakeShared<FAssetContextMenu>(AssetViewPtr);
-	AssetContextMenu->BindCommands(Commands);
-	AssetContextMenu->SetOnShowInPathsViewRequested( FAssetContextMenu::FOnShowInPathsViewRequested::CreateSP(this, &SContentBrowser::OnShowInPathsViewRequested) );
-	AssetContextMenu->SetOnRenameRequested( FAssetContextMenu::FOnRenameRequested::CreateSP(this, &SContentBrowser::OnRenameRequested) );
-	AssetContextMenu->SetOnDuplicateRequested( FAssetContextMenu::FOnDuplicateRequested::CreateSP(this, &SContentBrowser::OnDuplicateRequested) );
-	AssetContextMenu->SetOnAssetViewRefreshRequested( FAssetContextMenu::FOnAssetViewRefreshRequested::CreateSP( this, &SContentBrowser::OnAssetViewRefreshRequested) );
-	FavoritePathViewPtr->SetTreeTitle(LOCTEXT("Favorites", "Favorites"));
-
-	TOptional<FCollectionRef> SelectedCollection;
+	SaveAndShowNewFolderFavorites(FolderPathsAdded);
+}
 
-	if (Config != nullptr)
+void SContentBrowser::SaveAndShowNewFolderFavorites(const TArray<FString>& FolderPaths)
+{
+	// If the legacy content source isn't active - the settings will get updated when it is made active
+	if (!ContentSourcesContainer->IsLegacyContentSourceActive())
 	{
-		SelectedCollection = Config->SelectedCollection;
+		return;
+	}
+	
+	LegacyContentSourceWidgets->FavoritePathViewPtr->SaveSettings(GEditorPerProjectIni, SettingsIniSection, InstanceName.ToString() + TEXT(".Favorites"));
+	LegacyContentSourceWidgets->FavoritePathViewPtr->Populate();
 
-		PRAGMA_DISABLE_DEPRECATION_WARNINGS
-		// Check for someone using the deprecated SelectedCollectionName instead of SelectedCollection.
-		if (!SelectedCollection.GetValue().IsValid() && Config->SelectedCollectionName.Name != NAME_None)
+	if (!FolderPaths.IsEmpty())
+	{	
+		LegacyContentSourceWidgets->FavoritePathViewPtr->SetSelectedPaths(FolderPaths);
+		if (GetFavoriteFolderVisibility() == EVisibility::Collapsed)
 		{
-			SelectedCollection = FCollectionRef(FCollectionManagerModule::GetModule().Get().GetProjectCollectionContainer(), Config->SelectedCollectionName);
+			UContentBrowserSettings* Settings = GetMutableDefault<UContentBrowserSettings>();
+			Settings->SetDisplayFavorites(true);
+			Settings->SaveConfig();
 		}
-		PRAGMA_ENABLE_DEPRECATION_WARNINGS
 	}
+}
 
-	if (SelectedCollection.IsSet() && SelectedCollection.GetValue().IsValid())
-	{
-		// Select the specified collection by default
-		FAssetViewContentSources DefaultContentSources(SelectedCollection.GetValue());
-		AssetViewPtr->SetContentSources(DefaultContentSources);
-	}
-	else
+void SContentBrowser::TogglePrivateContentEdit(const TArray<FString>& FolderPaths)
+{
+	for (const FString& FolderPath : FolderPaths)
 	{
-		// Select /Game by default
-		const FName DefaultInvariantPath(TEXT("/Game"));
-		FName DefaultVirtualPath;
-		IContentBrowserDataModule::Get().GetSubsystem()->ConvertInternalPathToVirtual(DefaultInvariantPath, DefaultVirtualPath);
-
-		FAssetViewContentSources DefaultContentSources(DefaultVirtualPath);
-		AssetViewPtr->SetContentSources(DefaultContentSources);
-	}
-
-	// Set the initial history data
-	HistoryManager.AddHistoryData();
-
-	// Load settings if they were specified
-	LoadSettings(InInstanceName);
+		ensure(FContentBrowserSingleton::Get().IsFolderShowPrivateContentToggleable(FolderPath));
 
-	if (Config != nullptr)
-	{
-		// Make sure the sources view is initially visible if we were asked to show it
-		SetSourcesViewExpanded(Config->bExpandSourcesView && Config->bUseSourcesView);
-	}
-	else
-	{
-		// in case we do not have a config, see what the global default settings are for the Sources Panel
-		bool bSourcesExpanded = true;
-		if (EditorConfig != nullptr)
+		if (FContentBrowserSingleton::Get().IsShowingPrivateContent(FolderPath))
 		{
-			bSourcesExpanded = EditorConfig->bSourcesExpanded;
+			ContentBrowserUtils::RemoveShowPrivateContentFolder(FolderPath, TEXT("ContentBrowser"));
 		}
-
-		SetSourcesViewExpanded(bSourcesExpanded);
-	}
-
-	// Bindings to manage history when items are deleted
-	FCollectionManagerModule& CollectionManagerModule = FCollectionManagerModule::GetModule();
-
-	TArray<TSharedPtr<ICollectionContainer>> CollectionContainers;
-	CollectionManagerModule.Get().GetCollectionContainers(CollectionContainers);
-
-	int32 InsertIndex = 0;
-	for (const TSharedPtr<ICollectionContainer>& CollectionContainer : CollectionContainers)
-	{
-		if (!CollectionContainer->IsHidden())
+		else
 		{
-			AddSlotForCollectionContainer(InsertIndex++, CollectionContainer.ToSharedRef());
+			ContentBrowserUtils::AddShowPrivateContentFolder(FolderPath, TEXT("ContentBrowser"));
 		}
-
-		CollectionContainer->OnIsHiddenChanged().AddSP(this, &SContentBrowser::HandleIsHiddenChanged);
-		CollectionContainer->OnCollectionRenamed().AddSP(this, &SContentBrowser::HandleCollectionRenamed);
-		CollectionContainer->OnCollectionDestroyed().AddSP(this, &SContentBrowser::HandleCollectionRemoved);
-		CollectionContainer->OnCollectionUpdated().AddSP(this, &SContentBrowser::HandleCollectionUpdated);
 	}
 
-	CollectionManagerModule.Get().OnCollectionContainerCreated().AddSP(this, &SContentBrowser::HandleCollectionContainerAdded);
-	CollectionManagerModule.Get().OnCollectionContainerDestroyed().AddSP(this, &SContentBrowser::HandleCollectionContainerRemoved);
-
-	UContentBrowserDataSubsystem* ContentBrowserData = IContentBrowserDataModule::Get().GetSubsystem();
-	ContentBrowserData->OnItemDataUpdated().AddSP(this, &SContentBrowser::HandleItemDataUpdated);
-
-	// We want to be able to search the feature packs in the super search so we need the module loaded 
-	IAddContentDialogModule& AddContentDialogModule = FModuleManager::LoadModuleChecked<IAddContentDialogModule>("AddContentDialog");
-
-	// Update the breadcrumb trail path
-	OnContentBrowserSettingsChanged(NAME_None);
-
-	// Initialize the search options
-	HandleAssetViewSearchOptionsChanged();
+	OnAssetViewRefreshRequested();
 }
-END_SLATE_FUNCTION_BUILD_OPTIMIZATION
 
-void SContentBrowser::BindCommands()
+void SContentBrowser::HandleAssetViewSearchOptionsChanged()
 {
-	Commands = TSharedPtr< FUICommandList >(new FUICommandList);
-
-	Commands->MapAction(FGenericCommands::Get().Rename, FUIAction(
-		FExecuteAction::CreateSP(this, &SContentBrowser::HandleRenameCommand),
-		FCanExecuteAction::CreateSP(this, &SContentBrowser::HandleRenameCommandCanExecute)
-	));
-
-	Commands->MapAction(FGenericCommands::Get().Delete, FUIAction(
-		FExecuteAction::CreateSP(this, &SContentBrowser::HandleDeleteCommandExecute),
-		FCanExecuteAction::CreateSP(this, &SContentBrowser::HandleDeleteCommandCanExecute)
-	));
-
-	Commands->MapAction(FContentBrowserCommands::Get().OpenAssetsOrFolders, FUIAction(
-		FExecuteAction::CreateSP(this, &SContentBrowser::HandleOpenAssetsOrFoldersCommandExecute)
-	));
-
-	Commands->MapAction(FContentBrowserCommands::Get().PreviewAssets, FUIAction(
-		FExecuteAction::CreateSP(this, &SContentBrowser::HandlePreviewAssetsCommandExecute)
-	));
+	bool bIncludeClassName = LegacyContentSourceWidgets->AssetViewPtr->IsIncludingClassNames();
+	bool bIncludeAssetPath = LegacyContentSourceWidgets->AssetViewPtr->IsIncludingAssetPaths();
+	bool bIncludeCollectionNames = LegacyContentSourceWidgets->AssetViewPtr->IsIncludingCollectionNames();
+	
+	TextFilter->SetIncludeClassName(bIncludeClassName);
+	TextFilter->SetIncludeAssetPath(bIncludeAssetPath);
+	TextFilter->SetIncludeCollectionNames(bIncludeCollectionNames);
 
-	Commands->MapAction(FContentBrowserCommands::Get().CreateNewFolder, FUIAction(
-		FExecuteAction::CreateSP(this, &SContentBrowser::HandleCreateNewFolderCommandExecute)
-	));
+	// Make sure all custom text filters get the updated Asset View Search Options
+	LegacyContentSourceWidgets->FilterListPtr->UpdateCustomTextFilterIncludes(bIncludeClassName, bIncludeAssetPath, bIncludeCollectionNames);
+}
 
-	Commands->MapAction(FContentBrowserCommands::Get().GoUpToParentFolder, FUIAction(
-		FExecuteAction::CreateSP(this, &SContentBrowser::HandleGoUpToParentFolder),
-		FCanExecuteAction::CreateSP(this, &SContentBrowser::HandleCanGoUpToParentFolder)
-	));
+TSharedRef<SWidget> SContentBrowser::CreateToolBar(const FContentBrowserConfig* Config)
+{
+	FToolMenuContext MenuContext;
 
-	Commands->MapAction(FContentBrowserCommands::Get().SaveSelectedAsset, FUIAction(
-		FExecuteAction::CreateSP(this, &SContentBrowser::HandleSaveAssetCommand),
-		FCanExecuteAction::CreateSP(this, &SContentBrowser::HandleSaveAssetCommandCanExecute)
-	));
+	UContentBrowserToolbarMenuContext* CommonContextObject = NewObject<UContentBrowserToolbarMenuContext>();
+	CommonContextObject->ContentBrowser = SharedThis(this);
+	CommonContextObject->AssetView = LegacyContentSourceWidgets->AssetViewPtr;
+	CommonContextObject->ContentBrowserConfig = Config;
 
-	Commands->MapAction(FContentBrowserCommands::Get().SaveAllCurrentFolder, FUIAction(
-		FExecuteAction::CreateSP(this, &SContentBrowser::HandleSaveAllCurrentFolderCommand)
-	));
+	MenuContext.AddObject(CommonContextObject);
 
-	Commands->MapAction(FContentBrowserCommands::Get().ResaveAllCurrentFolder, FUIAction(
-		FExecuteAction::CreateSP(this, &SContentBrowser::HandleResaveAllCurrentFolderCommand)
-	));
+	return UToolMenus::Get()->GenerateWidget("ContentBrowser.ToolBar", MenuContext);
+}
 
-	Commands->MapAction(FContentBrowserCommands::Get().EditPath, FUIAction(
-		FExecuteAction::CreateSP(this, &SContentBrowser::EditPathCommand)
-	));
+TSharedRef<SWidget> SContentBrowser::CreateNavigationToolBar(const FContentBrowserConfig* Config)
+{
+	FToolMenuContext MenuContext;
 
-	// Allow extenders to add commands
-	FContentBrowserModule& ContentBrowserModule = FModuleManager::GetModuleChecked<FContentBrowserModule>(TEXT("ContentBrowser"));
-	TArray<FContentBrowserCommandExtender> CommmandExtenderDelegates = ContentBrowserModule.GetAllContentBrowserCommandExtenders();
+	UContentBrowserToolbarMenuContext* CommonContextObject = NewObject<UContentBrowserToolbarMenuContext>();
+	CommonContextObject->ContentBrowser = SharedThis(this);
+	CommonContextObject->AssetView = LegacyContentSourceWidgets->AssetViewPtr;
+	CommonContextObject->ContentBrowserConfig = Config;
 
-	for (int32 i = 0; i < CommmandExtenderDelegates.Num(); ++i)
-	{
-		if (CommmandExtenderDelegates[i].IsBound())
-		{
-			CommmandExtenderDelegates[i].Execute(Commands.ToSharedRef(), FOnContentBrowserGetSelection::CreateSP(this, &SContentBrowser::GetSelectionState));
-		}
-	}
+	MenuContext.AddObject(CommonContextObject);
 
-	FInputBindingManager::Get().RegisterCommandList(FContentBrowserCommands::Get().GetContextName(), Commands.ToSharedRef());
+	return UToolMenus::Get()->GenerateWidget("ContentBrowser.NavigationBar", MenuContext);
 }
 
-EVisibility SContentBrowser::GetFavoriteFolderVisibility() const
+TSharedRef<SWidget> SContentBrowser::CreateLockButton(const FContentBrowserConfig* Config)
 {
-	if (const FContentBrowserInstanceConfig* Config = GetConstInstanceConfig())
-	{
-		return Config->bShowFavorites ? EVisibility::Visible : EVisibility::Collapsed;
-	}
-
-	return GetDefault<UContentBrowserSettings>()->GetDisplayFavorites() ? EVisibility::Visible : EVisibility::Collapsed;
-}
-
-EVisibility SContentBrowser::GetLockButtonVisibility() const
-{
-	return IsLocked() ? EVisibility::Visible : EVisibility::Collapsed;
-}
-
-void SContentBrowser::AddFolderFavorite(const TArray<FString>& FolderPaths)
-{
-	for (const FString& FolderPath : FolderPaths)
-	{
-		const FContentBrowserItemPath ItemPath(FolderPath, EContentBrowserPathType::Virtual);
-		if (!ContentBrowserUtils::IsFavoriteFolder(ItemPath))
-		{
-			ContentBrowserUtils::AddFavoriteFolder(ItemPath);
-		}
-	}
-
-	SaveAndShowNewFolderFavorites(FolderPaths);
-}
-
-void SContentBrowser::ToggleFolderFavorite(const TArray<FString>& FolderPaths)
-{
-	TArray<FString> FolderPathsAdded;
-	for (const FString& FolderPath : FolderPaths)
-	{
-		const FContentBrowserItemPath ItemPath(FolderPath, EContentBrowserPathType::Virtual);
-		if (ContentBrowserUtils::IsFavoriteFolder(ItemPath))
-		{
-			ContentBrowserUtils::RemoveFavoriteFolder(ItemPath);
-		}
-		else
-		{
-			ContentBrowserUtils::AddFavoriteFolder(ItemPath);
-			FolderPathsAdded.Add(FolderPath);
-		}
-	}
-
-	SaveAndShowNewFolderFavorites(FolderPathsAdded);
-}
-
-void SContentBrowser::SaveAndShowNewFolderFavorites(const TArray<FString>& FolderPaths)
-{
-	FavoritePathViewPtr->SaveSettings(GEditorPerProjectIni, SettingsIniSection, InstanceName.ToString() + TEXT(".Favorites"));
-	FavoritePathViewPtr->Populate();
-
-	if (!FolderPaths.IsEmpty())
-	{	
-		FavoritePathViewPtr->SetSelectedPaths(FolderPaths);
-		if (GetFavoriteFolderVisibility() == EVisibility::Collapsed)
-		{
-			UContentBrowserSettings* Settings = GetMutableDefault<UContentBrowserSettings>();
-			Settings->SetDisplayFavorites(true);
-			Settings->SaveConfig();
-		}
-	}
-}
-
-void SContentBrowser::TogglePrivateContentEdit(const TArray<FString>& FolderPaths)
-{
-	for (const FString& FolderPath : FolderPaths)
-	{
-		ensure(FContentBrowserSingleton::Get().IsFolderShowPrivateContentToggleable(FolderPath));
-
-		if (FContentBrowserSingleton::Get().IsShowingPrivateContent(FolderPath))
-		{
-			ContentBrowserUtils::RemoveShowPrivateContentFolder(FolderPath, TEXT("ContentBrowser"));
-		}
-		else
-		{
-			ContentBrowserUtils::AddShowPrivateContentFolder(FolderPath, TEXT("ContentBrowser"));
-		}
-	}
-
-	OnAssetViewRefreshRequested();
-}
-
-void SContentBrowser::HandleAssetViewSearchOptionsChanged()
-{
-	bool bIncludeClassName = AssetViewPtr->IsIncludingClassNames();
-	bool bIncludeAssetPath = AssetViewPtr->IsIncludingAssetPaths();
-	bool bIncludeCollectionNames = AssetViewPtr->IsIncludingCollectionNames();
-	
-	TextFilter->SetIncludeClassName(bIncludeClassName);
-	TextFilter->SetIncludeAssetPath(bIncludeAssetPath);
-	TextFilter->SetIncludeCollectionNames(bIncludeCollectionNames);
-
-	// Make sure all custom text filters get the updated Asset View Search Options
-	FilterListPtr->UpdateCustomTextFilterIncludes(bIncludeClassName, bIncludeAssetPath, bIncludeCollectionNames);
-}
-
-TSharedRef<SWidget> SContentBrowser::CreateToolBar(const FContentBrowserConfig* Config)
-{
-	FToolMenuContext MenuContext;
-
-	UContentBrowserToolbarMenuContext* CommonContextObject = NewObject<UContentBrowserToolbarMenuContext>();
-	CommonContextObject->ContentBrowser = SharedThis(this);
-	CommonContextObject->AssetView = AssetViewPtr;
-	CommonContextObject->ContentBrowserConfig = Config;
-
-	MenuContext.AddObject(CommonContextObject);
-
-	return UToolMenus::Get()->GenerateWidget("ContentBrowser.ToolBar", MenuContext);
-}
-
-TSharedRef<SWidget> SContentBrowser::CreateNavigationToolBar(const FContentBrowserConfig* Config)
-{
-	FToolMenuContext MenuContext;
-
-	UContentBrowserToolbarMenuContext* CommonContextObject = NewObject<UContentBrowserToolbarMenuContext>();
-	CommonContextObject->ContentBrowser = SharedThis(this);
-	CommonContextObject->AssetView = AssetViewPtr;
-	CommonContextObject->ContentBrowserConfig = Config;
-
-	MenuContext.AddObject(CommonContextObject);
-
-	return UToolMenus::Get()->GenerateWidget("ContentBrowser.NavigationBar", MenuContext);
-}
-
-TSharedRef<SWidget> SContentBrowser::CreateLockButton(const FContentBrowserConfig* Config)
-{
-	if(Config == nullptr || Config->bCanShowLockButton)
+	if(Config == nullptr || Config->bCanShowLockButton)
 	{
 		return
 			SNew(SButton)
@@ -1133,7 +626,7 @@ void SContentBrowser::OnFilterBarLayoutChanging(EFilterBarLayout NewLayout)
 	{
 		if (NewLayout == EFilterBarLayout::Horizontal)
 		{
-			AssetViewBorder->SetContent(
+			LegacyContentSourceWidgets->AssetViewBorder->SetContent(
 
 				SNew(SVerticalBox)
 
@@ -1141,20 +634,20 @@ void SContentBrowser::OnFilterBarLayoutChanging(EFilterBarLayout NewLayout)
 				.Padding(0.0f)
 				.AutoHeight()
 				[
-					FilterListPtr.ToSharedRef()
+					LegacyContentSourceWidgets->FilterListPtr.ToSharedRef()
 				]
 
 				+ SVerticalBox::Slot()
 				.FillHeight(1.0f)
 				.Padding(0, 0)
 				[
-					AssetViewPtr.ToSharedRef()
+					LegacyContentSourceWidgets->AssetViewPtr.ToSharedRef()
 				]
 			);
 		}
 		else
 		{
-			AssetViewBorder->SetContent(
+			LegacyContentSourceWidgets->AssetViewBorder->SetContent(
 
 				SNew(SSplitter)
 				.PhysicalSplitterHandleSize(2.0f)
@@ -1171,7 +664,7 @@ void SContentBrowser::OnFilterBarLayoutChanging(EFilterBarLayout NewLayout)
 					// Don't take up space when there are no filters
 					.Visibility_Lambda([this]
 					{
-						return FilterListPtr->HasAnyFilters() ? EVisibility::Visible : EVisibility::Collapsed;
+						return LegacyContentSourceWidgets->FilterListPtr->HasAnyFilters() ? EVisibility::Visible : EVisibility::Collapsed;
 					})
 					[
 						SNew(SVerticalBox)
@@ -1199,7 +692,7 @@ void SContentBrowser::OnFilterBarLayoutChanging(EFilterBarLayout NewLayout)
 						.Padding(0.0f)
 						.FillHeight(1.0f)
 						[
-							FilterListPtr.ToSharedRef()
+							LegacyContentSourceWidgets->FilterListPtr.ToSharedRef()
 						]
 					]
 				]
@@ -1213,7 +706,7 @@ void SContentBrowser::OnFilterBarLayoutChanging(EFilterBarLayout NewLayout)
 					.FillHeight(1.0f)
 					.Padding(0, 0)
 					[
-						AssetViewPtr.ToSharedRef()
+						LegacyContentSourceWidgets->AssetViewPtr.ToSharedRef()
 					]
 				]
 			);
@@ -1246,16 +739,16 @@ void SContentBrowser::OnFilterBarLayoutChanging(EFilterBarLayout NewLayout)
 					.VAlign(VAlign_Center)
 					.WidthOverride(SearchBoxDesiredWidth)
 					[
-						SearchBoxPtr.ToSharedRef()
+						LegacyContentSourceWidgets->SearchBoxPtr.ToSharedRef()
 					]
 				];
 
 			/** We add the Combo Button and the Search Box to the FilterList itself, so that the filters wrap with them
 			 *	properly in the Horizontal Layout
 			 */
-			FilterListPtr->AddWidgetToCurrentLayout(SearchBoxSlot.ToSharedRef());
+			LegacyContentSourceWidgets->FilterListPtr->AddWidgetToCurrentLayout(SearchBoxSlot.ToSharedRef());
 
-			AssetViewBorder->SetContent(
+			LegacyContentSourceWidgets->AssetViewBorder->SetContent(
 
 				SNew(SVerticalBox)
 
@@ -1263,20 +756,20 @@ void SContentBrowser::OnFilterBarLayoutChanging(EFilterBarLayout NewLayout)
 				.Padding(0.0f)
 				.AutoHeight()
 				[
-					FilterListPtr.ToSharedRef()
+					LegacyContentSourceWidgets->FilterListPtr.ToSharedRef()
 				]
 
 				+ SVerticalBox::Slot()
 				.FillHeight(1.0f)
 				.Padding(0, 0)
 				[
-					AssetViewPtr.ToSharedRef()
+					LegacyContentSourceWidgets->AssetViewPtr.ToSharedRef()
 				]
 			);
 		}
 		else
 		{
-			AssetViewBorder->SetContent(
+			LegacyContentSourceWidgets->AssetViewBorder->SetContent(
 
 				SNew(SSplitter)
 				.PhysicalSplitterHandleSize(2.0f)
@@ -1293,7 +786,7 @@ void SContentBrowser::OnFilterBarLayoutChanging(EFilterBarLayout NewLayout)
 					// Don't take up space when there are no filters
 					.Visibility_Lambda([this]
 					{
-						return FilterListPtr->HasAnyFilters() ? EVisibility::Visible : EVisibility::Collapsed;
+						return LegacyContentSourceWidgets->FilterListPtr->HasAnyFilters() ? EVisibility::Visible : EVisibility::Collapsed;
 					})
 					[
 						SNew(SVerticalBox)
@@ -1321,7 +814,7 @@ void SContentBrowser::OnFilterBarLayoutChanging(EFilterBarLayout NewLayout)
 						.Padding(0.0f)
 						.FillHeight(1.0f)
 						[
-							FilterListPtr.ToSharedRef()
+							LegacyContentSourceWidgets->FilterListPtr.ToSharedRef()
 						]
 					]
 				]
@@ -1352,7 +845,7 @@ void SContentBrowser::OnFilterBarLayoutChanging(EFilterBarLayout NewLayout)
 							.VAlign(VAlign_Center)
 							.WidthOverride(SearchBoxDesiredWidth)
 							[
-								SearchBoxPtr.ToSharedRef()
+								LegacyContentSourceWidgets->SearchBoxPtr.ToSharedRef()
 							]
 						]
 					]
@@ -1361,7 +854,7 @@ void SContentBrowser::OnFilterBarLayoutChanging(EFilterBarLayout NewLayout)
 					.FillHeight(1.0f)
 					.Padding(0, 0)
 					[
-						AssetViewPtr.ToSharedRef()
+						LegacyContentSourceWidgets->AssetViewPtr.ToSharedRef()
 					]
 				]
 			);
@@ -1372,7 +865,7 @@ void SContentBrowser::OnFilterBarLayoutChanging(EFilterBarLayout NewLayout)
 TSharedRef<SWidget> SContentBrowser::CreateAssetView(const FContentBrowserConfig* Config)
 {
 	// Create the Filter Bar Widget
-	FilterListPtr = SNew(SFilterList)
+	LegacyContentSourceWidgets->FilterListPtr = SNew(SFilterList)
 					.OnFilterChanged(this, &SContentBrowser::OnFilterChanged)
 					.Visibility((Config != nullptr ? Config->bCanShowFilters : true) ? EVisibility::Visible : EVisibility::Collapsed)
 					.FrontendFilters(FrontendFilters)
@@ -1386,21 +879,23 @@ TSharedRef<SWidget> SContentBrowser::CreateAssetView(const FContentBrowserConfig
 						TSharedPtr<FFrontendFilter_CustomText> NewFilter = MakeShared<FFrontendFilter_CustomText>();
 
 						// Make sure the new filter has the right search options from the AssetView. We only have to set it once, SFilterList handles syncing it on change
-						NewFilter->UpdateCustomTextFilterIncludes(AssetViewPtr->IsIncludingClassNames(), AssetViewPtr->IsIncludingAssetPaths(), AssetViewPtr->IsIncludingCollectionNames());
+						NewFilter->UpdateCustomTextFilterIncludes(LegacyContentSourceWidgets->AssetViewPtr->IsIncludingClassNames(),
+							LegacyContentSourceWidgets->AssetViewPtr->IsIncludingAssetPaths(),
+							LegacyContentSourceWidgets->AssetViewPtr->IsIncludingCollectionNames());
 
 						return NewFilter;
 					}))
 					.AddMetaData<FTagMetaData>(FTagMetaData(TEXT("ContentBrowserFilters")));
 	
 	// Create the Filter Combo Button
-	FilterComboButton = SFilterList::MakeAddFilterButton(FilterListPtr.ToSharedRef());
+	FilterComboButton = SFilterList::MakeAddFilterButton(LegacyContentSourceWidgets->FilterListPtr.ToSharedRef());
 	TSharedPtr<ISlateMetaData> FilterComboButtonMetaData = MakeShared<FTagMetaData>(TEXT("ContentBrowserFiltersCombo"));
 	FilterComboButton->AddMetadata(FilterComboButtonMetaData.ToSharedRef());
 	FilterComboButton->SetVisibility((Config != nullptr ? Config->bCanShowFilters : true) ? EVisibility::Visible : EVisibility::Collapsed);
 
-	AssetViewPtr->SetFilterBar(FilterListPtr);
+	LegacyContentSourceWidgets->AssetViewPtr->SetFilterBar(LegacyContentSourceWidgets->FilterListPtr);
 
-	SearchBoxPtr = SNew(SAssetSearchBox)
+	LegacyContentSourceWidgets->SearchBoxPtr = SNew(SAssetSearchBox)
 				   .HintText(this, &SContentBrowser::GetSearchAssetsHintText)
 				   .ShowSearchHistory(true)
 				   .OnTextChanged(this, &SContentBrowser::OnSearchBoxChanged)
@@ -1415,7 +910,7 @@ TSharedRef<SWidget> SContentBrowser::CreateAssetView(const FContentBrowserConfig
 
 	if (UE::Editor::ContentBrowser::IsNewStyleEnabled())
 	{
-		NavigationBar = SNew(SNavigationBar)
+		LegacyContentSourceWidgets->NavigationBar = SNew(SNavigationBar)
 						.OnPathClicked(this, &SContentBrowser::OnPathClicked)
 						.GetPathMenuContent(this, &SContentBrowser::OnGetCrumbDelimiterContent)
 						.GetComboOptions(this, &SContentBrowser::GetRecentPaths)
@@ -1434,11 +929,11 @@ TSharedRef<SWidget> SContentBrowser::CreateAssetView(const FContentBrowserConfig
 			? FMargin(2.0f, 0.0f, 2.0f, 0.0f)
 			: FMargin(2.0f, 2.0f, 2.0f, 0.0f);
 
-	AssetViewBorder = SNew(SBorder)
+	LegacyContentSourceWidgets->AssetViewBorder = SNew(SBorder)
 			.BorderImage(FAppStyle::Get().GetBrush("Brushes.Panel"))
 			.Padding(AssetViewPadding);
 
-	return AssetViewBorder.ToSharedRef();
+	return LegacyContentSourceWidgets->AssetViewBorder.ToSharedRef();
 }
 
 SContentBrowser::FCollectionSource& SContentBrowser::AddSlotForCollectionContainer(int32 Index, const TSharedRef<ICollectionContainer>& CollectionContainer)
@@ -1450,21 +945,21 @@ SContentBrowser::FCollectionSource& SContentBrowser::AddSlotForCollectionContain
 
 	const TUniquePtr<FCollectionSource>& CollectionSource = CollectionSources.EmplaceAt_GetRef(Index, new FCollectionSource(*this, CollectionContainer));
 
-	const int32 SlotIndex = SourceTreeSplitterNumFixedSlots + Index;
+	const int32 SlotIndex = LegacyContentSourceWidgets->SourceTreeSplitterNumFixedSlots + Index;
 	SSplitter* Splitter = nullptr;
 	if (UE::Editor::ContentBrowser::IsNewStyleEnabled())
 	{
-		SourceTreePtr->AddSlot(SlotIndex)
+		LegacyContentSourceWidgets->SourceTreePtr->AddSlot(SlotIndex)
 			.AreaWidget(CreateCollectionsView(*CollectionSource))
 			.Size(0.4f);
 
-		Splitter = SourceTreePtr->GetSplitter().Get();
+		Splitter = LegacyContentSourceWidgets->SourceTreePtr->GetSplitter().Get();
 	}
 	else
 	{
 		const float SourceTreeHeaderHeightMin = 26.0f + 3.0f;
 
-		PathFavoriteSplitterPtr->AddSlot(SlotIndex)
+		LegacyContentSourceWidgets->PathFavoriteSplitterPtr->AddSlot(SlotIndex)
 			.SizeRule_Raw(CollectionSource.Get(), &FCollectionSource::GetCollectionsAreaSizeRule)
 			.MinSize(SourceTreeHeaderHeightMin)
 			.Value(0.4f)
@@ -1477,7 +972,7 @@ SContentBrowser::FCollectionSource& SContentBrowser::AddSlotForCollectionContain
 					]
 			];
 
-		Splitter = PathFavoriteSplitterPtr.Get();
+		Splitter = LegacyContentSourceWidgets->PathFavoriteSplitterPtr.Get();
 	}
 
 	FString Key, Filename;
@@ -1502,8 +997,9 @@ void SContentBrowser::RemoveSlotForCollectionContainer(const TSharedRef<ICollect
 	{
 		CollectionSources[Index]->SaveSettings(InstanceName);
 
-		const int32 SlotIndex = SourceTreeSplitterNumFixedSlots + Index;
-		SSplitter& Splitter = UE::Editor::ContentBrowser::IsNewStyleEnabled() ? *SourceTreePtr->GetSplitter() : *PathFavoriteSplitterPtr;
+		const int32 SlotIndex = LegacyContentSourceWidgets->SourceTreeSplitterNumFixedSlots + Index;
+		SSplitter& Splitter = UE::Editor::ContentBrowser::IsNewStyleEnabled() ?
+			*LegacyContentSourceWidgets->SourceTreePtr->GetSplitter() : *LegacyContentSourceWidgets->PathFavoriteSplitterPtr;
 
 		FString Key, Filename;
 		GetSourceTreeSplitterSlotSizeSettingKeyAndFilename(SlotIndex, Key, Filename);
@@ -1528,7 +1024,7 @@ void SContentBrowser::SetFavoritesExpanded(bool bExpanded)
 
 TSharedRef<SWidget> SContentBrowser::CreateFavoritesView(const FContentBrowserConfig* Config)
 {
-	SAssignNew(FavoritePathViewPtr, SFavoritePathView)
+	SAssignNew(LegacyContentSourceWidgets->FavoritePathViewPtr, SFavoritePathView)
 		.OnItemSelectionChanged(this, &SContentBrowser::OnItemSelectionChanged, EContentBrowserViewContext::FavoriteView)
 		.OnGetItemContextMenu(this, &SContentBrowser::GetItemContextMenu, EContentBrowserViewContext::FavoriteView)
 		.FocusSearchBoxWhenOpened(false)
@@ -1546,17 +1042,17 @@ TSharedRef<SWidget> SContentBrowser::CreateFavoritesView(const FContentBrowserCo
 			UE::Editor::ContentBrowser::Private::SContentBrowserSourceTreeArea,
 			"Favorites",
 			FavoritesSearch,
-			FavoritePathViewPtr.ToSharedRef())
+			LegacyContentSourceWidgets->FavoritePathViewPtr.ToSharedRef())
 			.Label(LOCTEXT("Favorites", "Favorites"))
 			.Visibility(this, &SContentBrowser::GetFavoriteFolderVisibility)
 			.OnExpansionChanged(this, &SContentBrowser::SetFavoritesExpanded)
-			.IsEmpty_Lambda([PathView = FavoritePathViewPtr]
+			.IsEmpty_Lambda([PathView = LegacyContentSourceWidgets->FavoritePathViewPtr]
 			{
 				return PathView->IsEmpty();
 			})
 			.EmptyBodyLabel(LOCTEXT("FavoritesEmpty", "Right click a folder to add it to your favorites."));
 
-	FavoritePathViewPtr->SetOnFolderFavoriteAdd(SFavoritePathView::FOnFolderFavoriteAdd::CreateSP(this, &SContentBrowser::AddFolderFavorite));
+	LegacyContentSourceWidgets->FavoritePathViewPtr->SetOnFolderFavoriteAdd(SFavoritePathView::FOnFolderFavoriteAdd::CreateSP(this, &SContentBrowser::AddFolderFavorite));
 	return ViewWidget;
 }
 
@@ -1571,13 +1067,13 @@ void SContentBrowser::SetPathViewExpanded(bool bExpanded)
 
 TSharedRef<SWidget> SContentBrowser::CreatePathView(const FContentBrowserConfig* Config)
 {
-	SAssignNew(PathViewPtr, SPathView)
+	SAssignNew(LegacyContentSourceWidgets->PathViewPtr, SPathView)
 		.OnItemSelectionChanged(this, &SContentBrowser::OnItemSelectionChanged, EContentBrowserViewContext::PathView)
 		.OnGetItemContextMenu(this, &SContentBrowser::GetItemContextMenu, EContentBrowserViewContext::PathView)
 		.FocusSearchBoxWhenOpened(false)
 		.ShowTreeTitle(false)
 		.ShowSeparator(false)
-		.ShowRedirectors_Lambda([this]() { return ContentBrowserUtils::ShouldShowRedirectors(FilterListPtr); })
+		.ShowRedirectors_Lambda([this]() { return ContentBrowserUtils::ShouldShowRedirectors(LegacyContentSourceWidgets->FilterListPtr); })
 		.AddMetaData<FTagMetaData>(FTagMetaData(TEXT("ContentBrowserSources")))
 		.ExternalSearch(SourcesSearch)
 		.PluginPathFilters(PluginPathFilters)
@@ -1591,7 +1087,7 @@ TSharedRef<SWidget> SContentBrowser::CreatePathView(const FContentBrowserConfig*
 			UE::Editor::ContentBrowser::Private::SContentBrowserSourceTreeArea,
 			"Path",
 			SourcesSearch,
-			PathViewPtr.ToSharedRef())
+			LegacyContentSourceWidgets->PathViewPtr.ToSharedRef())
 			.Label(FText::FromString(FApp::GetProjectName()))
 			.OnExpansionChanged(this, &SContentBrowser::SetPathViewExpanded);
 
@@ -1809,11 +1305,17 @@ FText SContentBrowser::GetHighlightedText() const
 
 void SContentBrowser::CreateNewAsset(const FString& DefaultAssetName, const FString& PackagePath, UClass* AssetClass, UFactory* Factory)
 {
-	AssetViewPtr->CreateNewAsset(DefaultAssetName, PackagePath, AssetClass, Factory);
+	// For now we just forcefully enable the legacy content source when a new asset creation is requested
+	ContentSourcesContainer->ActivateLegacyContentSource();
+	
+	LegacyContentSourceWidgets->AssetViewPtr->CreateNewAsset(DefaultAssetName, PackagePath, AssetClass, Factory);
 }
 
 void SContentBrowser::PrepareToSyncItems(TArrayView<const FContentBrowserItem> ItemsToSync, const bool bDisableFiltersThatHideAssets)
 {
+	// For now we just forcefully enable the legacy content source when a sync is requested so it succeeds
+	ContentSourcesContainer->ActivateLegacyContentSource();
+	
 	bool bRepopulate = false;
 
 	// Check to see if any of the assets require certain folders to be visible
@@ -1845,7 +1347,7 @@ void SContentBrowser::PrepareToSyncItems(TArrayView<const FContentBrowserItem> I
 			if (!bDisplayDev && ContentBrowserUtils::IsItemDeveloperContent(ItemToSync))
 			{
 				bDisplayDev = true;
-				AssetViewPtr->OverrideShowDeveloperContent();
+				LegacyContentSourceWidgets->AssetViewPtr->OverrideShowDeveloperContent();
 				bRepopulate = true;
 				bDisplayDevChanged = true;
 			}
@@ -1853,7 +1355,7 @@ void SContentBrowser::PrepareToSyncItems(TArrayView<const FContentBrowserItem> I
 			if (!bDisplayEngine && ContentBrowserUtils::IsItemEngineContent(ItemToSync))
 			{
 				bDisplayEngine = true;
-				AssetViewPtr->OverrideShowEngineContent();
+				LegacyContentSourceWidgets->AssetViewPtr->OverrideShowEngineContent();
 				bRepopulate = true;
 				bDisplayEngineChanged = true;
 			}
@@ -1861,7 +1363,7 @@ void SContentBrowser::PrepareToSyncItems(TArrayView<const FContentBrowserItem> I
 			if (!bDisplayPlugins && ContentBrowserUtils::IsItemPluginContent(ItemToSync))
 			{
 				bDisplayPlugins = true;
-				AssetViewPtr->OverrideShowPluginContent();
+				LegacyContentSourceWidgets->AssetViewPtr->OverrideShowPluginContent();
 				bRepopulate = true;
 				bDisplayPluginsChanged = true;
 			}
@@ -1869,7 +1371,7 @@ void SContentBrowser::PrepareToSyncItems(TArrayView<const FContentBrowserItem> I
 			if (!bDisplayLocalized && ContentBrowserUtils::IsItemLocalizedContent(ItemToSync))
 			{
 				bDisplayLocalized = true;
-				AssetViewPtr->OverrideShowLocalizedContent();
+				LegacyContentSourceWidgets->AssetViewPtr->OverrideShowLocalizedContent();
 				bRepopulate = true;
 				bDisplayLocalizedChanged = true;
 			}
@@ -1883,7 +1385,7 @@ void SContentBrowser::PrepareToSyncItems(TArrayView<const FContentBrowserItem> I
 
 	// Disable any plugin filters which hide the path we're navigating too in the path tree
 	bool bSomePluginPathFiltersChanged = false;
-	if (PathViewPtr->DisablePluginPathFiltersThatHideItems(ItemsToSync))
+	if (LegacyContentSourceWidgets->PathViewPtr->DisablePluginPathFiltersThatHideItems(ItemsToSync))
 	{
 		bSomePluginPathFiltersChanged = true;
 		bRepopulate = true;
@@ -1893,7 +1395,7 @@ void SContentBrowser::PrepareToSyncItems(TArrayView<const FContentBrowserItem> I
 	if (!bRepopulate)
 	{
 		bRepopulate = Algo::AnyOf(ItemsToSync, [this](const FContentBrowserItem& Item) {
-			return Item.IsFolder() && !PathViewPtr->DoesItemExist(Item.GetVirtualPath());
+			return Item.IsFolder() && !LegacyContentSourceWidgets->PathViewPtr->DoesItemExist(Item.GetVirtualPath());
 		});
 	}
 	
@@ -1901,7 +1403,7 @@ void SContentBrowser::PrepareToSyncItems(TArrayView<const FContentBrowserItem> I
 	{
 		// Disable the filter categories
 		// Do this before repopulate because the redirectors filter can hide folders
-		FilterListPtr->DisableFiltersThatHideItems(ItemsToSync);
+		LegacyContentSourceWidgets->FilterListPtr->DisableFiltersThatHideItems(ItemsToSync);
 	}
 
 	// If we have auto-enabled any flags or found a non-existant path, force a refresh
@@ -1946,19 +1448,22 @@ void SContentBrowser::PrepareToSyncItems(TArrayView<const FContentBrowserItem> I
 			FNotificationInfo NotificationInfo(NotificationBuilder.ToText());
 			FSlateNotificationManager::Get().AddNotification(NotificationInfo);
 		}
-		PathViewPtr->Populate();
-		FavoritePathViewPtr->Populate();
+		LegacyContentSourceWidgets->PathViewPtr->Populate();
+		LegacyContentSourceWidgets->FavoritePathViewPtr->Populate();
 	}
 
 	// Disable the filter search (reset the filter, then clear the search text)
 	// Note: we have to remove the filter immediately, we can't wait for OnSearchBoxChanged to hit
 	SetSearchBoxText(FText::GetEmpty());
-	SearchBoxPtr->SetText(FText::GetEmpty());
-	SearchBoxPtr->SetError(FText::GetEmpty());
+	LegacyContentSourceWidgets->SearchBoxPtr->SetText(FText::GetEmpty());
+	LegacyContentSourceWidgets->SearchBoxPtr->SetError(FText::GetEmpty());
 }
 
 void SContentBrowser::PrepareToSyncVirtualPaths(TArrayView<const FName> VirtualPathsToSync, const bool bDisableFiltersThatHideAssets)
 {
+	// For now we just forcefully enable the legacy content source when a sync is requested so it succeeds
+	ContentSourcesContainer->ActivateLegacyContentSource();
+	
 	// We need to try and resolve these paths back to items in order to query their attributes
 	// This will only work for items that have already been discovered
 	TArray<FContentBrowserItem> ItemsToSync;
@@ -1980,6 +1485,9 @@ void SContentBrowser::PrepareToSyncVirtualPaths(TArrayView<const FName> VirtualP
 
 void SContentBrowser::PrepareToSyncLegacy(TArrayView<const FAssetData> AssetDataList, TArrayView<const FString> FolderPaths, const bool bDisableFiltersThatHideAssets)
 {
+	// For now we just forcefully enable the legacy content source when a sync is requested so it succeeds
+	ContentSourcesContainer->ActivateLegacyContentSource();
+	
 	TArray<FName> VirtualPathsToSync;
 	ContentBrowserUtils::ConvertLegacySelectionToVirtualPaths(AssetDataList, FolderPaths, /*UseFolderPaths*/false, VirtualPathsToSync);
 
@@ -1988,16 +1496,25 @@ void SContentBrowser::PrepareToSyncLegacy(TArrayView<const FAssetData> AssetData
 
 void SContentBrowser::SyncToAssets(TArrayView<const FAssetData> AssetDataList, const bool bAllowImplicitSync, const bool bDisableFiltersThatHideAssets)
 {
+	// For now we just forcefully enable the legacy content source when a sync is requested so it succeeds
+	ContentSourcesContainer->ActivateLegacyContentSource();
+	
 	SyncToLegacy(AssetDataList, TArrayView<const FString>(), bAllowImplicitSync, bDisableFiltersThatHideAssets);
 }
 
 void SContentBrowser::SyncToFolders(TArrayView<const FString> FolderList, const bool bAllowImplicitSync)
 {
+	// For now we just forcefully enable the legacy content source when a sync is requested so it succeeds
+	ContentSourcesContainer->ActivateLegacyContentSource();
+	
 	SyncToLegacy(TArrayView<const FAssetData>(), FolderList, bAllowImplicitSync, /*bDisableFiltersThatHideAssets*/false);
 }
 
 void SContentBrowser::SyncToItems(TArrayView<const FContentBrowserItem> ItemsToSync, const bool bAllowImplicitSync, const bool bDisableFiltersThatHideAssets)
 {
+	// For now we just forcefully enable the legacy content source when a sync is requested so it succeeds
+	ContentSourcesContainer->ActivateLegacyContentSource();
+	
 	UContentBrowserDataSubsystem* ContentBrowserData = IContentBrowserDataModule::Get().GetSubsystem();
 	const TSharedRef<FPathPermissionList>& FolderPermissions = FAssetToolsModule::GetModule().Get().GetFolderPermissionList();
 
@@ -2006,13 +1523,16 @@ void SContentBrowser::SyncToItems(TArrayView<const FContentBrowserItem> ItemsToS
 	PrepareToSyncItems(ItemsToSync, bDisableFiltersThatHideAssets);
 
 	// Tell the sources view first so the asset view will be up to date by the time we request the sync
-	PathViewPtr->SyncToItems(ItemsToSync, bAllowImplicitSync);
-	FavoritePathViewPtr->SyncToItems(ItemsToSync, bAllowImplicitSync);
-	AssetViewPtr->SyncToItems(ItemsToSync);
+	LegacyContentSourceWidgets->PathViewPtr->SyncToItems(ItemsToSync, bAllowImplicitSync);
+	LegacyContentSourceWidgets->FavoritePathViewPtr->SyncToItems(ItemsToSync, bAllowImplicitSync);
+	LegacyContentSourceWidgets->AssetViewPtr->SyncToItems(ItemsToSync);
 }
 
 void SContentBrowser::SyncToVirtualPaths(TArrayView<const FName> VirtualPathsToSync, const bool bAllowImplicitSync, const bool bDisableFiltersThatHideAssets)
 {
+	// For now we just forcefully enable the legacy content source when a sync is requested so it succeeds
+	ContentSourcesContainer->ActivateLegacyContentSource();
+	
 	UContentBrowserDataSubsystem* ContentBrowserData = IContentBrowserDataModule::Get().GetSubsystem();
 	const TSharedRef<FPathPermissionList>& FolderPermissions = FAssetToolsModule::GetModule().Get().GetFolderPermissionList();
 
@@ -2042,13 +1562,16 @@ void SContentBrowser::SyncToVirtualPaths(TArrayView<const FName> VirtualPathsToS
 	PrepareToSyncVirtualPaths(VirtualPathsToSync, bDisableFiltersThatHideAssets);
 
 	// Tell the sources view first so the asset view will be up to date by the time we request the sync
-	PathViewPtr->SyncToVirtualPaths(VirtualPathsToSync, bAllowImplicitSync);
-	FavoritePathViewPtr->SyncToVirtualPaths(VirtualPathsToSync, bAllowImplicitSync);
-	AssetViewPtr->SyncToVirtualPaths(VirtualPathsToSync);
+	LegacyContentSourceWidgets->PathViewPtr->SyncToVirtualPaths(VirtualPathsToSync, bAllowImplicitSync);
+	LegacyContentSourceWidgets->FavoritePathViewPtr->SyncToVirtualPaths(VirtualPathsToSync, bAllowImplicitSync);
+	LegacyContentSourceWidgets->AssetViewPtr->SyncToVirtualPaths(VirtualPathsToSync);
 }
 
 void SContentBrowser::SyncToLegacy(TArrayView<const FAssetData> AssetDataList, TArrayView<const FString> FolderList, const bool bAllowImplicitSync, const bool bDisableFiltersThatHideAssets)
 {
+	// For now we just forcefully enable the legacy content source when a sync is requested so it succeeds
+	ContentSourcesContainer->ActivateLegacyContentSource();
+	
 	UContentBrowserDataSubsystem* ContentBrowserData = IContentBrowserDataModule::Get().GetSubsystem();
 	const TSharedRef<FPathPermissionList>& FolderPermissions = FAssetToolsModule::GetModule().Get().GetFolderPermissionList();
 
@@ -2080,13 +1603,16 @@ void SContentBrowser::SyncToLegacy(TArrayView<const FAssetData> AssetDataList, T
 	PrepareToSyncLegacy(AssetDataList, FolderList, bDisableFiltersThatHideAssets);
 
 	// Tell the sources view first so the asset view will be up to date by the time we request the sync
-	PathViewPtr->SyncToLegacy(AssetDataList, FolderList, bAllowImplicitSync);
-	FavoritePathViewPtr->SyncToLegacy(AssetDataList, FolderList, bAllowImplicitSync);
-	AssetViewPtr->SyncToLegacy(AssetDataList, FolderList);
+	LegacyContentSourceWidgets->PathViewPtr->SyncToLegacy(AssetDataList, FolderList, bAllowImplicitSync);
+	LegacyContentSourceWidgets->FavoritePathViewPtr->SyncToLegacy(AssetDataList, FolderList, bAllowImplicitSync);
+	LegacyContentSourceWidgets->AssetViewPtr->SyncToLegacy(AssetDataList, FolderList);
 }
 
 void SContentBrowser::SyncTo( const FContentBrowserSelection& ItemSelection, const bool bAllowImplicitSync, const bool bDisableFiltersThatHideAssets )
 {
+	// For now we just forcefully enable the legacy content source when a sync is requested so it succeeds
+	ContentSourcesContainer->ActivateLegacyContentSource();
+	
 	if (ItemSelection.IsLegacy())
 	{
 		SyncToLegacy(ItemSelection.SelectedAssets, ItemSelection.SelectedFolders, bAllowImplicitSync, bDisableFiltersThatHideAssets);
@@ -2138,7 +1664,7 @@ TSharedPtr<FTabManager> SContentBrowser::GetTabManager() const
 void SContentBrowser::LoadSelectedObjectsIfNeeded()
 {
 	// Get the selected assets in the asset view
-	const TArray<FAssetData>& SelectedAssets = AssetViewPtr->GetSelectedAssets();
+	const TArray<FAssetData>& SelectedAssets = LegacyContentSourceWidgets->AssetViewPtr->GetSelectedAssets();
 
 	// Load every asset that isn't already in memory
 	for ( auto AssetIt = SelectedAssets.CreateConstIterator(); AssetIt; ++AssetIt )
@@ -2161,22 +1687,38 @@ void SContentBrowser::LoadSelectedObjectsIfNeeded()
 
 void SContentBrowser::GetSelectedAssets(TArray<FAssetData>& SelectedAssets)
 {
-	SelectedAssets = AssetViewPtr->GetSelectedAssets();
+	// For now we just forcefully enable the legacy content source when this public function is called so it succeeds
+	ContentSourcesContainer->ActivateLegacyContentSource();
+	
+	SelectedAssets = LegacyContentSourceWidgets->AssetViewPtr->GetSelectedAssets();
 }
 
 void SContentBrowser::GetSelectedFolders(TArray<FString>& SelectedFolders)
 {
-	SelectedFolders = AssetViewPtr->GetSelectedFolders();
+	// For now we just forcefully enable the legacy content source when this public function is called so it succeeds
+	ContentSourcesContainer->ActivateLegacyContentSource();
+	
+	SelectedFolders = LegacyContentSourceWidgets->AssetViewPtr->GetSelectedFolders();
 }
 
 TArray<FString> SContentBrowser::GetSelectedPathViewFolders()
 {
-	check(PathViewPtr.IsValid());
-	return PathViewPtr->GetSelectedPaths();
+	// For now we just forcefully enable the legacy content source when this public function is called so it succeeds
+	ContentSourcesContainer->ActivateLegacyContentSource();
+	
+	check(LegacyContentSourceWidgets->PathViewPtr.IsValid());
+	return LegacyContentSourceWidgets->PathViewPtr->GetSelectedPaths();
 }
 
 void SContentBrowser::SaveSettings() const
 {
+	// Individual content sources will handle saving their own settings. If the legacy content source is active we save its settings, otherwise
+	// the settings were saved when the legacy content source was disabled
+	if (!ContentSourcesContainer->IsLegacyContentSourceActive())
+	{
+		return;
+	}
+	
 	const FString& SettingsString = InstanceName.ToString();
 
 	GConfig->SetBool(*SettingsIniSection, *(SettingsString + TEXT(".SourcesExpanded")), bSourcesViewExpanded, GEditorPerProjectIni);
@@ -2185,13 +1727,13 @@ void SContentBrowser::SaveSettings() const
 	FavoritesArea->SaveSettings(GEditorPerProjectIni, SettingsIniSection, SettingsString);
 	PathArea->SaveSettings(GEditorPerProjectIni, SettingsIniSection, SettingsString);
 
-	for(int32 SlotIndex = 0; SlotIndex < PathAssetSplitterPtr->GetChildren()->Num(); SlotIndex++)
+	for(int32 SlotIndex = 0; SlotIndex < LegacyContentSourceWidgets->PathAssetSplitterPtr->GetChildren()->Num(); SlotIndex++)
 	{
 		// First Slot containing the PathView is using SizeToContent so the SizeValue is not updated
 		// Adding another config line for older projects otherwise when they open for the first time after this update the SplitterSlot size will be based on the older one
 		// The older one is a normalized value, so it is too small and will make the SplitterSlot for the PathView seems like if it was collapsed the very first time you re-open it
 		const bool bIsFirstSlot = SlotIndex == 0;
-		float SplitterSize = bIsFirstSlot ? PathViewBoxWidth : PathAssetSplitterPtr->SlotAt(SlotIndex).GetSizeValue();
+		float SplitterSize = bIsFirstSlot ? PathViewBoxWidth : LegacyContentSourceWidgets->PathAssetSplitterPtr->SlotAt(SlotIndex).GetSizeValue();
 		if (bIsFirstSlot)
 		{
 			GConfig->SetFloat(*SettingsIniSection, *(SettingsString + FString::Printf(TEXT(".VerticalSplitter.FixedSlotSize%d"), SlotIndex)), SplitterSize, GEditorPerProjectIni);
@@ -2203,7 +1745,8 @@ void SContentBrowser::SaveSettings() const
 	}
 
 	{
-		SSplitter& Splitter = UE::Editor::ContentBrowser::IsNewStyleEnabled() ? *SourceTreePtr->GetSplitter() : *PathFavoriteSplitterPtr;
+		SSplitter& Splitter = UE::Editor::ContentBrowser::IsNewStyleEnabled() ?
+			*LegacyContentSourceWidgets->SourceTreePtr->GetSplitter() : *LegacyContentSourceWidgets->PathFavoriteSplitterPtr;
 		for (int32 SlotIndex = 0; SlotIndex < Splitter.GetChildren()->Num(); SlotIndex++)
 		{
 			FString Key, Filename;
@@ -2215,14 +1758,14 @@ void SContentBrowser::SaveSettings() const
 	}
 
 	// Save all our data using the settings string as a key in the user settings ini
-	FilterListPtr->SaveSettings(GEditorPerProjectIni, SettingsIniSection, SettingsString);
-	PathViewPtr->SaveSettings(GEditorPerProjectIni, SettingsIniSection, SettingsString);
-	FavoritePathViewPtr->SaveSettings(GEditorPerProjectIni, SettingsIniSection, SettingsString + TEXT(".Favorites"));
+	LegacyContentSourceWidgets->FilterListPtr->SaveSettings(GEditorPerProjectIni, SettingsIniSection, SettingsString);
+	LegacyContentSourceWidgets->PathViewPtr->SaveSettings(GEditorPerProjectIni, SettingsIniSection, SettingsString);
+	LegacyContentSourceWidgets->FavoritePathViewPtr->SaveSettings(GEditorPerProjectIni, SettingsIniSection, SettingsString + TEXT(".Favorites"));
 	for (const TUniquePtr<FCollectionSource>& CollectionSource : CollectionSources)
 	{
 		CollectionSource->SaveSettings(InstanceName);
 	}
-	AssetViewPtr->SaveSettings(GEditorPerProjectIni, SettingsIniSection, SettingsString);
+	LegacyContentSourceWidgets->AssetViewPtr->SaveSettings(GEditorPerProjectIni, SettingsIniSection, SettingsString);
 	
 	GConfig->SetArray(*SettingsIniSection, *(SettingsString + TEXT(".JumpMRU")), JumpMRU, GEditorPerProjectIni);
 }
@@ -2239,27 +1782,38 @@ bool SContentBrowser::IsLocked() const
 
 void SContentBrowser::SetKeyboardFocusOnSearch() const
 {
-	// Focus on the search box
-	FSlateApplication::Get().SetKeyboardFocus( SearchBoxPtr, EFocusCause::SetDirectly );
+	if (ContentSourcesContainer->IsLegacyContentSourceActive())
+	{
+		// Focus on the search box
+		FSlateApplication::Get().SetKeyboardFocus( LegacyContentSourceWidgets->SearchBoxPtr, EFocusCause::SetDirectly );
+	}
 }
 
 void SContentBrowser::CopySettingsFromBrowser(TSharedPtr<SContentBrowser> OtherBrowser)
 {
-	FName InstanceNameToCopyFrom = OtherBrowser->GetInstanceName();
+	if (ContentSourcesContainer->IsLegacyContentSourceActive())
+	{
+		FName InstanceNameToCopyFrom = OtherBrowser->GetInstanceName();
 
-	// Clear out any existing settings that dont get reset on load
-	FilterListPtr->RemoveAllFilters();
+		// Clear out any existing settings that dont get reset on load
+		LegacyContentSourceWidgets->FilterListPtr->RemoveAllFilters();
 
-	LoadSettings(InstanceNameToCopyFrom);
+		LoadSettings(InstanceNameToCopyFrom);
+	}
+	
 }
 
 FReply SContentBrowser::OnKeyDown( const FGeometry& MyGeometry, const FKeyEvent& InKeyEvent )
 {
-	bool bIsRenamingAsset = AssetViewPtr && AssetViewPtr->IsRenamingAsset();
-	if(bIsRenamingAsset || Commands->ProcessCommandBindings( InKeyEvent ) )
+	if (ContentSourcesContainer->IsLegacyContentSourceActive())
 	{
-		return FReply::Handled();
+		bool bIsRenamingAsset = LegacyContentSourceWidgets->AssetViewPtr && LegacyContentSourceWidgets->AssetViewPtr->IsRenamingAsset();
+		if(bIsRenamingAsset || Commands->ProcessCommandBindings( InKeyEvent ) )
+		{
+			return FReply::Handled();
+		}
 	}
+	
 	return FReply::Unhandled();
 }
 
@@ -2327,21 +1881,21 @@ void SContentBrowser::GetSourceTreeSplitterSlotSizeSettingKeyAndFilename(int32 S
 {
 	OutKey = InstanceName.ToString();
 
-	if (SlotIndex < SourceTreeSplitterNumFixedSlots)
+	if (SlotIndex < LegacyContentSourceWidgets->SourceTreeSplitterNumFixedSlots)
 	{
 		OutKey += FString::Printf(TEXT(".FavoriteSplitter.SlotSize%d"), SlotIndex);
 		OutFilename = GEditorPerProjectIni;
 	}
 	else
 	{
-		const TUniquePtr<FCollectionSource>& CollectionSource = CollectionSources[SlotIndex - SourceTreeSplitterNumFixedSlots];
+		const TUniquePtr<FCollectionSource>& CollectionSource = CollectionSources[SlotIndex - LegacyContentSourceWidgets->SourceTreeSplitterNumFixedSlots];
 		if (CollectionSource->IsProjectCollectionContainer())
 		{
 			// Reconsider the .FavoriteSplitter.SlotSize key naming scheme if you hit this check.
-			check(SourceTreeSplitterNumFixedSlots == 2);
+			check(LegacyContentSourceWidgets->SourceTreeSplitterNumFixedSlots == 2);
 
 			// Maintain backwards compatibility with previous version of the SContentBrowser which had a single collection view.
-			OutKey += FString::Printf(TEXT(".FavoriteSplitter.SlotSize%d"), SourceTreeSplitterNumFixedSlots);
+			OutKey += FString::Printf(TEXT(".FavoriteSplitter.SlotSize%d"), LegacyContentSourceWidgets->SourceTreeSplitterNumFixedSlots);
 		}
 		else
 		{
@@ -2353,6 +1907,13 @@ void SContentBrowser::GetSourceTreeSplitterSlotSizeSettingKeyAndFilename(int32 S
 
 void SContentBrowser::LoadSettings(const FName& InInstanceName)
 {
+	// Individual content sources will handle saving their own settings. If the legacy content source is active we load the settings, otherwise
+	// the settings will be loaded when the legacy content source is enabled
+	if (!ContentSourcesContainer->IsLegacyContentSourceActive())
+	{
+		return;
+	}
+	
 	FString SettingsString = InInstanceName.ToString();
 
 	// Now that we have determined the appropriate settings string, actually load the settings
@@ -2365,13 +1926,13 @@ void SContentBrowser::LoadSettings(const FName& InInstanceName)
 	FavoritesArea->LoadSettings(GEditorPerProjectIni, SettingsIniSection, SettingsString);
 	PathArea->LoadSettings(GEditorPerProjectIni, SettingsIniSection, SettingsString);
 
-	for(int32 SlotIndex = 0; SlotIndex < PathAssetSplitterPtr->GetChildren()->Num(); SlotIndex++)
+	for(int32 SlotIndex = 0; SlotIndex < LegacyContentSourceWidgets->PathAssetSplitterPtr->GetChildren()->Num(); SlotIndex++)
 	{
 		// First Slot containing the PathView is using SizeToContent so the SizeValue is not updated
 		// Adding another config line for older projects otherwise when they open for the first time after this update the SplitterSlot size will be based on the older one
 		// The older one is a normalized value, so it is too small and make will make the SplitterSlot for the PathView seems like if it was collapsed the very first time you re-open it
 		const bool bIsFirstSlot = SlotIndex == 0;
-		float SplitterSize = bIsFirstSlot ? PathViewBoxWidth : PathAssetSplitterPtr->SlotAt(SlotIndex).GetSizeValue();
+		float SplitterSize = bIsFirstSlot ? PathViewBoxWidth : LegacyContentSourceWidgets->PathAssetSplitterPtr->SlotAt(SlotIndex).GetSizeValue();
 		if (bIsFirstSlot)
 		{
 			GConfig->GetFloat(*SettingsIniSection, *(SettingsString + FString::Printf(TEXT(".VerticalSplitter.FixedSlotSize%d"), SlotIndex)), SplitterSize, GEditorPerProjectIni);
@@ -2380,12 +1941,13 @@ void SContentBrowser::LoadSettings(const FName& InInstanceName)
 		else
 		{
 			GConfig->GetFloat(*SettingsIniSection, *(SettingsString + FString::Printf(TEXT(".VerticalSplitter.SlotSize%d"), SlotIndex)), SplitterSize, GEditorPerProjectIni);
-			PathAssetSplitterPtr->SlotAt(SlotIndex).SetSizeValue(SplitterSize);
+			LegacyContentSourceWidgets->PathAssetSplitterPtr->SlotAt(SlotIndex).SetSizeValue(SplitterSize);
 		}
 	}
 
 	{
-		SSplitter& Splitter = UE::Editor::ContentBrowser::IsNewStyleEnabled() ? *SourceTreePtr->GetSplitter() : *PathFavoriteSplitterPtr;
+		SSplitter& Splitter = UE::Editor::ContentBrowser::IsNewStyleEnabled() ?
+			*LegacyContentSourceWidgets->SourceTreePtr->GetSplitter() : *LegacyContentSourceWidgets->PathFavoriteSplitterPtr;
 		for (int32 SlotIndex = 0; SlotIndex < Splitter.GetChildren()->Num(); SlotIndex++)
 		{
 			FString Key, Filename;
@@ -2398,14 +1960,14 @@ void SContentBrowser::LoadSettings(const FName& InInstanceName)
 	}
 
 	// Save all our data using the settings string as a key in the user settings ini
-	FilterListPtr->LoadSettings(InInstanceName, GEditorPerProjectIni, SettingsIniSection, SettingsString);
-	PathViewPtr->LoadSettings(GEditorPerProjectIni, SettingsIniSection, SettingsString);
-	FavoritePathViewPtr->LoadSettings(GEditorPerProjectIni, SettingsIniSection, SettingsString + TEXT(".Favorites"));
+	LegacyContentSourceWidgets->FilterListPtr->LoadSettings(InInstanceName, GEditorPerProjectIni, SettingsIniSection, SettingsString);
+	LegacyContentSourceWidgets->PathViewPtr->LoadSettings(GEditorPerProjectIni, SettingsIniSection, SettingsString);
+	LegacyContentSourceWidgets->FavoritePathViewPtr->LoadSettings(GEditorPerProjectIni, SettingsIniSection, SettingsString + TEXT(".Favorites"));
 	for (const TUniquePtr<FCollectionSource>& CollectionSource : CollectionSources)
 	{
 		CollectionSource->LoadSettings(InstanceName);
 	}
-	AssetViewPtr->LoadSettings(GEditorPerProjectIni, SettingsIniSection, SettingsString);
+	LegacyContentSourceWidgets->AssetViewPtr->LoadSettings(GEditorPerProjectIni, SettingsIniSection, SettingsString);
 
 	GConfig->GetArray(*SettingsIniSection, *(SettingsString + TEXT(".JumpMRU")), JumpMRU, GEditorPerProjectIni);
 }
@@ -2442,7 +2004,7 @@ void SContentBrowser::SourcesChanged(const TArray<FString>& SelectedPaths, const
 
 		const FText DynamicQueryText = FText::FromString(DynamicQueryString);
 		SetSearchBoxText(DynamicQueryText);
-		SearchBoxPtr->SetText(DynamicQueryText);
+		LegacyContentSourceWidgets->SearchBoxPtr->SetText(DynamicQueryText);
 	}
 	else if (StashedSearchBoxText.IsSet())
 	{
@@ -2451,17 +2013,17 @@ void SContentBrowser::SourcesChanged(const TArray<FString>& SelectedPaths, const
 		StashedSearchBoxText.Reset();
 
 		SetSearchBoxText(StashedText);
-		SearchBoxPtr->SetText(StashedText);
+		LegacyContentSourceWidgets->SearchBoxPtr->SetText(StashedText);
 	}
 
-	if (!AssetViewPtr->GetContentSources().IsEmpty())
+	if (!LegacyContentSourceWidgets->AssetViewPtr->GetContentSources().IsEmpty())
 	{
 		// Update the current history data to preserve selection if there is a valid ContentSources
 		HistoryManager.UpdateHistoryData();
 	}
 
 	// Change the filter for the asset view
-	AssetViewPtr->SetContentSources(ContentSources);
+	LegacyContentSourceWidgets->AssetViewPtr->SetContentSources(ContentSources);
 
 	// Add a new history data now that the source has changed
 	HistoryManager.AddHistoryData();
@@ -2504,7 +2066,7 @@ void SContentBrowser::FolderEntered(const FContentBrowserItem& Folder)
 		// set the path view to the incoming path
 		TArray<FString> SelectedPaths;
 		SelectedPaths.Add(Folder.GetVirtualPath().ToString());
-		PathViewPtr->SetSelectedPaths(SelectedPaths);
+		LegacyContentSourceWidgets->PathViewPtr->SetSelectedPaths(SelectedPaths);
 
 		PathSelected(SelectedPaths[0]);
 	}
@@ -2520,13 +2082,13 @@ void SContentBrowser::PathSelected(const FString& FolderPath)
 		CollectionSource->CollectionViewPtr->ClearSelection();
 	}
 
-	TArray<FString> SelectedPaths = PathViewPtr->GetSelectedPaths();
+	TArray<FString> SelectedPaths = LegacyContentSourceWidgets->PathViewPtr->GetSelectedPaths();
 	// Selecting a folder shows it in the favorite list also
-	FavoritePathViewPtr->SetSelectedPaths(SelectedPaths);
+	LegacyContentSourceWidgets->FavoritePathViewPtr->SetSelectedPaths(SelectedPaths);
 	TArray<FCollectionRef> SelectedCollections;
 	SourcesChanged(SelectedPaths, SelectedCollections);
 
-	PathViewPtr->SaveSettings(GEditorPerProjectIni, SettingsIniSection, InstanceName.ToString());
+	LegacyContentSourceWidgets->PathViewPtr->SaveSettings(GEditorPerProjectIni, SettingsIniSection, InstanceName.ToString());
 
 	// Notify 'asset path changed' delegate
 	FContentBrowserModule& ContentBrowserModule = FModuleManager::GetModuleChecked<FContentBrowserModule>( TEXT("ContentBrowser") );
@@ -2537,7 +2099,7 @@ void SContentBrowser::PathSelected(const FString& FolderPath)
 	}
 
 	// Update the context menu's selected paths list
-	PathContextMenu->SetSelectedFolders(PathViewPtr->GetSelectedFolderItems());
+	LegacyContentSourceWidgets->PathContextMenu->SetSelectedFolders(LegacyContentSourceWidgets->PathViewPtr->GetSelectedFolderItems());
 }
 
 void SContentBrowser::FavoritePathSelected(const FString& FolderPath)
@@ -2550,13 +2112,13 @@ void SContentBrowser::FavoritePathSelected(const FString& FolderPath)
 		CollectionSource->CollectionViewPtr->ClearSelection();
 	}
 
-	TArray<FString> SelectedPaths = FavoritePathViewPtr->GetSelectedPaths();
+	TArray<FString> SelectedPaths = LegacyContentSourceWidgets->FavoritePathViewPtr->GetSelectedPaths();
 	// Selecting a favorite shows it in the main list also
-	PathViewPtr->SetSelectedPaths(SelectedPaths);
+	LegacyContentSourceWidgets->PathViewPtr->SetSelectedPaths(SelectedPaths);
 	TArray<FCollectionRef> SelectedCollections;
 	SourcesChanged(SelectedPaths, SelectedCollections);
 
-	PathViewPtr->SaveSettings(GEditorPerProjectIni, SettingsIniSection, InstanceName.ToString());
+	LegacyContentSourceWidgets->PathViewPtr->SaveSettings(GEditorPerProjectIni, SettingsIniSection, InstanceName.ToString());
 
 	// Notify 'asset path changed' delegate
 	FContentBrowserModule& ContentBrowserModule = FModuleManager::GetModuleChecked<FContentBrowserModule>(TEXT("ContentBrowser"));
@@ -2567,19 +2129,19 @@ void SContentBrowser::FavoritePathSelected(const FString& FolderPath)
 	}
 
 	// Update the context menu's selected paths list
-	PathContextMenu->SetSelectedFolders(FavoritePathViewPtr->GetSelectedFolderItems());
+	LegacyContentSourceWidgets->PathContextMenu->SetSelectedFolders(LegacyContentSourceWidgets->FavoritePathViewPtr->GetSelectedFolderItems());
 }
 
 TSharedRef<FExtender> SContentBrowser::GetPathContextMenuExtender(const TArray<FString>& InSelectedPaths) const
 {
-	return PathContextMenu->MakePathViewContextMenuExtender(InSelectedPaths);
+	return LegacyContentSourceWidgets->PathContextMenu->MakePathViewContextMenuExtender(InSelectedPaths);
 }
 
 void SContentBrowser::CollectionSelected(const TSharedPtr<ICollectionContainer>& CollectionContainer, const FCollectionNameType& SelectedCollection)
 {
 	// You may not select both collections and paths
-	PathViewPtr->ClearSelection();
-	FavoritePathViewPtr->ClearSelection();
+	LegacyContentSourceWidgets->PathViewPtr->ClearSelection();
+	LegacyContentSourceWidgets->FavoritePathViewPtr->ClearSelection();
 
 	TArray<FCollectionRef> SelectedCollections;
 	for (const TUniquePtr<FCollectionSource>& CollectionSource : CollectionSources)
@@ -2614,33 +2176,41 @@ void SContentBrowser::SetSelectedPaths(const TArray<FString>& FolderPaths, bool
 	{
 		if (bNeedsRefresh)
 		{
-			PathViewPtr->Populate();
-			FavoritePathViewPtr->Populate();
+			LegacyContentSourceWidgets->PathViewPtr->Populate();
+			LegacyContentSourceWidgets->FavoritePathViewPtr->Populate();
 		}
 
-		PathViewPtr->SetSelectedPaths(FolderPaths);
-		FavoritePathViewPtr->SetSelectedPaths(FolderPaths);
+		LegacyContentSourceWidgets->PathViewPtr->SetSelectedPaths(FolderPaths);
+		LegacyContentSourceWidgets->FavoritePathViewPtr->SetSelectedPaths(FolderPaths);
 		PathSelected(FolderPaths[0]);
 	}
 }
 
 void SContentBrowser::ForceShowPluginContent(bool bEnginePlugin)
 {
-	if (AssetViewPtr.IsValid())
+	// For now we just forcefully enable the legacy content source when this public function is called so it succeeds
+	ContentSourcesContainer->ActivateLegacyContentSource();
+	
+	if (LegacyContentSourceWidgets->AssetViewPtr.IsValid())
 	{
-		AssetViewPtr->ForceShowPluginFolder(bEnginePlugin);
+		LegacyContentSourceWidgets->AssetViewPtr->ForceShowPluginFolder(bEnginePlugin);
 	}
 }
 
 void SContentBrowser::OnApplyHistoryData( const FHistoryData& History )
 {
-	PathViewPtr->ApplyHistoryData(History);
-	FavoritePathViewPtr->ApplyHistoryData(History);
+	if (!ContentSourcesContainer->IsLegacyContentSourceActive())
+	{
+		return;
+	}
+	
+	LegacyContentSourceWidgets->PathViewPtr->ApplyHistoryData(History);
+	LegacyContentSourceWidgets->FavoritePathViewPtr->ApplyHistoryData(History);
 	for (const TUniquePtr<FCollectionSource>& CollectionSource : CollectionSources)
 	{
 		CollectionSource->CollectionViewPtr->ApplyHistoryData(History);
 	}
-	AssetViewPtr->ApplyHistoryData(History);
+	LegacyContentSourceWidgets->AssetViewPtr->ApplyHistoryData(History);
 
 	// Update the breadcrumb trail path
 	UpdatePath();
@@ -2659,10 +2229,16 @@ void SContentBrowser::OnApplyHistoryData( const FHistoryData& History )
 
 void SContentBrowser::OnUpdateHistoryData(FHistoryData& HistoryData) const
 {
-	const FAssetViewContentSources& ContentSources = AssetViewPtr->GetContentSources();
-	const TArray<FContentBrowserItem> SelectedItems = AssetViewPtr->GetSelectedItems();
+	if (!ContentSourcesContainer->IsLegacyContentSourceActive())
+	{
+		return;
+	}
+	
+	const FAssetViewContentSources& ContentSources = LegacyContentSourceWidgets->AssetViewPtr->GetContentSources();
+	const TArray<FContentBrowserItem> SelectedItems = LegacyContentSourceWidgets->AssetViewPtr->GetSelectedItems();
 
-	const FText NewSource = ContentSources.HasVirtualPaths() ? FText::FromName(ContentSources.GetVirtualPaths()[0]) : (ContentSources.HasCollections() ? FText::FromName(ContentSources.GetCollections()[0].Name) : LOCTEXT("AllAssets", "All Assets"));
+	const FText NewSource = ContentSources.HasVirtualPaths() ? FText::FromName(ContentSources.GetVirtualPaths()[0]) :
+		(ContentSources.HasCollections() ? FText::FromName(ContentSources.GetCollections()[0].Name) : LOCTEXT("AllAssets", "All Assets"));
 
 	HistoryData.HistoryDesc = NewSource;
 	HistoryData.ContentSources = ContentSources;
@@ -2676,25 +2252,25 @@ void SContentBrowser::OnUpdateHistoryData(FHistoryData& HistoryData) const
 
 void SContentBrowser::NewFolderRequested(const FString& SelectedPath)
 {
-	if( ensure(SelectedPath.Len() > 0) && AssetViewPtr.IsValid() )
+	if( ensure(SelectedPath.Len() > 0) && LegacyContentSourceWidgets->AssetViewPtr.IsValid() )
 	{
-		CreateNewFolder(SelectedPath, FOnCreateNewFolder::CreateSP(AssetViewPtr.Get(), &SAssetView::NewFolderItemRequested));
+		CreateNewFolder(SelectedPath, FOnCreateNewFolder::CreateSP(LegacyContentSourceWidgets->AssetViewPtr.Get(), &SAssetView::NewFolderItemRequested));
 	}
 }
 
 void SContentBrowser::NewFileItemRequested(const FContentBrowserItemDataTemporaryContext& NewItemContext)
 {
-	if (AssetViewPtr)
+	if (LegacyContentSourceWidgets && LegacyContentSourceWidgets->AssetViewPtr)
 	{
-		AssetViewPtr->NewFileItemRequested(NewItemContext);
+		LegacyContentSourceWidgets->AssetViewPtr->NewFileItemRequested(NewItemContext);
 	}
 }
 
 void SContentBrowser::SetSearchText(const FText& InSearchText)
 {
-	if (SearchBoxPtr.IsValid())
+	if (LegacyContentSourceWidgets && LegacyContentSourceWidgets->SearchBoxPtr.IsValid())
 	{
-		SearchBoxPtr->SetText(InSearchText);
+		LegacyContentSourceWidgets->SearchBoxPtr->SetText(InSearchText);
 	}
 }
 
@@ -2704,14 +2280,14 @@ void SContentBrowser::SetSearchBoxText(const FText& InSearchText)
 	if (!InSearchText.ToString().Equals(TextFilter->GetRawFilterText().ToString(), ESearchCase::CaseSensitive))
 	{
 		TextFilter->SetRawFilterText(InSearchText);
-		SearchBoxPtr->SetError(TextFilter->GetFilterErrorText());
+		LegacyContentSourceWidgets->SearchBoxPtr->SetError(TextFilter->GetFilterErrorText());
 		if (InSearchText.IsEmpty())
 		{
-			AssetViewPtr->SetUserSearching(false);
+			LegacyContentSourceWidgets->AssetViewPtr->SetUserSearching(false);
 		}
 		else
 		{
-			AssetViewPtr->SetUserSearching(true);
+			LegacyContentSourceWidgets->AssetViewPtr->SetUserSearching(true);
 		}
 	}
 }
@@ -2732,9 +2308,9 @@ void SContentBrowser::OnSearchBoxCommitted(const FText& InSearchText, ETextCommi
 
 void SContentBrowser::Tick(const FGeometry& AllottedGeometry, const double InCurrentTime, const float InDeltaTime)
 {
-	if (UE::Editor::ContentBrowser::IsNewStyleEnabled())
+	if (UE::Editor::ContentBrowser::IsNewStyleEnabled() && ContentSourcesContainer->IsLegacyContentSourceActive())
 	{
-		NavigationToolBarParentSwitcher->Update();
+		LegacyContentSourceWidgets->NavigationToolBarParentSwitcher->Update();
 	}
 
 	SCompoundWidget::Tick(AllottedGeometry, InCurrentTime, InDeltaTime);
@@ -2775,7 +2351,7 @@ void SContentBrowser::OnSaveSearchButtonClicked(const FText& InSearchText)
 	// We want to add any currently selected paths to the final saved query so that you get back roughly the same list of objects as what you're currently seeing
 	FString SelectedPathsQuery;
 	{
-		const TArray<FName>& VirtualPaths = AssetViewPtr->GetContentSources().GetVirtualPaths();
+		const TArray<FName>& VirtualPaths = LegacyContentSourceWidgets->AssetViewPtr->GetContentSources().GetVirtualPaths();
 		for (int32 SelectedPathIndex = 0; SelectedPathIndex < VirtualPaths.Num(); ++SelectedPathIndex)
 		{
 			SelectedPathsQuery.Append(TEXT("Path:'"));
@@ -2871,12 +2447,12 @@ void SContentBrowser::OnSaveSearchButtonClicked(const FText& InSearchText)
 
 void SContentBrowser::SaveSearchAsFilter()
 {
-	FilterListPtr->CreateCustomFilterDialog(TextFilter->GetRawFilterText());
+	LegacyContentSourceWidgets->FilterListPtr->CreateCustomFilterDialog(TextFilter->GetRawFilterText());
 }
 
 void SContentBrowser::EditPathCommand()
 {
-	NavigationBar->StartEditingPath();
+	LegacyContentSourceWidgets->NavigationBar->StartEditingPath();
 }
 
 void SContentBrowser::OnNavigateToPath(const FString& NewPath)
@@ -2890,7 +2466,7 @@ void SContentBrowser::OnNavigateToPath(const FString& NewPath)
 
 void SContentBrowser::OnPathClicked( const FString& CrumbData )
 {
-	FAssetViewContentSources ContentSources = AssetViewPtr->GetContentSources();
+	FAssetViewContentSources ContentSources = LegacyContentSourceWidgets->AssetViewPtr->GetContentSources();
 
 	if (ContentSources.HasCollections() )
 	{
@@ -2918,8 +2494,8 @@ void SContentBrowser::OnPathClicked( const FString& CrumbData )
 		// More than one path is selected or the crumb that was clicked is not the same path as the current one. Change the path.
 		TArray<FString> SelectedPaths;
 		SelectedPaths.Add(CrumbData);
-		PathViewPtr->SetSelectedPaths(SelectedPaths);
-		FavoritePathViewPtr->SetSelectedPaths(SelectedPaths);
+		LegacyContentSourceWidgets->PathViewPtr->SetSelectedPaths(SelectedPaths);
+		LegacyContentSourceWidgets->FavoritePathViewPtr->SetSelectedPaths(SelectedPaths);
 		PathSelected(SelectedPaths[0]);
 	}
 }
@@ -2936,7 +2512,7 @@ void SContentBrowser::OnPathMenuItemClicked(FString ClickedPath)
 
 bool SContentBrowser::OnCanEditPathAsText(const FString& Text) const
 {
-	const FAssetViewContentSources& ContentSources = AssetViewPtr->GetContentSources();
+	const FAssetViewContentSources& ContentSources = LegacyContentSourceWidgets->AssetViewPtr->GetContentSources();
 	if (ContentSources.HasCollections())
 	{
 		// Do not present collections as text because their names are not very user friendly right now.
@@ -2961,10 +2537,10 @@ TArray<FString> SContentBrowser::OnCompletePathPrefix(const FString& Prefix) con
 	UContentBrowserDataSubsystem* ContentBrowserData = IContentBrowserDataModule::Get().GetSubsystem();
 	TArray<FContentBrowserItem> SubItems = ContentBrowserUtils::GetChildItemsFromVirtualPath(
 		Parent,
-		PathViewPtr->GetContentBrowserItemCategoryFilter(),
-		PathViewPtr->GetContentBrowserItemAttributeFilter(),
+		LegacyContentSourceWidgets->PathViewPtr->GetContentBrowserItemCategoryFilter(),
+		LegacyContentSourceWidgets->PathViewPtr->GetContentBrowserItemAttributeFilter(),
 		InstanceName,
-		*PathViewPtr); 
+		*LegacyContentSourceWidgets->PathViewPtr); 
 	TArray<FString> Results;
 	for (const FContentBrowserItem& Item : SubItems)
 	{
@@ -2980,7 +2556,7 @@ TArray<FString> SContentBrowser::OnCompletePathPrefix(const FString& Prefix) con
 
 TSharedRef<SWidget> SContentBrowser::OnGetCrumbDelimiterContent(const FString& CrumbData) const
 {
-	const FAssetViewContentSources& ContentSources = AssetViewPtr->GetContentSources();
+	const FAssetViewContentSources& ContentSources = LegacyContentSourceWidgets->AssetViewPtr->GetContentSources();
 
 	TSharedPtr<SWidget> Widget = SNullWidget::NullWidget;
 	TSharedPtr<SWidget> MenuWidget;
@@ -3018,10 +2594,10 @@ TSharedRef<SWidget> SContentBrowser::OnGetCrumbDelimiterContent(const FString& C
 	{
 		TArray<FContentBrowserItem> SubItems = ContentBrowserUtils::GetChildItemsFromVirtualPath(
 			*CrumbData,
-			PathViewPtr->GetContentBrowserItemCategoryFilter(),
-			PathViewPtr->GetContentBrowserItemAttributeFilter(),
+			LegacyContentSourceWidgets->PathViewPtr->GetContentBrowserItemCategoryFilter(),
+			LegacyContentSourceWidgets->PathViewPtr->GetContentBrowserItemAttributeFilter(),
 			InstanceName,
-			*PathViewPtr); 
+			*LegacyContentSourceWidgets->PathViewPtr); 
 		SubItems.Sort([](const FContentBrowserItem& ItemOne, const FContentBrowserItem& ItemTwo)
 		{
 			return ItemOne.GetDisplayName().CompareTo(ItemTwo.GetDisplayName()) < 0;
@@ -3118,7 +2694,7 @@ bool SContentBrowser::ParseCollectionCrumbData(const FString& CrumbData, FCollec
 FString SContentBrowser::GetCurrentPath(const EContentBrowserPathType PathType) const
 {
 	FString CurrentPath;
-	const FAssetViewContentSources& ContentSources = AssetViewPtr->GetContentSources();
+	const FAssetViewContentSources& ContentSources = LegacyContentSourceWidgets->AssetViewPtr->GetContentSources();
 	if ( ContentSources.HasVirtualPaths() && ContentSources.GetVirtualPaths()[0] != NAME_None )
 	{
 		if (PathType == EContentBrowserPathType::Virtual)
@@ -3179,7 +2755,7 @@ void SContentBrowser::AppendNewMenuContextObjects(const EContentBrowserDataMenuC
 
 TSharedRef<SWidget> SContentBrowser::MakeAddNewContextMenu(const EContentBrowserDataMenuContext_AddNewMenuDomain InDomain, UContentBrowserToolbarMenuContext* CommonContext)
 {
-	const FAssetViewContentSources& ContentSources = AssetViewPtr->GetContentSources();
+	const FAssetViewContentSources& ContentSources = LegacyContentSourceWidgets->AssetViewPtr->GetContentSources();
 
 	bool bCanBeModified = false;
 
@@ -3264,9 +2840,9 @@ void SContentBrowser::PopulateAddNewContextMenu(class UToolMenu* Menu)
 
 bool SContentBrowser::CanWriteToCurrentPath() const
 {
-	if (AssetViewPtr.IsValid())
+	if (LegacyContentSourceWidgets && LegacyContentSourceWidgets->AssetViewPtr.IsValid())
 	{
-		const FAssetViewContentSources& ContentSources = AssetViewPtr->GetContentSources();
+		const FAssetViewContentSources& ContentSources = LegacyContentSourceWidgets->AssetViewPtr->GetContentSources();
 		if (ContentSources.GetVirtualPaths().Num() == 1)
 		{
 			FName CurrentPath = ContentSources.GetVirtualPaths()[0];
@@ -3315,25 +2891,25 @@ bool SContentBrowser::CanWriteToPath(const FContentBrowserItemPath InPath) const
 
 void SContentBrowser::AddCustomTextFilter(const FCustomTextFilterData& FilterData, bool bApplyFilter)
 {
-	if (FilterListPtr.IsValid())
+	if (LegacyContentSourceWidgets && LegacyContentSourceWidgets->FilterListPtr.IsValid())
 	{
-		FilterListPtr->OnCreateCustomTextFilter(FilterData, bApplyFilter);
+		LegacyContentSourceWidgets->FilterListPtr->OnCreateCustomTextFilter(FilterData, bApplyFilter);
 	}
 }
 
 void SContentBrowser::DeleteCustomTextFilterByLabel(const FText& FilterLabel)
 {
-	if (FilterListPtr.IsValid())
+	if (LegacyContentSourceWidgets && LegacyContentSourceWidgets->FilterListPtr.IsValid())
 	{
-		FilterListPtr->DeleteCustomTextFilterByLabel(FilterLabel);
+		LegacyContentSourceWidgets->FilterListPtr->DeleteCustomTextFilterByLabel(FilterLabel);
 	}
 }
 
 void SContentBrowser::ModifyCustomTextFilterByLabel(const FCustomTextFilterData& NewFilterData, const FText& FilterLabel)
 {
-	if (FilterListPtr.IsValid())
+	if (LegacyContentSourceWidgets && LegacyContentSourceWidgets->FilterListPtr.IsValid())
 	{
-		FilterListPtr->ModifyCustomTextFilterByLabel(NewFilterData, FilterLabel);
+		LegacyContentSourceWidgets->FilterListPtr->ModifyCustomTextFilterByLabel(NewFilterData, FilterLabel);
 	}
 }
 
@@ -3341,10 +2917,10 @@ bool SContentBrowser::IsAssetViewDoneFiltering()
 {
 	bool isDoneFiltering = false;
 
-	if (AssetViewPtr.IsValid())
+	if (LegacyContentSourceWidgets && LegacyContentSourceWidgets->AssetViewPtr.IsValid())
 	{
-		isDoneFiltering = !AssetViewPtr->HasItemsPendingFilter()
-			&& !AssetViewPtr->HasThumbnailsPendingUpdate();
+		isDoneFiltering = !LegacyContentSourceWidgets->AssetViewPtr->HasItemsPendingFilter()
+			&& !LegacyContentSourceWidgets->AssetViewPtr->HasThumbnailsPendingUpdate();
 	}
 
 	return isDoneFiltering;
@@ -3357,7 +2933,7 @@ bool SContentBrowser::IsAddNewEnabled() const
 
 FText SContentBrowser::GetAddNewToolTipText() const
 {
-	const FAssetViewContentSources& ContentSources = AssetViewPtr->GetContentSources();
+	const FAssetViewContentSources& ContentSources = LegacyContentSourceWidgets->AssetViewPtr->GetContentSources();
 
 	if ( ContentSources.GetVirtualPaths().Num() == 1 )
 	{
@@ -3380,9 +2956,9 @@ FText SContentBrowser::GetAddNewToolTipText() const
 
 void SContentBrowser::PopulatePathViewFiltersMenu(UToolMenu* Menu)
 {
-	if (PathViewPtr.IsValid())
+	if (LegacyContentSourceWidgets && LegacyContentSourceWidgets->PathViewPtr.IsValid())
 	{
-		PathViewPtr->PopulatePathViewFiltersMenu(Menu);
+		LegacyContentSourceWidgets->PathViewPtr->PopulatePathViewFiltersMenu(Menu);
 	}
 }
 
@@ -3413,224 +2989,821 @@ void SContentBrowser::OnNewItemRequested(const FContentBrowserItem& NewItem)
 	TArray<FString> SelectedPaths;
 	SelectedPaths.Add(FPaths::GetPath(NewItem.GetVirtualPath().ToString()));
 
-	const TArray<FString> CurrentlySelectedPath = PathViewPtr->GetSelectedPaths();
+	const TArray<FString> CurrentlySelectedPath = LegacyContentSourceWidgets->PathViewPtr->GetSelectedPaths();
+
+	// Only change the selected paths if needed. (To avoid adding an entry to navigation history when it is not needed)
+	if (SelectedPaths != CurrentlySelectedPath)
+	{
+		LegacyContentSourceWidgets->PathViewPtr->SetSelectedPaths(SelectedPaths);
+		PathSelected(SelectedPaths[0]);
+	}
+}
+
+void SContentBrowser::OnItemSelectionChanged(const FContentBrowserItem& SelectedItem, ESelectInfo::Type SelectInfo, EContentBrowserViewContext ViewContext)
+{
+	if (ViewContext == EContentBrowserViewContext::AssetView)
+	{		
+		if (bIsPrimaryBrowser)
+		{
+			SyncGlobalSelectionSet();
+		}
+
+		// Notify 'asset selection changed' delegate
+		FContentBrowserModule& ContentBrowserModule = FModuleManager::GetModuleChecked<FContentBrowserModule>(TEXT("ContentBrowser"));
+		FContentBrowserModule::FOnAssetSelectionChanged& AssetSelectionChangedDelegate = ContentBrowserModule.GetOnAssetSelectionChanged();
+
+		const TArray<FContentBrowserItem> SelectedItems = LegacyContentSourceWidgets->AssetViewPtr->GetSelectedItems();
+		LegacyContentSourceWidgets->AssetContextMenu->SetSelectedItems(SelectedItems);
+
+		{
+			TArray<FSoftObjectPath> SelectedCollectionItems;
+			for (const FContentBrowserItem& SelectedAssetItem : SelectedItems)
+			{
+				FSoftObjectPath CollectionItemId;
+				if (SelectedAssetItem.TryGetCollectionId(CollectionItemId))
+				{
+					SelectedCollectionItems.Add(CollectionItemId);
+				}
+			}
+
+			for (const TUniquePtr<FCollectionSource>& CollectionSource : CollectionSources)
+			{
+				CollectionSource->CollectionViewPtr->SetSelectedAssetPaths(SelectedCollectionItems);
+			}
+		}
+
+		if (AssetSelectionChangedDelegate.IsBound())
+		{
+			TArray<FAssetData> SelectedAssets;
+			for (const FContentBrowserItem& SelectedAssetItem : SelectedItems)
+			{
+				FAssetData ItemAssetData;
+				if (SelectedAssetItem.Legacy_TryGetAssetData(ItemAssetData))
+				{
+					SelectedAssets.Add(MoveTemp(ItemAssetData));
+				}
+			}
+
+			AssetSelectionChangedDelegate.Broadcast(SelectedAssets, bIsPrimaryBrowser);
+		}
+	}
+	else if (ViewContext == EContentBrowserViewContext::FavoriteView)
+	{
+		checkf(!SelectedItem.IsValid() || SelectedItem.IsFolder(), TEXT("File item passed to path view selection!"));
+		FavoritePathSelected(SelectedItem.IsValid() ? SelectedItem.GetVirtualPath().ToString() : FString());
+	}
+	else
+	{
+		checkf(!SelectedItem.IsValid() || SelectedItem.IsFolder(), TEXT("File item passed to path view selection!"));
+		PathSelected(SelectedItem.IsValid() ? SelectedItem.GetVirtualPath().ToString() : FString());
+	}
+}
+
+void SContentBrowser::OnItemsActivated(TArrayView<const FContentBrowserItem> ActivatedItems, EAssetTypeActivationMethod::Type ActivationMethod)
+{
+	FContentBrowserItem FirstActivatedFolder;
+
+	// Batch these by their data sources
+	TMap<UContentBrowserDataSource*, TArray<FContentBrowserItemData>> SourcesAndItems;
+	for (const FContentBrowserItem& ActivatedItem : ActivatedItems)
+	{
+		if (ActivatedItem.IsFile())
+		{
+			FContentBrowserItem::FItemDataArrayView ItemDataArray = ActivatedItem.GetInternalItems();
+			for (const FContentBrowserItemData& ItemData : ItemDataArray)
+			{
+				if (UContentBrowserDataSource* ItemDataSource = ItemData.GetOwnerDataSource())
+				{
+					TArray<FContentBrowserItemData>& ItemsForSource = SourcesAndItems.FindOrAdd(ItemDataSource);
+					ItemsForSource.Add(ItemData);
+				}
+			}
+		}
+
+		if (ActivatedItem.IsFolder() && !FirstActivatedFolder.IsValid())
+		{
+			FirstActivatedFolder = ActivatedItem;
+		}
+	}
+
+	if (SourcesAndItems.Num() == 0 && FirstActivatedFolder.IsValid())
+	{
+		// Activate the selected folder
+		FolderEntered(FirstActivatedFolder);
+		return;
+	}
+
+	// Execute the operation now
+	for (const auto& SourceAndItemsPair : SourcesAndItems)
+	{
+		if (ActivationMethod == EAssetTypeActivationMethod::Previewed)
+		{
+			SourceAndItemsPair.Key->BulkPreviewItems(SourceAndItemsPair.Value);
+		}
+		else
+		{
+			for (const FContentBrowserItemData& ItemToEdit : SourceAndItemsPair.Value)
+			{
+				FText EditErrorMsg;
+				if (!SourceAndItemsPair.Key->CanEditItem(ItemToEdit, &EditErrorMsg) && !SourceAndItemsPair.Key->CanViewItem(ItemToEdit, &EditErrorMsg))
+				{
+					AssetViewUtils::ShowErrorNotifcation(EditErrorMsg);
+				}
+			}
+			
+			if (!SourceAndItemsPair.Key->BulkEditItems(SourceAndItemsPair.Value))
+			{
+				static const FText ErrorMessage = LOCTEXT("EditItemsFailure", "Unable to edit assets");
+
+				FNotificationInfo WarningNotification(ErrorMessage);
+				WarningNotification.ExpireDuration = 5.0f;
+				WarningNotification.Hyperlink = FSimpleDelegate::CreateStatic([](){ FMessageLog("LoadErrors").Open(EMessageSeverity::Info, true); });
+				WarningNotification.HyperlinkText = LOCTEXT("LoadObjectHyperlink", "Show Message Log");
+				WarningNotification.bFireAndForget = true;
+				FSlateNotificationManager::Get().AddNotification(WarningNotification);
+			}
+		}
+	}
+}
+
+FReply SContentBrowser::ToggleLockClicked()
+{
+	bIsLocked = !bIsLocked;
+
+	return FReply::Handled();
+}
+
+FReply SContentBrowser::DockInLayoutClicked()
+{
+	FContentBrowserSingleton::Get().DockContentBrowserDrawer();
+
+	return FReply::Handled();
+}
+
+FText SContentBrowser::GetLockMenuText() const
+{
+	return IsLocked() ? LOCTEXT("ContentBrowserLockMenu_Unlock", "Unlock Content Browser") : LOCTEXT("ContentBrowserLockMenu_Lock", "Lock Content Browser");
+}
+
+FSlateIcon SContentBrowser::GetLockIcon() const
+{
+	static const FSlateIcon Unlocked = FSlateIcon(FAppStyle::GetAppStyleSetName(), "Icons.Unlock");
+	static const FSlateIcon Locked = FSlateIcon(FAppStyle::GetAppStyleSetName(), "Icons.Lock");
+	return IsLocked() ? Locked : Unlocked;
+}
+
+const FSlateBrush* SContentBrowser::GetLockIconBrush() const
+{
+	static const FName Unlock = "Icons.Unlock";
+	static const FName Lock = "Icons.Lock";
+
+	return FAppStyle::Get().GetBrush(IsLocked() ? Lock : Unlock);
+}
+
+EVisibility SContentBrowser::GetSourcesViewVisibility() const
+{
+	return bSourcesViewExpanded ? EVisibility::Visible : EVisibility::Collapsed;
+}
+
+void SContentBrowser::SetSourcesViewExpanded(bool bExpanded)
+{
+	bSourcesViewExpanded = bExpanded;
+
+	if (FContentBrowserInstanceConfig* EditorConfig = GetMutableInstanceConfig())
+	{
+		EditorConfig->bSourcesExpanded = bSourcesViewExpanded;
+		UContentBrowserConfig::Get()->SaveEditorConfig();
+	}
+
+
+	// Notify 'Sources View Expanded' delegate
+	FContentBrowserModule& ContentBrowserModule = FModuleManager::GetModuleChecked<FContentBrowserModule>( TEXT("ContentBrowser") );
+	FContentBrowserModule::FOnSourcesViewChanged& SourcesViewChangedDelegate = ContentBrowserModule.GetOnSourcesViewChanged();
+	if (SourcesViewChangedDelegate.IsBound())
+	{
+		SourcesViewChangedDelegate.Broadcast(bSourcesViewExpanded);
+	}
+}
+
+FReply SContentBrowser::SourcesViewExpandClicked()
+{
+	SetSourcesViewExpanded(!bSourcesViewExpanded);
+	return FReply::Handled();
+}
+
+void SContentBrowser::OnContentBrowserSettingsChanged(FName PropertyName)
+{
+	if (PropertyName.IsNone())
+	{
+		// Ensure the path is set to the correct view mode
+		UpdatePath();
+	}
+}
+
+void SContentBrowser::OnConsoleVariableChanged()
+{
+	UpdatePrivateContentFeatureEnabled(true /* bUpdateFilterIfChanged */);
+}
+
+void SContentBrowser::UpdatePrivateContentFeatureEnabled(bool bUpdateFilterIfChanged)
+{
+}
+
+void SContentBrowser::OnLegacyContentSourceEnabled()
+{
+	// Re-bind our commands so they work properly
+	BindCommands();
+
+	// Create the content browser's default widgets and set them as the child widget contents
+	LegacyContentSource->SetContent(CreateLegacyAssetViewWidgets());
+
+	// Load our settings
+	LoadSettings(InstanceName);
+
+	// Sanity sync to make sure the global selection set is synced
+	SyncGlobalSelectionSet();
+}
+
+void SContentBrowser::OnLegacyContentSourceDisabled()
+{
+	// Save our settings before destroying the widgets
+	SaveSettings();
+
+	// Unbind commands
+	UnbindCommands();
+
+	// Set the child widget contents to null and destroy all asset view widgets and they will be re-bound when the legacy content source is enabled
+	LegacyContentSource->SetContent(SNullWidget::NullWidget);
+	LegacyContentSourceWidgets.Reset();
+	
+	// Unbind all delegates since we don't need them anymore, and 
+	FCollectionManagerModule& CollectionManagerModule = FCollectionManagerModule::GetModule();
+
+	TArray<TSharedPtr<ICollectionContainer>> CollectionContainers;
+	CollectionManagerModule.Get().GetCollectionContainers(CollectionContainers);
+
+	int32 InsertIndex = 0;
+	for (const TSharedPtr<ICollectionContainer>& CollectionContainer : CollectionContainers)
+	{
+		CollectionContainer->OnIsHiddenChanged().RemoveAll(this);
+		CollectionContainer->OnCollectionRenamed().RemoveAll(this);
+		CollectionContainer->OnCollectionDestroyed().RemoveAll(this);
+		CollectionContainer->OnCollectionUpdated().RemoveAll(this);
+	}
+
+	CollectionManagerModule.Get().OnCollectionContainerCreated().RemoveAll(this);
+	CollectionManagerModule.Get().OnCollectionContainerDestroyed().RemoveAll(this);
+	
+	CollectionSources.Empty();
+
+	UContentBrowserDataSubsystem* ContentBrowserData = IContentBrowserDataModule::Get().GetSubsystem();
+	ContentBrowserData->OnItemDataUpdated().RemoveAll(this);
+}
+
+TSharedRef<SWidget> SContentBrowser::CreateLegacyAssetViewWidgets()
+{
+	LegacyContentSourceWidgets = MakeShared<FLegacyContentSourceWidgets>();
+	
+	using namespace UE::Editor::ContentBrowser::Private;
+
+	// The final widget that conatins all child widgets
+	TSharedRef<SWidget> FinalWidget = SNullWidget::NullWidget;
+
+	LegacyContentSourceWidgets->PathContextMenu = MakeShareable(new FPathContextMenu( AsShared() ));
+	LegacyContentSourceWidgets->PathContextMenu->SetOnRenameFolderRequested(FPathContextMenu::FOnRenameFolderRequested::CreateSP(this, &SContentBrowser::OnRenameRequested));
+	LegacyContentSourceWidgets->PathContextMenu->SetOnFolderDeleted(FPathContextMenu::FOnFolderDeleted::CreateSP(this, &SContentBrowser::OnOpenedFolderDeleted));
+	LegacyContentSourceWidgets->PathContextMenu->SetOnFolderFavoriteToggled(FPathContextMenu::FOnFolderFavoriteToggled::CreateSP(this, &SContentBrowser::ToggleFolderFavorite));
+	LegacyContentSourceWidgets->PathContextMenu->SetOnPrivateContentEditToggled(FPathContextMenu::FOnPrivateContentEditToggled::CreateSP(this, &SContentBrowser::TogglePrivateContentEdit));
+
+	// Currently this controls the asset count
+	const bool bShowBottomToolbar = InitConfig.bShowBottomToolbar;
+
+	LegacyContentSourceWidgets->AssetViewPtr = SNew(SAssetView)
+			.ThumbnailLabel(InitConfig.ThumbnailLabel)
+			//.ThumbnailScale(Config != nullptr ? Config->ThumbnailScale : 0.18f)
+			.InitialViewType(InitConfig.InitialAssetViewType)
+			.OnNewItemRequested(this, &SContentBrowser::OnNewItemRequested)
+			.OnItemSelectionChanged(this, &SContentBrowser::OnItemSelectionChanged, EContentBrowserViewContext::AssetView)
+			.OnItemsActivated(this, &SContentBrowser::OnItemsActivated)
+			.OnGetItemContextMenu(this, &SContentBrowser::GetItemContextMenu, EContentBrowserViewContext::AssetView)
+			.OnItemRenameCommitted(this, &SContentBrowser::OnItemRenameCommitted)
+			.FrontendFilters(FrontendFilters)
+			.TextFilter(TextFilter)
+			.ShowRedirectors(this, &SContentBrowser::ShouldShowRedirectors)
+			.HighlightedText(this, &SContentBrowser::GetHighlightedText)
+			.ShowBottomToolbar(bShowBottomToolbar)
+			.ShowViewOptions(false) // We control this for the main content browser
+			.AllowThumbnailEditMode(true)
+			.AllowThumbnailHintLabel(false)
+			.CanShowFolders(InitConfig.bCanShowFolders)
+			.CanShowClasses(InitConfig.bCanShowClasses)
+			.CanShowRealTimeThumbnails(InitConfig.bCanShowRealTimeThumbnails)
+			.CanShowDevelopersFolder(InitConfig.bCanShowDevelopersFolder)
+			.CanShowFavorites(true)
+			.AddMetaData<FTagMetaData>(FTagMetaData(TEXT("ContentBrowserAssets")))
+			.OwningContentBrowser(SharedThis(this))
+			.OnSearchOptionsChanged(this, &SContentBrowser::HandleAssetViewSearchOptionsChanged)
+			.bShowPathViewFilters(true)
+			.FillEmptySpaceInTileView(true)
+			.ShowDisallowedAssetClassAsUnsupportedItems(true)
+			.AllowCustomView(true);
+	
+	TSharedRef<SWidget> ViewOptions = SNullWidget::NullWidget;
+
+	// Note, for backwards compatibility ShowBottomToolbar controls the visibility of view options so we respect that here
+	if (bShowBottomToolbar)
+	{
+		if (UE::Editor::ContentBrowser::IsNewStyleEnabled())
+		{
+			ViewOptions =
+				SNew(SActionButton)
+				.ActionButtonType(EActionButtonType::Simple)
+				.OnGetMenuContent(LegacyContentSourceWidgets->AssetViewPtr.ToSharedRef(), &SAssetView::GetViewButtonContent)
+				.Icon(FAppStyle::Get().GetBrush("Icons.Settings"));
+		}
+		else
+		{
+			ViewOptions =
+				SNew(SComboButton)
+				.ComboButtonStyle(&FAppStyle::Get().GetWidgetStyle<FComboButtonStyle>("SimpleComboButton"))
+				.OnGetMenuContent(LegacyContentSourceWidgets->AssetViewPtr.ToSharedRef(), &SAssetView::GetViewButtonContent)
+				.HasDownArrow(false)
+				.ButtonContent()
+				[
+					SNew(SHorizontalBox)
+					+ SHorizontalBox::Slot()
+					.AutoWidth()
+					.VAlign(VAlign_Center)
+					.Padding(4.0, 0.0f)
+					[
+						SNew(SImage)
+						.ColorAndOpacity(FSlateColor::UseForeground())
+						.Image(FAppStyle::Get().GetBrush("Icons.Settings"))
+					]
+
+					+ SHorizontalBox::Slot()
+					.VAlign(VAlign_Center)
+					.Padding(4.0, 0.0f)
+					[
+						SNew(STextBlock)
+						.Text(LOCTEXT("Settings", "Settings"))
+						.ColorAndOpacity(FSlateColor::UseForeground())
+					]
+				];
+		}
+	}
+	
+	if (UE::Editor::ContentBrowser::IsNewStyleEnabled())
+	{
+		constexpr float ToolBarVerticalPadding = 4.0f; // ToolBar total height should be 36.0f
+		constexpr float PanelInsetPadding = 2.0f;
+		constexpr float SourceTreeSectionPadding = 2.0f;
+
+		const TSharedRef<SWidget> AssetView = CreateAssetView(&InitConfig);
+
+		CreateFavoritesView(&InitConfig);
+		CreatePathView(&InitConfig);
+
+		{
+			SAssignNew(LegacyContentSourceWidgets->SourceTreePtr, SContentBrowserSourceTree)
+
+			+ SContentBrowserSourceTree::Slot()
+			.AreaWidget(FavoritesArea)
+			.Size(0.2f)
+			.Visibility(this, &SContentBrowser::GetFavoriteFolderVisibility)
+
+			+ SContentBrowserSourceTree::Slot()
+			.AreaWidget(PathArea)
+			.Size(0.8f);
+		}
+
+		LegacyContentSourceWidgets->NavigationToolBarWidget = CreateNavigationToolBar(&InitConfig);
+
+		TSharedPtr<SBox> AssetViewNavigationToolBarContainer = nullptr;
+		TSharedPtr<SBox> SourceTreeAndAssetViewNavigationToolBarContainer = nullptr;
+
+		FinalWidget = 
+		SNew(SVerticalBox)
+
+		+ SVerticalBox::Slot()
+		.AutoHeight()
+		[
+			SNew(SSeparator)
+			.Thickness(2.0f)
+		]
+
+		// Source / Tree + Assets
+		+ SVerticalBox::Slot()
+		.FillHeight(1.0f)
+		.Padding(0.0f)
+		[
+			SNew(SHorizontalBox)
+
+			// Tree + Assets + Navigation Bar
+			+ SHorizontalBox::Slot()
+			.FillWidth(1.0f)
+			.Padding(FMargin(0.0f))
+			[
+				SNew(SVerticalBox)
+
+				// Assets/tree
+				+ SVerticalBox::Slot()
+				.FillHeight(1.0f)
+				.Padding(FMargin(0.0f))
+				[
+					// The tree/assets splitter
+					SAssignNew(LegacyContentSourceWidgets->PathAssetSplitterPtr, SSplitter)
+					.PhysicalSplitterHandleSize(PanelInsetPadding)
+
+					// Sources View
+					+ SSplitter::Slot()
+					.Resizable(true)
+					.SizeRule(SSplitter::SizeToContent)
+					.OnSlotResized(this, &SContentBrowser::OnPathViewBoxColumnResized)
+					[
+						SNew(SBox)
+						.Padding(0.0f)
+						.Visibility(this, &SContentBrowser::GetSourcesViewVisibility)
+						.WidthOverride(this, &SContentBrowser::GetPathViewBoxWidthOverride)
+						[
+							SNew(SBorder)
+							.Padding(SourceTreeSectionPadding)
+							.BorderImage(FAppStyle::GetBrush("Brushes.Panel"))
+							[
+								// Panel background, seen between items
+								SNew(SBorder)
+								.BorderImage(FAppStyle::GetBrush("WhiteBrush"))
+								.BorderBackgroundColor(FStyleColors::Panel)
+								.HAlign(HAlign_Fill)
+								.VAlign(VAlign_Fill)
+								.Padding(0.0f)
+								[
+									LegacyContentSourceWidgets->SourceTreePtr.ToSharedRef()
+								]
+							]
+						]
+					]
+
+					// Asset View
+					+ SSplitter::Slot()
+					.Value(0.75f)
+					[
+						SNew(SBox)
+						.Padding(0.0f)
+						[
+							SNew(SVerticalBox)
+
+							+ SVerticalBox::Slot()
+							.AutoHeight()
+							[
+								SNew(SBorder)
+								.Padding(FMargin(3, ToolBarVerticalPadding))
+								.BorderImage(bIsDrawer ? FStyleDefaults::GetNoBrush() : FAppStyle::Get().GetBrush("Brushes.Panel"))
+								[
+									SNew(SHorizontalBox)
+
+									+ SHorizontalBox::Slot()
+									.FillWidth(1.0f)
+									.HAlign(HAlign_Fill)
+									.VAlign(VAlign_Center)
+									.Padding(5, 0, 0, 0)
+									[
+										CreateToolBar(&InitConfig)
+									]
+
+									+ SHorizontalBox::Slot()
+									.AutoWidth()
+									.HAlign(HAlign_Right)
+									.VAlign(VAlign_Center)
+									[
+										CreateDrawerDockButton(&InitConfig)
+									]
+
+									+ SHorizontalBox::Slot()
+									.AutoWidth()
+									.HAlign(HAlign_Right)
+									.VAlign(VAlign_Center)
+									.Padding(5.0f, 0.0f, 0.0f, 0.0f)
+									[
+										ViewOptions
+									]
+								]
+							]
+
+							+ SVerticalBox::Slot()
+							[
+								AssetView
+							]
+
+							// Navigation ToolBar
+							+ SVerticalBox::Slot()
+							.AutoHeight()
+							.HAlign(HAlign_Fill)
+							.Padding(FMargin(0, PanelInsetPadding, 0, 0))
+							[
+								SNew(SBorder)
+								.Padding(FMargin(3))
+								.BorderImage(bIsDrawer ? FStyleDefaults::GetNoBrush() : FAppStyle::Get().GetBrush("Brushes.Panel"))
+								[
+									SAssignNew(AssetViewNavigationToolBarContainer, SBox)
+									.Padding(0)
+									[
+										LegacyContentSourceWidgets->NavigationToolBarWidget.ToSharedRef()
+									]
+								]
+							]
+						]
+					]
+				]
+
+				+ SVerticalBox::Slot()
+				.AutoHeight()
+				.Padding(0)
+				[
+					SAssignNew(SourceTreeAndAssetViewNavigationToolBarContainer, SBox)
+					.Padding(0)
+				]
+			]
+		];
+
+		LegacyContentSourceWidgets->NavigationToolBarParentSwitcher = MakeShared<UE::Editor::ContentBrowser::Private::TWidgetSizeParentSwitcher<Orient_Horizontal>>(
+			LegacyContentSourceWidgets->NavigationToolBarWidget,
+			TArray<TSharedPtr<SBox>>{ SourceTreeAndAssetViewNavigationToolBarContainer, AssetViewNavigationToolBarContainer },
+			TArray<FInt16Range>{ FInt16Range(0, 550), FInt16Range(550, 10000) },
+			[this]()
+			{
+				const float Size = this->GetCachedGeometry().GetLocalSize().X;
+				return Size;
+			}
+		);
+
+		LegacyContentSourceWidgets->SourceTreeSplitterNumFixedSlots = LegacyContentSourceWidgets->SourceTreePtr->GetSplitter()->GetChildren()->NumSlot();
+	}
+	else
+	{
+		FinalWidget = 
+		SNew(SVerticalBox)
+		+SVerticalBox::Slot()
+		.AutoHeight()
+		.Padding( 0, 0, 0, 0 )
+		[
+			SNew( SBorder )
+			.Padding( FMargin( 3 ) )
+			.BorderImage(bIsDrawer ? FStyleDefaults::GetNoBrush() : FAppStyle::Get().GetBrush("Brushes.Panel"))
+			[
+				SNew(SHorizontalBox)
+				+ SHorizontalBox::Slot()
+				.AutoWidth()
+				.HAlign(HAlign_Left)
+				.Padding(5, 0, 0, 0)
+				[
+					CreateToolBar(&InitConfig)
+				]
+				// History Back Button
+				+SHorizontalBox::Slot()
+				.Padding(10, 0, 0, 0)
+				.AutoWidth()
+				[
+					SNew(SButton)
+					.VAlign(EVerticalAlignment::VAlign_Center)
+					.ButtonStyle(FAppStyle::Get(), "SimpleButton")
+					.ToolTipText( this, &SContentBrowser::GetHistoryBackTooltip )
+					.ContentPadding( FMargin(1, 0) )
+					.OnClicked(this, &SContentBrowser::BackClicked)
+					.IsEnabled(this, &SContentBrowser::IsBackEnabled)
+					.AddMetaData<FTagMetaData>(FTagMetaData(TEXT("ContentBrowserHistoryBack")))
+					[
+						SNew(SImage)
+						.Image(FAppStyle::Get().GetBrush("Icons.CircleArrowLeft"))
+						.ColorAndOpacity(FSlateColor::UseForeground())
+					]
+				]
+				// History Forward Button
+				+ SHorizontalBox::Slot()
+				.Padding(2, 0, 0, 0)
+				.AutoWidth()
+				[
+					SNew(SButton)
+					.VAlign(EVerticalAlignment::VAlign_Center)
+					.ButtonStyle(FAppStyle::Get(), "SimpleButton")
+					.ToolTipText( this, &SContentBrowser::GetHistoryForwardTooltip )
+					.ContentPadding( FMargin(1, 0) )
+					.OnClicked(this, &SContentBrowser::ForwardClicked)
+					.IsEnabled(this, &SContentBrowser::IsForwardEnabled)
+					.AddMetaData<FTagMetaData>(FTagMetaData(TEXT("ContentBrowserHistoryForward")))
+					[
+						SNew(SImage)
+						.Image(FAppStyle::Get().GetBrush("Icons.CircleArrowRight"))
+						.ColorAndOpacity(FSlateColor::UseForeground())
+					]
+				]
+
+				// Path
+				+ SHorizontalBox::Slot()
+				.VAlign(VAlign_Center)
+				.HAlign(HAlign_Fill)
+				.FillWidth(1.0f)
+				.Padding(2, 0, 0, 0)
+				[
+					SAssignNew(LegacyContentSourceWidgets->NavigationBar, SNavigationBar)
+					.OnPathClicked(this, &SContentBrowser::OnPathClicked)
+					.GetPathMenuContent(this, &SContentBrowser::OnGetCrumbDelimiterContent)
+					.GetComboOptions(this, &SContentBrowser::GetRecentPaths)
+					.OnNavigateToPath(this, &SContentBrowser::OnNavigateToPath)
+					.OnCompletePrefix(this, &SContentBrowser::OnCompletePathPrefix)
+					.OnCanEditPathAsText(this, &SContentBrowser::OnCanEditPathAsText)
+					.AddMetaData<FTagMetaData>(FTagMetaData(TEXT("ContentBrowserPath")))
+				]
+
+				+ SHorizontalBox::Slot()
+				.AutoWidth()
+				.HAlign(HAlign_Right)
+				.VAlign(VAlign_Center)
+				[
+					CreateLockButton(&InitConfig)
+				]
+				+ SHorizontalBox::Slot()
+				.AutoWidth()
+				.HAlign(HAlign_Right)
+				.VAlign(VAlign_Center)
+				[
+					CreateDrawerDockButton(&InitConfig)
+				]
+				+ SHorizontalBox::Slot()
+				.AutoWidth()
+				.Padding(5.0f, 0.0f, 0.0f, 0.0f)
+				.HAlign(HAlign_Right)
+				.VAlign(VAlign_Center)
+				[
+					ViewOptions
+				]
+			]
+		]
+
+		+SVerticalBox::Slot()
+		.AutoHeight()
+		[
+			SNew(SSeparator)
+			.Thickness(2.0f)
+		]
 
-	// Only change the selected paths if needed. (To avoid adding an entry to navigation history when it is not needed)
-	if (SelectedPaths != CurrentlySelectedPath)
-	{
-		PathViewPtr->SetSelectedPaths(SelectedPaths);
-		PathSelected(SelectedPaths[0]);
-	}
-}
+		// Assets/tree
+		+ SVerticalBox::Slot()
+		.FillHeight(1.0f)
+		.Padding(0.0f)
+		[
+			// The tree/assets splitter
+			SAssignNew(LegacyContentSourceWidgets->PathAssetSplitterPtr, SSplitter)
+			.PhysicalSplitterHandleSize(2.0f)
+
+			// Sources View
+			+ SSplitter::Slot()
+			.Resizable(true)
+			.SizeRule(SSplitter::SizeToContent)
+			.OnSlotResized(this, &SContentBrowser::OnPathViewBoxColumnResized)
+			[
+				SNew(SBox)
+				.Padding(FMargin(4.f))
+				.Visibility(this, &SContentBrowser::GetSourcesViewVisibility)
+				.WidthOverride(this, &SContentBrowser::GetPathViewBoxWidthOverride)
+				[
+					SNew(SBorder)
+					.Padding(FMargin(0))
+					.BorderImage(FAppStyle::GetBrush("Brushes.Recessed"))
 
-void SContentBrowser::OnItemSelectionChanged(const FContentBrowserItem& SelectedItem, ESelectInfo::Type SelectInfo, EContentBrowserViewContext ViewContext)
-{
-	if (ViewContext == EContentBrowserViewContext::AssetView)
-	{		
-		if (bIsPrimaryBrowser)
-		{
-			SyncGlobalSelectionSet();
-		}
+					[
+						SAssignNew(LegacyContentSourceWidgets->PathFavoriteSplitterPtr, SSplitter)
+						.Clipping(EWidgetClipping::ClipToBounds)
+						.PhysicalSplitterHandleSize(2.0f)
+						.HitDetectionSplitterHandleSize(8.0f)
+						.Orientation(EOrientation::Orient_Vertical)
+						.MinimumSlotHeight(26.0f)
+						.Visibility( this, &SContentBrowser::GetSourcesViewVisibility )
+						+SSplitter::Slot()
+						.SizeRule(TAttribute<SSplitter::ESizeRule>(this, &SContentBrowser::GetFavoritesAreaSizeRule))
+						.MinSize(TAttribute<float>(this, &SContentBrowser::GetFavoritesAreaMinSize))
+						.Value(0.2f)
+						[
+							SNew(SBorder)
+							.BorderImage(FAppStyle::Get().GetBrush("Brushes.Header"))
+							.Padding(0.0f, 2.0f, 0.0f, 0.0f)
+							[
+								CreateFavoritesView(&InitConfig)
+							]
+						]
 
-		// Notify 'asset selection changed' delegate
-		FContentBrowserModule& ContentBrowserModule = FModuleManager::GetModuleChecked<FContentBrowserModule>(TEXT("ContentBrowser"));
-		FContentBrowserModule::FOnAssetSelectionChanged& AssetSelectionChangedDelegate = ContentBrowserModule.GetOnAssetSelectionChanged();
+						+SSplitter::Slot()
+						.SizeRule(TAttribute<SSplitter::ESizeRule>(this, &SContentBrowser::GetPathAreaSizeRule))
+						.MinSize(29.0f)
+						.Value(0.8f)
+						[
+							SNew(SBorder)
+							.BorderImage(FAppStyle::Get().GetBrush("Brushes.Header"))
+							.Padding(0.0f, 2.0f, 0.0f, 0.0f)
+							[
+								CreatePathView(&InitConfig)
+							]
+						]
+					]
+				]
+			]
 
-		const TArray<FContentBrowserItem> SelectedItems = AssetViewPtr->GetSelectedItems();
-		AssetContextMenu->SetSelectedItems(SelectedItems);
+			// Asset View
+			+ SSplitter::Slot()
+			.Value(0.75f)
+			[
+				CreateAssetView(&InitConfig)
+			]
+		];
 
-		{
-			TArray<FSoftObjectPath> SelectedCollectionItems;
-			for (const FContentBrowserItem& SelectedAssetItem : SelectedItems)
-			{
-				FSoftObjectPath CollectionItemId;
-				if (SelectedAssetItem.TryGetCollectionId(CollectionItemId))
-				{
-					SelectedCollectionItems.Add(CollectionItemId);
-				}
-			}
+		LegacyContentSourceWidgets->SourceTreeSplitterNumFixedSlots = LegacyContentSourceWidgets->PathFavoriteSplitterPtr->GetChildren()->NumSlot();
+	}
 
-			for (const TUniquePtr<FCollectionSource>& CollectionSource : CollectionSources)
-			{
-				CollectionSource->CollectionViewPtr->SetSelectedAssetPaths(SelectedCollectionItems);
-			}
-		}
+LegacyContentSourceWidgets->AssetContextMenu = MakeShared<FAssetContextMenu>(LegacyContentSourceWidgets->AssetViewPtr);
+	LegacyContentSourceWidgets->AssetContextMenu->BindCommands(Commands);
+	LegacyContentSourceWidgets->AssetContextMenu->SetOnShowInPathsViewRequested( FAssetContextMenu::FOnShowInPathsViewRequested::CreateSP(this, &SContentBrowser::OnShowInPathsViewRequested) );
+	LegacyContentSourceWidgets->AssetContextMenu->SetOnRenameRequested( FAssetContextMenu::FOnRenameRequested::CreateSP(this, &SContentBrowser::OnRenameRequested) );
+	LegacyContentSourceWidgets->AssetContextMenu->SetOnDuplicateRequested( FAssetContextMenu::FOnDuplicateRequested::CreateSP(this, &SContentBrowser::OnDuplicateRequested) );
+	LegacyContentSourceWidgets->AssetContextMenu->SetOnAssetViewRefreshRequested( FAssetContextMenu::FOnAssetViewRefreshRequested::CreateSP( this, &SContentBrowser::OnAssetViewRefreshRequested) );
 
-		if (AssetSelectionChangedDelegate.IsBound())
-		{
-			TArray<FAssetData> SelectedAssets;
-			for (const FContentBrowserItem& SelectedAssetItem : SelectedItems)
-			{
-				FAssetData ItemAssetData;
-				if (SelectedAssetItem.Legacy_TryGetAssetData(ItemAssetData))
-				{
-					SelectedAssets.Add(MoveTemp(ItemAssetData));
-				}
-			}
+	TOptional<FCollectionRef> SelectedCollection;
 
-			AssetSelectionChangedDelegate.Broadcast(SelectedAssets, bIsPrimaryBrowser);
-		}
-	}
-	else if (ViewContext == EContentBrowserViewContext::FavoriteView)
+	SelectedCollection = InitConfig.SelectedCollection;
+
+	PRAGMA_DISABLE_DEPRECATION_WARNINGS
+	// Check for someone using the deprecated SelectedCollectionName instead of SelectedCollection.
+	if (!SelectedCollection.GetValue().IsValid() && InitConfig.SelectedCollectionName.Name != NAME_None)
 	{
-		checkf(!SelectedItem.IsValid() || SelectedItem.IsFolder(), TEXT("File item passed to path view selection!"));
-		FavoritePathSelected(SelectedItem.IsValid() ? SelectedItem.GetVirtualPath().ToString() : FString());
+		SelectedCollection = FCollectionRef(FCollectionManagerModule::GetModule().Get().GetProjectCollectionContainer(), InitConfig.SelectedCollectionName);
 	}
-	else
+	PRAGMA_ENABLE_DEPRECATION_WARNINGS
+
+	if (SelectedCollection.IsSet() && SelectedCollection.GetValue().IsValid())
 	{
-		checkf(!SelectedItem.IsValid() || SelectedItem.IsFolder(), TEXT("File item passed to path view selection!"));
-		PathSelected(SelectedItem.IsValid() ? SelectedItem.GetVirtualPath().ToString() : FString());
+		// Select the specified collection by default
+		FAssetViewContentSources DefaultContentSources(SelectedCollection.GetValue());
+		LegacyContentSourceWidgets->AssetViewPtr->SetContentSources(DefaultContentSources);
 	}
-}
-
-void SContentBrowser::OnItemsActivated(TArrayView<const FContentBrowserItem> ActivatedItems, EAssetTypeActivationMethod::Type ActivationMethod)
-{
-	FContentBrowserItem FirstActivatedFolder;
-
-	// Batch these by their data sources
-	TMap<UContentBrowserDataSource*, TArray<FContentBrowserItemData>> SourcesAndItems;
-	for (const FContentBrowserItem& ActivatedItem : ActivatedItems)
+	else
 	{
-		if (ActivatedItem.IsFile())
-		{
-			FContentBrowserItem::FItemDataArrayView ItemDataArray = ActivatedItem.GetInternalItems();
-			for (const FContentBrowserItemData& ItemData : ItemDataArray)
-			{
-				if (UContentBrowserDataSource* ItemDataSource = ItemData.GetOwnerDataSource())
-				{
-					TArray<FContentBrowserItemData>& ItemsForSource = SourcesAndItems.FindOrAdd(ItemDataSource);
-					ItemsForSource.Add(ItemData);
-				}
-			}
-		}
+		// Select /Game by default
+		const FName DefaultInvariantPath(TEXT("/Game"));
+		FName DefaultVirtualPath;
+		IContentBrowserDataModule::Get().GetSubsystem()->ConvertInternalPathToVirtual(DefaultInvariantPath, DefaultVirtualPath);
 
-		if (ActivatedItem.IsFolder() && !FirstActivatedFolder.IsValid())
-		{
-			FirstActivatedFolder = ActivatedItem;
-		}
+		FAssetViewContentSources DefaultContentSources(DefaultVirtualPath);
+		LegacyContentSourceWidgets->AssetViewPtr->SetContentSources(DefaultContentSources);
 	}
 
-	if (SourcesAndItems.Num() == 0 && FirstActivatedFolder.IsValid())
+	if (bHasInitConfig)
 	{
-		// Activate the selected folder
-		FolderEntered(FirstActivatedFolder);
-		return;
+		// Make sure the sources view is initially visible if we were asked to show it
+		SetSourcesViewExpanded(InitConfig.bExpandSourcesView && InitConfig.bUseSourcesView);
 	}
-
-	// Execute the operation now
-	for (const auto& SourceAndItemsPair : SourcesAndItems)
+	//else
 	{
-		if (ActivationMethod == EAssetTypeActivationMethod::Previewed)
-		{
-			SourceAndItemsPair.Key->BulkPreviewItems(SourceAndItemsPair.Value);
-		}
-		else
+		// in case we do not have a config, see what the global default settings are for the Sources Panel
+		bool bSourcesExpanded = true;
+		if (const FContentBrowserInstanceConfig* EditorConfig = GetConstInstanceConfig())
 		{
-			for (const FContentBrowserItemData& ItemToEdit : SourceAndItemsPair.Value)
-			{
-				FText EditErrorMsg;
-				if (!SourceAndItemsPair.Key->CanEditItem(ItemToEdit, &EditErrorMsg) && !SourceAndItemsPair.Key->CanViewItem(ItemToEdit, &EditErrorMsg))
-				{
-					AssetViewUtils::ShowErrorNotifcation(EditErrorMsg);
-				}
-			}
-			
-			if (!SourceAndItemsPair.Key->BulkEditItems(SourceAndItemsPair.Value))
-			{
-				static const FText ErrorMessage = LOCTEXT("EditItemsFailure", "Unable to edit assets");
-
-				FNotificationInfo WarningNotification(ErrorMessage);
-				WarningNotification.ExpireDuration = 5.0f;
-				WarningNotification.Hyperlink = FSimpleDelegate::CreateStatic([](){ FMessageLog("LoadErrors").Open(EMessageSeverity::Info, true); });
-				WarningNotification.HyperlinkText = LOCTEXT("LoadObjectHyperlink", "Show Message Log");
-				WarningNotification.bFireAndForget = true;
-				FSlateNotificationManager::Get().AddNotification(WarningNotification);
-			}
+			bSourcesExpanded = EditorConfig->bSourcesExpanded;
 		}
+	
+		SetSourcesViewExpanded(bSourcesExpanded);
 	}
-}
-
-FReply SContentBrowser::ToggleLockClicked()
-{
-	bIsLocked = !bIsLocked;
-
-	return FReply::Handled();
-}
-
-FReply SContentBrowser::DockInLayoutClicked()
-{
-	FContentBrowserSingleton::Get().DockContentBrowserDrawer();
-
-	return FReply::Handled();
-}
-
-FText SContentBrowser::GetLockMenuText() const
-{
-	return IsLocked() ? LOCTEXT("ContentBrowserLockMenu_Unlock", "Unlock Content Browser") : LOCTEXT("ContentBrowserLockMenu_Lock", "Lock Content Browser");
-}
-
-FSlateIcon SContentBrowser::GetLockIcon() const
-{
-	static const FSlateIcon Unlocked = FSlateIcon(FAppStyle::GetAppStyleSetName(), "Icons.Unlock");
-	static const FSlateIcon Locked = FSlateIcon(FAppStyle::GetAppStyleSetName(), "Icons.Lock");
-	return IsLocked() ? Locked : Unlocked;
-}
-
-const FSlateBrush* SContentBrowser::GetLockIconBrush() const
-{
-	static const FName Unlock = "Icons.Unlock";
-	static const FName Lock = "Icons.Lock";
-
-	return FAppStyle::Get().GetBrush(IsLocked() ? Lock : Unlock);
-}
-
-EVisibility SContentBrowser::GetSourcesViewVisibility() const
-{
-	return bSourcesViewExpanded ? EVisibility::Visible : EVisibility::Collapsed;
-}
+	
+	// Bindings to manage history when items are deleted
+	FCollectionManagerModule& CollectionManagerModule = FCollectionManagerModule::GetModule();
 
-void SContentBrowser::SetSourcesViewExpanded(bool bExpanded)
-{
-	bSourcesViewExpanded = bExpanded;
+	TArray<TSharedPtr<ICollectionContainer>> CollectionContainers;
+	CollectionManagerModule.Get().GetCollectionContainers(CollectionContainers);
 
-	if (FContentBrowserInstanceConfig* EditorConfig = GetMutableInstanceConfig())
+	int32 InsertIndex = 0;
+	for (const TSharedPtr<ICollectionContainer>& CollectionContainer : CollectionContainers)
 	{
-		EditorConfig->bSourcesExpanded = bSourcesViewExpanded;
-		UContentBrowserConfig::Get()->SaveEditorConfig();
-	}
-
+		if (!CollectionContainer->IsHidden())
+		{
+			AddSlotForCollectionContainer(InsertIndex++, CollectionContainer.ToSharedRef());
+		}
 
-	// Notify 'Sources View Expanded' delegate
-	FContentBrowserModule& ContentBrowserModule = FModuleManager::GetModuleChecked<FContentBrowserModule>( TEXT("ContentBrowser") );
-	FContentBrowserModule::FOnSourcesViewChanged& SourcesViewChangedDelegate = ContentBrowserModule.GetOnSourcesViewChanged();
-	if (SourcesViewChangedDelegate.IsBound())
-	{
-		SourcesViewChangedDelegate.Broadcast(bSourcesViewExpanded);
+		CollectionContainer->OnIsHiddenChanged().AddSP(this, &SContentBrowser::HandleIsHiddenChanged);
+		CollectionContainer->OnCollectionRenamed().AddSP(this, &SContentBrowser::HandleCollectionRenamed);
+		CollectionContainer->OnCollectionDestroyed().AddSP(this, &SContentBrowser::HandleCollectionRemoved);
+		CollectionContainer->OnCollectionUpdated().AddSP(this, &SContentBrowser::HandleCollectionUpdated);
 	}
-}
 
-FReply SContentBrowser::SourcesViewExpandClicked()
-{
-	SetSourcesViewExpanded(!bSourcesViewExpanded);
-	return FReply::Handled();
-}
+	CollectionManagerModule.Get().OnCollectionContainerCreated().AddSP(this, &SContentBrowser::HandleCollectionContainerAdded);
+	CollectionManagerModule.Get().OnCollectionContainerDestroyed().AddSP(this, &SContentBrowser::HandleCollectionContainerRemoved);
 
-void SContentBrowser::OnContentBrowserSettingsChanged(FName PropertyName)
-{
-	if (PropertyName.IsNone())
-	{
-		// Ensure the path is set to the correct view mode
-		UpdatePath();
-	}
-}
+	UContentBrowserDataSubsystem* ContentBrowserData = IContentBrowserDataModule::Get().GetSubsystem();
+	ContentBrowserData->OnItemDataUpdated().AddSP(this, &SContentBrowser::HandleItemDataUpdated);
 
-void SContentBrowser::OnConsoleVariableChanged()
-{
-	UpdatePrivateContentFeatureEnabled(true /* bUpdateFilterIfChanged */);
-}
+	LegacyContentSourceWidgets->FavoritePathViewPtr->SetTreeTitle(LOCTEXT("Favorites", "Favorites"));
 
-void SContentBrowser::UpdatePrivateContentFeatureEnabled(bool bUpdateFilterIfChanged)
-{
+	// Initialize the search options
+	HandleAssetViewSearchOptionsChanged();
+	
+	return FinalWidget;
 }
 
 FReply SContentBrowser::BackClicked()
@@ -3652,39 +3825,39 @@ bool SContentBrowser::HandleRenameCommandCanExecute() const
 	// The order of these conditions are carefully crafted to match the logic of the context menu summoning, as this callback 
 	// is shared between the path and asset views, and is given zero context as to which one is making the request
 	// Change this logic at your peril, lest the the dominoes fall like a house of cards (checkmate)
-	if (PathViewPtr->HasFocusedDescendants())
+	if (LegacyContentSourceWidgets->PathViewPtr->HasFocusedDescendants())
 	{
 		// Prefer the path view if it has focus, which may be the case when using the keyboard to invoke the action,
 		// but will be false when using the context menu (which isn't an issue, as the path view clears the asset view 
 		// selection when invoking its context menu to avoid the selection ambiguity present when using the keyboard)
-		if (PathViewPtr->GetSelectedFolderItems().Num() > 0)
+		if (LegacyContentSourceWidgets->PathViewPtr->GetSelectedFolderItems().Num() > 0)
 		{
-			return PathContextMenu->CanExecuteRename();
+			return LegacyContentSourceWidgets->PathContextMenu->CanExecuteRename();
 		}
 	}
-	else if (AssetViewPtr->HasFocusedDescendants())
+	else if (LegacyContentSourceWidgets->AssetViewPtr->HasFocusedDescendants())
 	{
 		// Prefer the asset menu if the asset view has focus (which may be the case when using the keyboard to invoke
 		// the action), as it is the only thing that is updated with the correct selection context when no context menu 
 		// has been invoked, and can work for both folders and files
-		if (AssetViewPtr->GetSelectedItems().Num() > 0)
+		if (LegacyContentSourceWidgets->AssetViewPtr->GetSelectedItems().Num() > 0)
 		{
-			return AssetContextMenu->CanExecuteRename();
+			return LegacyContentSourceWidgets->AssetContextMenu->CanExecuteRename();
 		}
 	}
-	else if (AssetViewPtr->GetSelectedFolderItems().Num() > 0)
+	else if (LegacyContentSourceWidgets->AssetViewPtr->GetSelectedFolderItems().Num() > 0)
 	{
 		// Folder selection takes precedence over file selection for the context menu used...
-		return PathContextMenu->CanExecuteRename();
+		return LegacyContentSourceWidgets->PathContextMenu->CanExecuteRename();
 	}
-	else if (AssetViewPtr->GetSelectedFileItems().Num() > 0)
+	else if (LegacyContentSourceWidgets->AssetViewPtr->GetSelectedFileItems().Num() > 0)
 	{
 		// ... but the asset view still takes precedence over an unfocused path view unless it has no selection
-		return AssetContextMenu->CanExecuteRename();
+		return LegacyContentSourceWidgets->AssetContextMenu->CanExecuteRename();
 	}
-	else if (PathViewPtr->GetSelectedFolderItems().Num() > 0)
+	else if (LegacyContentSourceWidgets->PathViewPtr->GetSelectedFolderItems().Num() > 0)
 	{
-		return PathContextMenu->CanExecuteRename();
+		return LegacyContentSourceWidgets->PathContextMenu->CanExecuteRename();
 	}
 
 	return false;
@@ -3695,47 +3868,47 @@ void SContentBrowser::HandleRenameCommand()
 	// The order of these conditions are carefully crafted to match the logic of the context menu summoning, as this callback 
 	// is shared between the path and asset views, and is given zero context as to which one is making the request
 	// Change this logic at your peril, lest the the dominoes fall like a house of cards (checkmate)
-	if (PathViewPtr->HasFocusedDescendants())
+	if (LegacyContentSourceWidgets->PathViewPtr->HasFocusedDescendants())
 	{
 		// Prefer the path view if it has focus, which may be the case when using the keyboard to invoke the action,
 		// but will be false when using the context menu (which isn't an issue, as the path view clears the asset view 
 		// selection when invoking its context menu to avoid the selection ambiguity present when using the keyboard)
-		if (PathViewPtr->GetSelectedFolderItems().Num() > 0)
+		if (LegacyContentSourceWidgets->PathViewPtr->GetSelectedFolderItems().Num() > 0)
 		{
-			PathContextMenu->ExecuteRename(EContentBrowserViewContext::PathView);
+			LegacyContentSourceWidgets->PathContextMenu->ExecuteRename(EContentBrowserViewContext::PathView);
 		}
 	}
-	else if (AssetViewPtr->HasFocusedDescendants())
+	else if (LegacyContentSourceWidgets->AssetViewPtr->HasFocusedDescendants())
 	{
 		// Prefer the asset menu if the asset view has focus (which may be the case when using the keyboard to invoke
 		// the action), as it is the only thing that is updated with the correct selection context when no context menu 
 		// has been invoked, and can work for both folders and files
-		if (AssetViewPtr->GetSelectedItems().Num() > 0)
+		if (LegacyContentSourceWidgets->AssetViewPtr->GetSelectedItems().Num() > 0)
 		{
-			AssetContextMenu->ExecuteRename(EContentBrowserViewContext::AssetView);
+			LegacyContentSourceWidgets->AssetContextMenu->ExecuteRename(EContentBrowserViewContext::AssetView);
 		}
 	}
-	else if (AssetViewPtr->GetSelectedFolderItems().Num() > 0)
+	else if (LegacyContentSourceWidgets->AssetViewPtr->GetSelectedFolderItems().Num() > 0)
 	{
 		// Folder selection takes precedence over file selection for the context menu used...
-		PathContextMenu->ExecuteRename(EContentBrowserViewContext::AssetView);
+		LegacyContentSourceWidgets->PathContextMenu->ExecuteRename(EContentBrowserViewContext::AssetView);
 	}
-	else if (AssetViewPtr->GetSelectedFileItems().Num() > 0)
+	else if (LegacyContentSourceWidgets->AssetViewPtr->GetSelectedFileItems().Num() > 0)
 	{
 		// ... but the asset view still takes precedence over an unfocused path view unless it has no selection
-		AssetContextMenu->ExecuteRename(EContentBrowserViewContext::AssetView);
+		LegacyContentSourceWidgets->AssetContextMenu->ExecuteRename(EContentBrowserViewContext::AssetView);
 	}
-	else if (PathViewPtr->GetSelectedFolderItems().Num() > 0)
+	else if (LegacyContentSourceWidgets->PathViewPtr->GetSelectedFolderItems().Num() > 0)
 	{
-		PathContextMenu->ExecuteRename(EContentBrowserViewContext::PathView);
+		LegacyContentSourceWidgets->PathContextMenu->ExecuteRename(EContentBrowserViewContext::PathView);
 	}
 }
 
 bool SContentBrowser::HandleSaveAssetCommandCanExecute() const
 {
-	if (AssetViewPtr->GetSelectedFileItems().Num() > 0 && !AssetViewPtr->IsRenamingAsset())
+	if (LegacyContentSourceWidgets->AssetViewPtr->GetSelectedFileItems().Num() > 0 && !LegacyContentSourceWidgets->AssetViewPtr->IsRenamingAsset())
 	{
-		return AssetContextMenu->CanExecuteSaveAsset();
+		return LegacyContentSourceWidgets->AssetContextMenu->CanExecuteSaveAsset();
 	}
 
 	return false;
@@ -3743,25 +3916,25 @@ bool SContentBrowser::HandleSaveAssetCommandCanExecute() const
 
 void SContentBrowser::HandleSaveAssetCommand()
 {
-	if (AssetViewPtr->GetSelectedFileItems().Num() > 0)
+	if (LegacyContentSourceWidgets->AssetViewPtr->GetSelectedFileItems().Num() > 0)
 	{
-		AssetContextMenu->ExecuteSaveAsset();
+		LegacyContentSourceWidgets->AssetContextMenu->ExecuteSaveAsset();
 	}
 }
 
 void SContentBrowser::HandleSaveAllCurrentFolderCommand() const
 {
-	PathContextMenu->ExecuteSaveFolder();
+	LegacyContentSourceWidgets->PathContextMenu->ExecuteSaveFolder();
 }
 
 void SContentBrowser::HandleResaveAllCurrentFolderCommand() const
 {
-	PathContextMenu->ExecuteResaveFolder();
+	LegacyContentSourceWidgets->PathContextMenu->ExecuteResaveFolder();
 }
 
 void SContentBrowser::CopySelectedAssetPathCommand() const
 {
-	PathContextMenu->CopySelectedFolder();
+	LegacyContentSourceWidgets->PathContextMenu->CopySelectedFolder();
 }
 
 bool SContentBrowser::HandleDeleteCommandCanExecute() const
@@ -3769,43 +3942,43 @@ bool SContentBrowser::HandleDeleteCommandCanExecute() const
 	// The order of these conditions are carefully crafted to match the logic of the context menu summoning, as this callback 
 	// is shared between the path and asset views, and is given zero context as to which one is making the request
 	// Change this logic at your peril, lest the the dominoes fall like a house of cards (checkmate)
-	if (PathViewPtr->HasFocusedDescendants())
+	if (LegacyContentSourceWidgets->PathViewPtr->HasFocusedDescendants())
 	{
 		// Prefer the path view if it has focus, which may be the case when using the keyboard to invoke the action,
 		// but will be false when using the context menu (which isn't an issue, as the path view clears the asset view 
 		// selection when invoking its context menu to avoid the selection ambiguity present when using the keyboard)
-		if (PathViewPtr->GetSelectedFolderItems().Num() > 0)
+		if (LegacyContentSourceWidgets->PathViewPtr->GetSelectedFolderItems().Num() > 0)
 		{
-			return PathContextMenu->CanExecuteDelete();
+			return LegacyContentSourceWidgets->PathContextMenu->CanExecuteDelete();
 		}
 	}
-	else if (AssetViewPtr->HasFocusedDescendants())
+	else if (LegacyContentSourceWidgets->AssetViewPtr->HasFocusedDescendants())
 	{
 		// Prefer the asset menu if the asset view has focus (which may be the case when using the keyboard to invoke
 		// the action), as it is the only thing that is updated with the correct selection context when no context menu 
 		// has been invoked, and can work for both folders and files
-		if (AssetViewPtr->GetSelectedItems().Num() > 0)
+		if (LegacyContentSourceWidgets->AssetViewPtr->GetSelectedItems().Num() > 0)
 		{
-			return AssetContextMenu->CanExecuteDelete();
+			return LegacyContentSourceWidgets->AssetContextMenu->CanExecuteDelete();
 		}
 	}
-	else if (AssetViewPtr->GetSelectedFolderItems().Num() > 0)
+	else if (LegacyContentSourceWidgets->AssetViewPtr->GetSelectedFolderItems().Num() > 0)
 	{
 		// Folder selection takes precedence over file selection for the context menu used...
-		return PathContextMenu->CanExecuteDelete();
+		return LegacyContentSourceWidgets->PathContextMenu->CanExecuteDelete();
 	}
-	else if (AssetViewPtr->GetSelectedFileItems().Num() > 0)
+	else if (LegacyContentSourceWidgets->AssetViewPtr->GetSelectedFileItems().Num() > 0)
 	{
 		// ... but the asset view still takes precedence over an unfocused path view unless it has no selection
-		return AssetContextMenu->CanExecuteDelete();
+		return LegacyContentSourceWidgets->AssetContextMenu->CanExecuteDelete();
 	}
-	else if (FavoritePathViewPtr->GetSelectedFolderItems().Num() > 0)
+	else if (LegacyContentSourceWidgets->FavoritePathViewPtr->GetSelectedFolderItems().Num() > 0)
 	{
 		return true;
 	}
-	else if (PathViewPtr->GetSelectedFolderItems().Num() > 0)
+	else if (LegacyContentSourceWidgets->PathViewPtr->GetSelectedFolderItems().Num() > 0)
 	{
-		return PathContextMenu->CanExecuteDelete();
+		return LegacyContentSourceWidgets->PathContextMenu->CanExecuteDelete();
 	}
 
 	return false;
@@ -3816,49 +3989,49 @@ void SContentBrowser::HandleDeleteCommandExecute()
 	// The order of these conditions are carefully crafted to match the logic of the context menu summoning, as this callback 
 	// is shared between the path and asset views, and is given zero context as to which one is making the request
 	// Change this logic at your peril, lest the the dominoes fall like a house of cards (checkmate)
-	if (PathViewPtr->HasFocusedDescendants())
+	if (LegacyContentSourceWidgets->PathViewPtr->HasFocusedDescendants())
 	{
 		// Prefer the path view if it has focus, which may be the case when using the keyboard to invoke the action,
 		// but will be false when using the context menu (which isn't an issue, as the path view clears the asset view 
 		// selection when invoking its context menu to avoid the selection ambiguity present when using the keyboard)
-		if (PathViewPtr->GetSelectedFolderItems().Num() > 0)
+		if (LegacyContentSourceWidgets->PathViewPtr->GetSelectedFolderItems().Num() > 0)
 		{
-			PathContextMenu->ExecuteDelete();
+			LegacyContentSourceWidgets->PathContextMenu->ExecuteDelete();
 		}
 	}
-	else if (AssetViewPtr->HasFocusedDescendants())
+	else if (LegacyContentSourceWidgets->AssetViewPtr->HasFocusedDescendants())
 	{
 		// Prefer the asset menu if the asset view has focus (which may be the case when using the keyboard to invoke
 		// the action), as it is the only thing that is updated with the correct selection context when no context menu 
 		// has been invoked, and can work for both folders and files
-		if (AssetViewPtr->GetSelectedItems().Num() > 0)
+		if (LegacyContentSourceWidgets->AssetViewPtr->GetSelectedItems().Num() > 0)
 		{
-			AssetContextMenu->ExecuteDelete();
+			LegacyContentSourceWidgets->AssetContextMenu->ExecuteDelete();
 		}
 	}
-	else if (AssetViewPtr->GetSelectedFolderItems().Num() > 0)
+	else if (LegacyContentSourceWidgets->AssetViewPtr->GetSelectedFolderItems().Num() > 0)
 	{
 		// Folder selection takes precedence over file selection for the context menu used...
-		PathContextMenu->ExecuteDelete();
+		LegacyContentSourceWidgets->PathContextMenu->ExecuteDelete();
 	}
-	else if (AssetViewPtr->GetSelectedFileItems().Num() > 0)
+	else if (LegacyContentSourceWidgets->AssetViewPtr->GetSelectedFileItems().Num() > 0)
 	{
 		// ... but the asset view still takes precedence over an unfocused path view unless it has no selection
-		AssetContextMenu->ExecuteDelete();
+		LegacyContentSourceWidgets->AssetContextMenu->ExecuteDelete();
 	}
-	else if (FavoritePathViewPtr->GetSelectedFolderItems().Num() > 0)
+	else if (LegacyContentSourceWidgets->FavoritePathViewPtr->GetSelectedFolderItems().Num() > 0)
 	{
-		HandleDeleteFavorite(PathContextMenu->GetParentContent());
+		HandleDeleteFavorite(LegacyContentSourceWidgets->PathContextMenu->GetParentContent());
 	}
-	else if (PathViewPtr->GetSelectedFolderItems().Num() > 0)
+	else if (LegacyContentSourceWidgets->PathViewPtr->GetSelectedFolderItems().Num() > 0)
 	{
-		PathContextMenu->ExecuteDelete();
+		LegacyContentSourceWidgets->PathContextMenu->ExecuteDelete();
 	}
 }
 
 void SContentBrowser::HandleDeleteFavorite(TSharedPtr<SWidget> ParentWidget)
 {
-	TArray<FContentBrowserItem> SelectedFolders = FavoritePathViewPtr->GetSelectedFolderItems();
+	TArray<FContentBrowserItem> SelectedFolders = LegacyContentSourceWidgets->FavoritePathViewPtr->GetSelectedFolderItems();
 	if (ParentWidget.IsValid() && SelectedFolders.Num() > 0)
 	{
 		FText Prompt;
@@ -3885,7 +4058,7 @@ void SContentBrowser::HandleDeleteFavorite(TSharedPtr<SWidget> ParentWidget)
 				}
 
 				GConfig->Flush(false, GEditorPerProjectIni);
-				FavoritePathViewPtr->Populate();
+				LegacyContentSourceWidgets->FavoritePathViewPtr->Populate();
 
 				return FReply::Handled();
 			})
@@ -3895,17 +4068,17 @@ void SContentBrowser::HandleDeleteFavorite(TSharedPtr<SWidget> ParentWidget)
 
 void SContentBrowser::HandleOpenAssetsOrFoldersCommandExecute()
 {
-	AssetViewPtr->OnOpenAssetsOrFolders();
+	LegacyContentSourceWidgets->AssetViewPtr->OnOpenAssetsOrFolders();
 }
 
 void SContentBrowser::HandlePreviewAssetsCommandExecute()
 {
-	AssetViewPtr->OnPreviewAssets();
+	LegacyContentSourceWidgets->AssetViewPtr->OnPreviewAssets();
 }
 
 void SContentBrowser::HandleCreateNewFolderCommandExecute()
 {
-	TArray<FString> SelectedPaths = PathViewPtr->GetSelectedPaths();
+	TArray<FString> SelectedPaths = LegacyContentSourceWidgets->PathViewPtr->GetSelectedPaths();
 
 	// only create folders when a single path is selected
 	UContentBrowserDataSubsystem* ContentBrowserData = IContentBrowserDataModule::Get().GetSubsystem();
@@ -3917,13 +4090,13 @@ void SContentBrowser::HandleCreateNewFolderCommandExecute()
 			SelectedPaths.Num() > 0
 			? SelectedPaths[0]
 			: FString(),
-			FOnCreateNewFolder::CreateSP(AssetViewPtr.Get(), &SAssetView::NewFolderItemRequested));
+			FOnCreateNewFolder::CreateSP(LegacyContentSourceWidgets->AssetViewPtr.Get(), &SAssetView::NewFolderItemRequested));
 	}
 }
 
 void SContentBrowser::HandleGoUpToParentFolder()
 {
-	const FString SelectedPath = PathViewPtr->GetSelectedPath();
+	const FString SelectedPath = LegacyContentSourceWidgets->PathViewPtr->GetSelectedPath();
 	int32 LastSlashIdx = INDEX_NONE;
 	if (ensure(SelectedPath.FindLastChar('/', LastSlashIdx)))
 	{
@@ -3936,7 +4109,7 @@ void SContentBrowser::HandleGoUpToParentFolder()
 bool SContentBrowser::HandleCanGoUpToParentFolder() const
 {
 	// Allow going up if there's one non-root folder selected
-	TArray<FString> SelectedPaths = PathViewPtr->GetSelectedPaths();
+	TArray<FString> SelectedPaths = LegacyContentSourceWidgets->PathViewPtr->GetSelectedPaths();
 	if (SelectedPaths.Num() == 1)
 	{
 		int32 LastSlashIdx = INDEX_NONE;
@@ -3952,14 +4125,14 @@ void SContentBrowser::GetSelectionState(TArray<FAssetData>& SelectedAssets, TArr
 {
 	SelectedAssets.Reset();
 	SelectedPaths.Reset();
-	if (AssetViewPtr->HasAnyUserFocusOrFocusedDescendants())
+	if (LegacyContentSourceWidgets->AssetViewPtr->HasAnyUserFocusOrFocusedDescendants())
 	{
-		SelectedAssets = AssetViewPtr->GetSelectedAssets();
-		SelectedPaths = AssetViewPtr->GetSelectedFolders();
+		SelectedAssets = LegacyContentSourceWidgets->AssetViewPtr->GetSelectedAssets();
+		SelectedPaths = LegacyContentSourceWidgets->AssetViewPtr->GetSelectedFolders();
 	}
-	else if (PathViewPtr->HasAnyUserFocusOrFocusedDescendants())
+	else if (LegacyContentSourceWidgets->PathViewPtr->HasAnyUserFocusOrFocusedDescendants())
 	{
-		SelectedPaths = PathViewPtr->GetSelectedPaths();
+		SelectedPaths = LegacyContentSourceWidgets->PathViewPtr->GetSelectedPaths();
 	}
 }
 
@@ -3993,6 +4166,11 @@ FText SContentBrowser::GetHistoryForwardTooltip() const
 
 void SContentBrowser::SyncGlobalSelectionSet()
 {
+	if (!ContentSourcesContainer->IsLegacyContentSourceActive())
+	{
+		return;
+	}
+	
 	USelection* EditorSelection = GEditor->GetSelectedObjects();
 	if ( !ensure( EditorSelection != NULL ) )
 	{
@@ -4000,7 +4178,7 @@ void SContentBrowser::SyncGlobalSelectionSet()
 	}
 
 	// Get the selected assets in the asset view
-	const TArray<FAssetData>& SelectedAssets = AssetViewPtr->GetSelectedAssets();
+	const TArray<FAssetData>& SelectedAssets = LegacyContentSourceWidgets->AssetViewPtr->GetSelectedAssets();
 
 	EditorSelection->BeginBatchSelectOperation();
 	{
@@ -4038,16 +4216,19 @@ void SContentBrowser::SyncGlobalSelectionSet()
 
 void SContentBrowser::UpdatePath()
 {
-	ContentBrowserUtils::UpdateNavigationBar(NavigationBar, AssetViewPtr, PathViewPtr);
+	if (ContentSourcesContainer->IsLegacyContentSourceActive())
+	{
+		ContentBrowserUtils::UpdateNavigationBar(LegacyContentSourceWidgets->NavigationBar, LegacyContentSourceWidgets->AssetViewPtr, LegacyContentSourceWidgets->PathViewPtr);
 	
-	CachedCanWriteToCurrentPath.Reset();
+		CachedCanWriteToCurrentPath.Reset();
+	}
 }
 
 void SContentBrowser::OnFilterChanged()
 {
 	TArray<TSharedRef<const FPathPermissionList>> CustomPermissionLists;
-	FARFilter Filter = FilterListPtr->GetCombinedBackendFilter(CustomPermissionLists);
-	AssetViewPtr->SetBackendFilter(Filter, &CustomPermissionLists);
+	FARFilter Filter = LegacyContentSourceWidgets->FilterListPtr->GetCombinedBackendFilter(CustomPermissionLists);
+	LegacyContentSourceWidgets->AssetViewPtr->SetBackendFilter(Filter, &CustomPermissionLists);
 
 	// Notify 'filter changed' delegate
 	FContentBrowserModule& ContentBrowserModule = FModuleManager::GetModuleChecked<FContentBrowserModule>( TEXT("ContentBrowser") );
@@ -4060,7 +4241,7 @@ FText SContentBrowser::GetPathText() const
 
 	if ( IsFilteredBySource() )
 	{
-		const FAssetViewContentSources& ContentSources = AssetViewPtr->GetContentSources();
+		const FAssetViewContentSources& ContentSources = LegacyContentSourceWidgets->AssetViewPtr->GetContentSources();
 
 		// At least one source is selected
 		const int32 NumSources = ContentSources.GetVirtualPaths().Num() + ContentSources.GetCollections().Num();
@@ -4085,7 +4266,7 @@ FText SContentBrowser::GetPathText() const
 
 bool SContentBrowser::IsFilteredBySource() const
 {
-	const FAssetViewContentSources& ContentSources = AssetViewPtr->GetContentSources();
+	const FAssetViewContentSources& ContentSources = LegacyContentSourceWidgets->AssetViewPtr->GetContentSources();
 	return !ContentSources.IsEmpty();
 }
 
@@ -4111,11 +4292,11 @@ void SContentBrowser::OnRenameRequested(const FContentBrowserItem& Item, EConten
 	{
 		if (ViewContext == EContentBrowserViewContext::AssetView)
 		{
-			AssetViewPtr->RenameItem(Item);
+			LegacyContentSourceWidgets->AssetViewPtr->RenameItem(Item);
 		}
 		else
 		{
-			PathViewPtr->RenameFolderItem(Item);
+			LegacyContentSourceWidgets->PathViewPtr->RenameFolderItem(Item);
 		}
 	}
 	else
@@ -4129,7 +4310,7 @@ void SContentBrowser::OnOpenedFolderDeleted()
 	// Since the contents of the asset view have just been deleted, set the selected path to the default "/Game"
 	TArray<FString> DefaultSelectedPaths;
 	DefaultSelectedPaths.Add(TEXT("/Game"));
-	PathViewPtr->SetSelectedPaths(DefaultSelectedPaths);
+	LegacyContentSourceWidgets->PathViewPtr->SetSelectedPaths(DefaultSelectedPaths);
 	PathSelected(TEXT("/Game"));
 }
 
@@ -4147,7 +4328,7 @@ void SContentBrowser::OnDuplicateRequested(TArrayView<const FContentBrowserItem>
 				const FContentBrowserItemDataTemporaryContext NewItemContext = OriginalItem.Duplicate();
 				if (NewItemContext.IsValid())
 				{
-					AssetViewPtr->NewFileItemRequested(NewItemContext);
+					LegacyContentSourceWidgets->AssetViewPtr->NewFileItemRequested(NewItemContext);
 				}
 			}
 			else
@@ -4204,7 +4385,7 @@ void SContentBrowser::OnDuplicateRequested(TArrayView<const FContentBrowserItem>
 
 void SContentBrowser::OnAssetViewRefreshRequested()
 {
-	AssetViewPtr->RequestSlowFullListRefresh();
+	LegacyContentSourceWidgets->AssetViewPtr->RequestSlowFullListRefresh();
 }
 
 void SContentBrowser::HandleCollectionContainerAdded(const TSharedRef<ICollectionContainer>& CollectionContainer)
@@ -4223,7 +4404,7 @@ void SContentBrowser::HandleCollectionContainerAdded(const TSharedRef<ICollectio
 void SContentBrowser::ShowCollectionContainer(const TSharedRef<ICollectionContainer>& CollectionContainer)
 {
 	TArray<FSoftObjectPath> SelectedCollectionItems;
-	for (const FContentBrowserItem& SelectedAssetItem : AssetViewPtr->GetSelectedItems())
+	for (const FContentBrowserItem& SelectedAssetItem : LegacyContentSourceWidgets->AssetViewPtr->GetSelectedItems())
 	{
 		FSoftObjectPath CollectionItemId;
 		if (SelectedAssetItem.TryGetCollectionId(CollectionItemId))
@@ -4290,9 +4471,9 @@ void SContentBrowser::HideCollectionContainer(const TSharedRef<ICollectionContai
 
 	{
 		FAssetViewContentSources NewContentSources;
-		if (UpdateContentSources(AssetViewPtr->GetContentSources(), NewContentSources))
+		if (UpdateContentSources(LegacyContentSourceWidgets->AssetViewPtr->GetContentSources(), NewContentSources))
 		{
-			AssetViewPtr->SetContentSources(NewContentSources);
+			LegacyContentSourceWidgets->AssetViewPtr->SetContentSources(NewContentSources);
 		}
 	}
 
@@ -4350,9 +4531,9 @@ void SContentBrowser::HandleCollectionRemoved(ICollectionContainer& CollectionCo
 
 	{
 		FAssetViewContentSources NewContentSources;
-		if (UpdateContentSources(AssetViewPtr->GetContentSources(), NewContentSources))
+		if (UpdateContentSources(LegacyContentSourceWidgets->AssetViewPtr->GetContentSources(), NewContentSources))
 		{
-			AssetViewPtr->SetContentSources(NewContentSources);
+			LegacyContentSourceWidgets->AssetViewPtr->SetContentSources(NewContentSources);
 		}
 	}
 
@@ -4398,9 +4579,9 @@ void SContentBrowser::HandleCollectionRenamed(ICollectionContainer& CollectionCo
 
 	{
 		FAssetViewContentSources NewContentSources;
-		if (UpdateContentSources(AssetViewPtr->GetContentSources(), NewContentSources))
+		if (UpdateContentSources(LegacyContentSourceWidgets->AssetViewPtr->GetContentSources(), NewContentSources))
 		{
-			AssetViewPtr->SetContentSources(NewContentSources);
+			LegacyContentSourceWidgets->AssetViewPtr->SetContentSources(NewContentSources);
 		}
 	}
 
@@ -4416,7 +4597,7 @@ void SContentBrowser::HandleCollectionRenamed(ICollectionContainer& CollectionCo
 
 void SContentBrowser::HandleCollectionUpdated(ICollectionContainer& CollectionContainer, const FCollectionNameType& Collection)
 {
-	const FAssetViewContentSources& ContentSources = AssetViewPtr->GetContentSources();
+	const FAssetViewContentSources& ContentSources = LegacyContentSourceWidgets->AssetViewPtr->GetContentSources();
 
 	// If we're currently viewing the dynamic collection that was updated, make sure our active filter text is up-to-date
 	if (ContentSources.IsDynamicCollection())
@@ -4431,7 +4612,7 @@ void SContentBrowser::HandleCollectionUpdated(ICollectionContainer& CollectionCo
 
 			const FText DynamicQueryText = FText::FromString(DynamicQueryString);
 			SetSearchBoxText(DynamicQueryText);
-			SearchBoxPtr->SetText(DynamicQueryText);
+			LegacyContentSourceWidgets->SearchBoxPtr->SetText(DynamicQueryText);
 		}
 	}
 }
@@ -4472,9 +4653,9 @@ void SContentBrowser::HandleItemDataUpdated(TArrayView<const FContentBrowserItem
 
 FText SContentBrowser::GetSearchAssetsHintText() const
 {
-	if (PathViewPtr.IsValid())
+	if (LegacyContentSourceWidgets->PathViewPtr.IsValid())
 	{
-		TArray<FContentBrowserItem> Paths = PathViewPtr->GetSelectedFolderItems();
+		TArray<FContentBrowserItem> Paths = LegacyContentSourceWidgets->PathViewPtr->GetSelectedFolderItems();
 		if (Paths.Num() > 0)
 		{
 			FString SearchHint = NSLOCTEXT( "ContentBrowser", "SearchBoxPartialHint", "Search" ).ToString();
@@ -4725,11 +4906,11 @@ TSharedPtr<SWidget> SContentBrowser::GetItemContextMenu(TArrayView<const FConten
 		// This is important for determining which context menu may be open based on the asset selection for rename/delete operations
 		if (ViewContext != EContentBrowserViewContext::AssetView)
 		{
-			AssetViewPtr->ClearSelection();
+			LegacyContentSourceWidgets->AssetViewPtr->ClearSelection();
 		}
 
 		// Ensure the path context menu has the up-to-date list of paths being worked on
-		PathContextMenu->SetSelectedFolders(SelectedFolders);
+		LegacyContentSourceWidgets->PathContextMenu->SetSelectedFolders(SelectedFolders);
 
 		TArray<FString> SelectedPackagePaths;
 		bool bPhysicalPathExists = false;
@@ -4760,7 +4941,10 @@ TSharedPtr<SWidget> SContentBrowser::GetItemContextMenu(TArrayView<const FConten
 		UContentBrowserFolderContext* Context = NewObject<UContentBrowserFolderContext>();
 		Context->ContentBrowser = SharedThis(this);
 		// Note: This always uses the path view to manage the temporary folder item, even if the context menu came from the favorites view, as the favorites view can't make folders correctly
-		Context->OnCreateNewFolder = ViewContext == EContentBrowserViewContext::AssetView ? FOnCreateNewFolder::CreateSP(AssetViewPtr.Get(), &SAssetView::NewFolderItemRequested) : FOnCreateNewFolder::CreateSP(PathViewPtr.Get(), &SPathView::NewFolderItemRequested);
+		Context->OnCreateNewFolder = ViewContext == EContentBrowserViewContext::AssetView
+			? FOnCreateNewFolder::CreateSP(LegacyContentSourceWidgets->AssetViewPtr.Get(), &SAssetView::NewFolderItemRequested)
+			: FOnCreateNewFolder::CreateSP(LegacyContentSourceWidgets->PathViewPtr.Get(), &SPathView::NewFolderItemRequested);
+		
 		ContentBrowserUtils::CountPathTypes(SelectedPackagePaths, Context->NumAssetPaths, Context->NumClassPaths);
 
 		FAssetToolsModule& AssetToolsModule = FModuleManager::LoadModuleChecked<FAssetToolsModule>("AssetTools");
@@ -4784,7 +4968,12 @@ TSharedPtr<SWidget> SContentBrowser::GetItemContextMenu(TArrayView<const FConten
 			// Include the items that are not folders to help the batch operations operate on these also.
 			DataContextObject->SelectedItems = SelectedItems;
 			DataContextObject->bCanBeModified = Context->bCanBeModified;
-			DataContextObject->ParentWidget = ViewContext == EContentBrowserViewContext::AssetView ? TSharedPtr<SWidget>(AssetViewPtr) : ViewContext == EContentBrowserViewContext::FavoriteView ? TSharedPtr<SWidget>(FavoritePathViewPtr) : TSharedPtr<SWidget>(PathViewPtr);
+			DataContextObject->ParentWidget = ViewContext == EContentBrowserViewContext::AssetView
+				? TSharedPtr<SWidget>(LegacyContentSourceWidgets->AssetViewPtr)
+				: ViewContext == EContentBrowserViewContext::FavoriteView
+						? TSharedPtr<SWidget>(LegacyContentSourceWidgets->FavoritePathViewPtr)
+						: TSharedPtr<SWidget>(LegacyContentSourceWidgets->PathViewPtr);
+			
 			MenuContext.AddObject(DataContextObject);
 		}
 
@@ -4804,7 +4993,7 @@ TSharedPtr<SWidget> SContentBrowser::GetItemContextMenu(TArrayView<const FConten
 	{
 		// Files selected - show the file menu
 		checkf(ViewContext == EContentBrowserViewContext::AssetView, TEXT("File items were passed from a path view!"));
-		return AssetContextMenu->MakeContextMenu(SelectedItems, AssetViewPtr->GetContentSources(), Commands);
+		return LegacyContentSourceWidgets->AssetContextMenu->MakeContextMenu(SelectedItems, LegacyContentSourceWidgets->AssetViewPtr->GetContentSources(), Commands);
 	}
 	else if (ViewContext == EContentBrowserViewContext::AssetView)
 	{
@@ -4820,7 +5009,7 @@ void SContentBrowser::PopulateFolderContextMenu(UToolMenu* Menu)
 	UContentBrowserFolderContext* Context = Menu->FindContext<UContentBrowserFolderContext>();
 	check(Context);
 
-	const TArray<FContentBrowserItem>& SelectedFolders = PathContextMenu->GetSelectedFolders();
+	const TArray<FContentBrowserItem>& SelectedFolders = LegacyContentSourceWidgets->PathContextMenu->GetSelectedFolders();
 
 	// We can only create folders when we have a single path selected
 	UContentBrowserDataSubsystem* ContentBrowserData = IContentBrowserDataModule::Get().GetSubsystem();
@@ -4870,7 +5059,7 @@ void SContentBrowser::PopulateFolderContextMenu(UToolMenu* Menu)
 		);
 	}
 
-	PathContextMenu->MakePathViewContextMenu(Menu);
+	LegacyContentSourceWidgets->PathContextMenu->MakePathViewContextMenu(Menu);
 }
 
 void SContentBrowser::CreateNewFolder(FString FolderPath, FOnCreateNewFolder InOnCreateNewFolder)
@@ -4909,7 +5098,7 @@ void SContentBrowser::CreateNewFolder(FString FolderPath, FOnCreateNewFolder InO
 
 void SContentBrowser::OpenNewContentBrowser()
 {
-	const TArray<FContentBrowserItem> SelectedFolders = PathContextMenu->GetSelectedFolders();
+	const TArray<FContentBrowserItem> SelectedFolders = LegacyContentSourceWidgets->PathContextMenu->GetSelectedFolders();
 	FContentBrowserSingleton::Get().SyncBrowserToItems(SelectedFolders, false, true, NAME_None, true);
 }
  
diff --git a/Engine/Source/Editor/ContentBrowser/Private/SContentBrowser.h b/Engine/Source/Editor/ContentBrowser/Private/SContentBrowser.h
--- a/Engine/Source/Editor/ContentBrowser/Private/SContentBrowser.h
+++ b/Engine/Source/Editor/ContentBrowser/Private/SContentBrowser.h
@@ -29,11 +29,18 @@
 #include "Widgets/SCompoundWidget.h"
 #include "Widgets/ParentWidgetSwitcher.h"
 
-namespace UE::Editor::ContentBrowser::Private
+namespace UE::Editor::ContentBrowser
 {
-	class SContentBrowserSourceTree;
-	class FContentBrowserSourceTreeEntry;
-	class SContentBrowserSourceTreeArea;
+	class SContentSourcesView;
+	class SLegacyContentSource;
+
+	namespace Private
+	{
+		class SContentBrowserSourceTree;
+		class FContentBrowserSourceTreeEntry;
+		class SContentBrowserSourceTreeArea;
+	}
+	
 }
 
 class FContentBrowserItemDataTemporaryContext;
@@ -589,6 +596,9 @@ private:
 	/** Bind our UI commands */
 	void BindCommands();
 
+	/** Unbind our UI commands */
+	void UnbindCommands();
+
 	/** Gets the visibility of the favorites view */
 	EVisibility GetFavoriteFolderVisibility() const;
 
@@ -666,8 +676,68 @@ private:
 	FContentBrowserInstanceConfig* CreateEditorConfigIfRequired();
 
 	void UpdatePrivateContentFeatureEnabled(bool bUpdateFilterIfChanged);
+
+	/** Executed when the legacy content source is enabled. */
+	void OnLegacyContentSourceEnabled();
+	
+	/** Executed when the legacy content source is enabled. */
+	void OnLegacyContentSourceDisabled();
+	
+	/** Create the content browser widgets (SPathView, SAssetView, search and filter bars etc) */
+	TSharedRef<SWidget> CreateLegacyAssetViewWidgets();
+	
 private:
 
+	/** Struct that contains the widgets used by content browser's default layout (i.e Legacy Content Source) */
+	struct FLegacyContentSourceWidgets
+	{
+		/** The context menu manager for the path view */
+		TSharedPtr<class FPathContextMenu> PathContextMenu;
+
+		/** The asset view widget */
+		TSharedPtr<SAssetView> AssetViewPtr;
+
+		/** The border that holds the content in AssetView */
+		TSharedPtr<SBorder> AssetViewBorder;
+		
+		/** A helper class to manage asset context menu options */
+		TSharedPtr<class FAssetContextMenu> AssetContextMenu;
+
+		/** The filter list */
+		TSharedPtr<SFilterList> FilterListPtr;
+		
+		/** The text box used to search for assets */
+		TSharedPtr<SAssetSearchBox> SearchBoxPtr;
+		
+		/** Combine breadcrumb/text-box widget for showing & changing location */
+		TSharedPtr<SNavigationBar> NavigationBar;
+
+		/** The favorites tree widget */
+		TSharedPtr<SFavoritePathView> FavoritePathViewPtr;
+
+		/** The asset tree widget */
+		TSharedPtr<SPathView> PathViewPtr;
+
+		TSharedPtr<UE::Editor::ContentBrowser::Private::SContentBrowserSourceTree> SourceTreePtr;
+
+		/** The generated navigation toolbar widget */
+		TSharedPtr<SWidget> NavigationToolBarWidget;
+		
+		/** The splitter between the path & asset view */
+		TSharedPtr<SSplitter> PathAssetSplitterPtr;
+
+		/** The splitter between the path, favorite and collection views */
+		TSharedPtr<SSplitter> PathFavoriteSplitterPtr;
+		
+		/** Handles the navigation bar location switching */
+		TSharedPtr<UE::Editor::ContentBrowser::Private::TWidgetSizeParentSwitcher<Orient_Horizontal>> NavigationToolBarParentSwitcher;
+
+		/** The number of fixed slots in the splitter, slots after this number are collection view slots */
+		int32 SourceTreeSplitterNumFixedSlots = 0;
+	};
+
+	TSharedPtr<FLegacyContentSourceWidgets> LegacyContentSourceWidgets;
+
 	/** The tab that contains this browser */
 	TWeakPtr<SDockTab> ContainingTab;
 
@@ -679,12 +749,6 @@ private:
 	 *  "find in content browser" or typing in a path should populate it 
 	 */
 	TMRUArray<FString> JumpMRU;
-	
-	/** A helper class to manage asset context menu options */
-	TSharedPtr<class FAssetContextMenu> AssetContextMenu;
-
-	/** The context menu manager for the path view */
-	TSharedPtr<class FPathContextMenu> PathContextMenu;
 
 	/** The sources search for favorites */
 	TSharedPtr<FSourcesSearch> FavoritesSearch;
@@ -692,38 +756,9 @@ private:
 	/** The sources search for paths */
 	TSharedPtr<FSourcesSearch> SourcesSearch;
 
-	TSharedPtr<UE::Editor::ContentBrowser::Private::SContentBrowserSourceTree> SourceTreePtr;
-
-	/** The asset tree widget */
-	TSharedPtr<SPathView> PathViewPtr;
-
-	/** The favorites tree widget */
-	TSharedPtr<SFavoritePathView> FavoritePathViewPtr;
-
-	/** The asset view widget */
-	TSharedPtr<SAssetView> AssetViewPtr;
-	
-	/** Combine breadcrumb/text-box widget for showing & changing location */
-	TSharedPtr<SNavigationBar> NavigationBar;
-
-	/** Handles the navigation bar location switching */
-	TSharedPtr<UE::Editor::ContentBrowser::Private::TWidgetSizeParentSwitcher<Orient_Horizontal>> NavigationToolBarParentSwitcher;
-
-	/** The generated navigation toolbar widget */
-	TSharedPtr<SWidget> NavigationToolBarWidget;
-
-	/** The text box used to search for assets */
-	TSharedPtr<SAssetSearchBox> SearchBoxPtr;
-
-	/** The filter list */
-	TSharedPtr<SFilterList> FilterListPtr;
-
 	/** The Combo Button used to summon the filter dropdown */
 	TSharedPtr<SWidget> FilterComboButton;
 
-	/** The border that holds the content in AssetView */
-	TSharedPtr<SBorder> AssetViewBorder;
-
 	/** Favorites area widget */
 	TSharedPtr<UE::Editor::ContentBrowser::Private::SContentBrowserSourceTreeArea> FavoritesArea;
 
@@ -778,15 +813,6 @@ private:
 	/** Delegate used to create a new folder */
 	FOnCreateNewFolder OnCreateNewFolder;
 
-	/** The splitter between the path & asset view */
-	TSharedPtr<SSplitter> PathAssetSplitterPtr;
-
-	/** The splitter between the path, favorite and collection views */
-	TSharedPtr<SSplitter> PathFavoriteSplitterPtr;
-
-	/** The number of fixed slots in the splitter, slots after this number are collection view slots */
-	int32 SourceTreeSplitterNumFixedSlots = 0;
-
 	/** The list of plugin filters currently applied to the path view */
 	TSharedPtr<FPluginFilterCollectionType> PluginPathFilters;
 
@@ -796,6 +822,20 @@ private:
 	/** Handle to a callback for when CVars change */
 	FConsoleVariableSinkHandle CVarSinkHandle;
 
+	/** Widget that contains the content sources selector (if enabled) and the currently enabled content source */
+	TSharedPtr<UE::Editor::ContentBrowser::SContentSourcesView> ContentSourcesContainer;
+	
+	/** A legacy content source, which when enabled displays the default content browser widgets */
+	TSharedPtr<UE::Editor::ContentBrowser::SLegacyContentSource> LegacyContentSource;
+
+	/** A copy of the init params used by the content browser so we can re-create the widgets dynamically
+	 *  If no init params were provided, the default values will be used.
+	 */
+	FContentBrowserConfig InitConfig;
+
+	/** Whether this content browser was specified an init config */
+	bool bHasInitConfig = false;
+
 	/** Registers the "Add New" drop-down menu */
 	static FDelayedAutoRegisterHelper AddNewContextMenuRegistration;
 
diff --git a/Engine/Source/Editor/ContentBrowser/Public/ContentSources/Columns/ContentSourcesColumns.h b/Engine/Source/Editor/ContentBrowser/Public/ContentSources/Columns/ContentSourcesColumns.h
new file mode 100644
--- /dev/null
+++ b/Engine/Source/Editor/ContentBrowser/Public/ContentSources/Columns/ContentSourcesColumns.h
@@ -0,0 +1,19 @@
+// Copyright Epic Games, Inc. All Rights Reserved.
+
+#pragma once
+
+#include "ContentSources/IContentSource.h"
+#include "Elements/Interfaces/TypedElementDataStorageInterface.h"
+
+#include "ContentSourcesColumns.generated.h"
+
+/**
+ * Column added to a widget row that contains a pointer to the content source the widget belongs to
+ */
+USTRUCT(meta = (DisplayName = "Content Source"))
+struct FContentSourceColumn : public FEditorDataStorageColumn
+{
+	GENERATED_BODY()
+	
+	TWeakPtr<UE::Editor::ContentBrowser::IContentSource> ContentSource;
+};
\ No newline at end of file
diff --git a/Engine/Source/Editor/ContentBrowser/Public/ContentSources/IContentSource.h b/Engine/Source/Editor/ContentBrowser/Public/ContentSources/IContentSource.h
new file mode 100644
--- /dev/null
+++ b/Engine/Source/Editor/ContentBrowser/Public/ContentSources/IContentSource.h
@@ -0,0 +1,54 @@
+// Copyright Epic Games, Inc. All Rights Reserved.
+
+#pragma once
+
+#include "Elements/Common/TypedElementQueryDescription.h"
+#include "Internationalization/Text.h"
+#include "Textures/SlateIcon.h"
+#include "UObject/NameTypes.h"
+
+namespace UE::Editor::ContentBrowser
+{
+	/*
+	 * Init Params used by the STedsTableViewer displayed in place of the legacy Content Browser when this content source is active
+	 */
+	struct FTableViewerInitParams
+	{
+		// Query description to populate the rows of the table viewer
+		DataStorage::FQueryDescription QueryDescription;
+
+		// The columns shown in the table viewer
+		TArray<TWeakObjectPtr<const UScriptStruct>> Columns;
+
+		// The widget purposes used by the table viewer to display the widgets in the cells
+		TArray<FName> WidgetPurposes;
+	};
+
+	/*
+	 * A Content Source is a construct that can customize the look and behavior of the content browser through TEDS (Editor Data Storage) intended
+	 * to replace the legacy asset specific content browser.
+	 * Currently supports displaying an STedsTableViewer driven by a query specified by the content source in place of the legacy content browser
+	 * layout and widgets
+	 * Content Sources can be registered with IContentBrowserSingleton which will cause them to be displayed in a vertical toolbar next to the
+	 * Content Browser allowing a user to select and activate them.
+	 * 
+	 * NOTE: This API is experimental and subject to change
+	 */
+	class IContentSource
+	{
+	public:
+		virtual ~IContentSource() = default;
+
+		// Get the internal name of the content source
+		virtual FName GetName() = 0;
+
+		// Get the user facing name of the content source
+		virtual FText GetDisplayName() = 0;
+
+		// Get an icon representing the content source
+		virtual FSlateIcon GetIcon() = 0;
+
+		// Get the init params used by the TEDS Table Viewer when this content source is active
+		virtual void GetAssetViewInitParams(FTableViewerInitParams& OutInitParams) = 0;
+	};
+}
\ No newline at end of file
diff --git a/Engine/Source/Editor/ContentBrowser/Public/IContentBrowserSingleton.h b/Engine/Source/Editor/ContentBrowser/Public/IContentBrowserSingleton.h
--- a/Engine/Source/Editor/ContentBrowser/Public/IContentBrowserSingleton.h
+++ b/Engine/Source/Editor/ContentBrowser/Public/IContentBrowserSingleton.h
@@ -47,6 +47,11 @@ namespace EAssetViewType
 	};
 }
 
+namespace UE::Editor::ContentBrowser
+{
+	class IContentSource;
+}
+
 
 /** A selection of items in the Content Browser */
 struct FContentBrowserSelection
@@ -142,7 +147,9 @@ struct FContentBrowserConfig
 		, ThumbnailScale(0.1f)
 		, InitialAssetViewType(EAssetViewType::Tile)
 		, SelectedCollection(nullptr, NAME_None, ECollectionShareType::CST_Local)
+	PRAGMA_DISABLE_DEPRECATION_WARNINGS
 		, SelectedCollectionName( NAME_None, ECollectionShareType::CST_Local )
+	PRAGMA_ENABLE_DEPRECATION_WARNINGS
 		, bShowBottomToolbar(true)
 		, bCanShowClasses(true)
 		, bUseSourcesView(true)
@@ -814,4 +821,20 @@ public:
 
 	/** Docks the current content browser drawer as a tabbed content browser in a layout */
 	virtual void DockContentBrowserDrawer() = 0;
+
+	/** Factory delegate that instantiates a content source  */
+	DECLARE_DELEGATE_RetVal(TSharedRef<UE::Editor::ContentBrowser::IContentSource>, FContentSourceFactory)
+
+	/** Register a content source by name, causing it to appear in the content sources bar as selectable (if the bar is enabled)  */
+	virtual void RegisterContentSourceFactory(const FName& InContentSourceName, const FContentSourceFactory& InContentSourceFactory) = 0;
+
+	/** Unregister a content source by name */
+	virtual void UnregisterContentSourceFactory(const FName& InContentSourceName) = 0;
+
+	/** Execute a delegate for each registered content source factory */
+	virtual void ForEachContentSourceFactory(const TFunctionRef<void(const FName&, const FContentSourceFactory&)>& InFunction) = 0;
+
+	/** Delegate fired when a new content source factory is registered/unregistered */
+	DECLARE_MULTICAST_DELEGATE(FOnContentSourceFactoriesChanged)
+	virtual FOnContentSourceFactoriesChanged& OnContentSourceFactoriesChanged() = 0;
 };
