SContentSourcesView: Fix build error due to conflicting variable names

diff --git a/Engine/Source/Editor/ContentBrowser/Private/ContentSources/Widgets/SContentSourcesView.cpp b/Engine/Source/Editor/ContentBrowser/Private/ContentSources/Widgets/SContentSourcesView.cpp
--- a/Engine/Source/Editor/ContentBrowser/Private/ContentSources/Widgets/SContentSourcesView.cpp
+++ b/Engine/Source/Editor/ContentBrowser/Private/ContentSources/Widgets/SContentSourcesView.cpp
@@ -42,7 +42,7 @@ namespace UE::Editor::ContentBrowser
 			{
 				if (const UContentSourcesViewMenuContext* Context = InSection.FindContext<UContentSourcesViewMenuContext>())
 				{
-					if (const TSharedPtr<SContentSourcesView> ContentSourceWidget = Context->ContentSourcesWidget.Pin())
+					if (const TSharedPtr<SContentSourcesView> ContentSourceWidgetPin = Context->ContentSourcesWidget.Pin())
 					{
 						// Legacy Content Source entry
 						InSection.AddEntry(
@@ -51,25 +51,25 @@ namespace UE::Editor::ContentBrowser
 								FUIAction(
 									FExecuteAction::CreateLambda([Context]()
 									{
-										if (TSharedPtr<SContentSourcesView> ContentSourceWidget = Context->ContentSourcesWidget.Pin())
+										if (TSharedPtr<SContentSourcesView> ContextContentSourceWidget = Context->ContentSourcesWidget.Pin())
 										{
-											ContentSourceWidget->ActivateLegacyContentSource();
+											ContextContentSourceWidget->ActivateLegacyContentSource();
 										}
 									}),
 									FCanExecuteAction(),
 									FIsActionChecked::CreateLambda([Context]()
 									{
-										if (TSharedPtr<SContentSourcesView> ContentSourceWidget = Context->ContentSourcesWidget.Pin())
+										if (TSharedPtr<SContentSourcesView> ContextContentSourceWidget = Context->ContentSourcesWidget.Pin())
 										{
-											return ContentSourceWidget->IsLegacyContentSourceActive();
+											return ContextContentSourceWidget->IsLegacyContentSourceActive();
 										}
 										return false;
 									}),
 									FIsActionButtonVisible::CreateLambda([Context]()
 									{
-										if (TSharedPtr<SContentSourcesView> ContentSourceWidget = Context->ContentSourcesWidget.Pin())
+										if (TSharedPtr<SContentSourcesView> ContextContentSourceWidget = Context->ContentSourcesWidget.Pin())
 										{
-											return ContentSourceWidget->HasLegacyContentSource();
+											return ContextContentSourceWidget->HasLegacyContentSource();
 										}
 										return false;
 									})),
@@ -89,17 +89,17 @@ namespace UE::Editor::ContentBrowser
 								FUIAction(
 									FExecuteAction::CreateLambda([ContentSource, Context]()
 									{
-										if (TSharedPtr<SContentSourcesView> ContentSourceWidget = Context->ContentSourcesWidget.Pin())
+										if (TSharedPtr<SContentSourcesView> ContextContentSourceWidget = Context->ContentSourcesWidget.Pin())
 										{
-											ContentSourceWidget->ActivateContentSource(ContentSource);
+											ContextContentSourceWidget->ActivateContentSource(ContentSource);
 										}
 									}),
 									FCanExecuteAction(),
 									FIsActionChecked::CreateLambda([ContentSource, Context]()
 									{
-										if (TSharedPtr<SContentSourcesView> ContentSourceWidget = Context->ContentSourcesWidget.Pin())
+										if (TSharedPtr<SContentSourcesView> ContextContentSourceWidget = Context->ContentSourcesWidget.Pin())
 										{
-											return ContentSourceWidget->IsContentSourceActive(ContentSource);
+											return ContextContentSourceWidget->IsContentSourceActive(ContentSource);
 										}
 										return false;
 									})),
@@ -110,7 +110,7 @@ namespace UE::Editor::ContentBrowser
 							));
 						};
 
-						ContentSourceWidget->ForEachContentSource(MakeSourceEntry);
+						ContentSourceWidgetPin->ForEachContentSource(MakeSourceEntry);
 					}
 				}
 			}));
