SContentBrowser: Fix implicit access of deprecated member through default assignment operator

diff --git a/Engine/Source/Editor/ContentBrowser/Private/SContentBrowser.cpp b/Engine/Source/Editor/ContentBrowser/Private/SContentBrowser.cpp
--- a/Engine/Source/Editor/ContentBrowser/Private/SContentBrowser.cpp
+++ b/Engine/Source/Editor/ContentBrowser/Private/SContentBrowser.cpp
@@ -322,7 +322,10 @@ void SContentBrowser::Construct( const FArguments& InArgs, const FName& InInstan
 	// Store a copy of the init config if specified so we can re-create the asset view widgets dynamically
 	if (Config)
 	{
+		// We have to disable deprecation warnings because the default assignment operator copies a deprecated member
+		PRAGMA_DISABLE_DEPRECATION_WARNINGS
 		InitConfig = *Config;
+		PRAGMA_ENABLE_DEPRECATION_WARNINGS
 		bHasInitConfig = true;
 	}
 
