[ContentBrowser] Add names on tiny folders
Added Overlay names for Folders when in tiny size and fixed border radius

diff --git a/Engine/Source/Editor/ContentBrowser/Private/AssetViewWidgets.cpp b/Engine/Source/Editor/ContentBrowser/Private/AssetViewWidgets.cpp
--- a/Engine/Source/Editor/ContentBrowser/Private/AssetViewWidgets.cpp
+++ b/Engine/Source/Editor/ContentBrowser/Private/AssetViewWidgets.cpp
@@ -2202,16 +2202,39 @@ void SAssetTileItem::Construct( const FArguments& InArgs )
 					.BorderImage(IsFolder() ? TAttribute<const FSlateBrush*>(this, &SAssetTileItem::GetFolderBackgroundImage) : TAttribute<const FSlateBrush*>(this, &SAssetTileItem::GetNameAreaBackgroundImage))
 					[
 						SNew(SVerticalBox)
-						// Thumbnail
 						+SVerticalBox::Slot()
 						.AutoHeight()
 						[
-							// The remainder of the space is reserved for the name.
-							SNew(SBox)
-							.WidthOverride(this, &SAssetTileItem::GetThumbnailBoxSize)
-							.HeightOverride(this, &SAssetTileItem::GetThumbnailBoxSize)
+							SNew(SOverlay)
+
+							// Thumbnail
+							+ SOverlay::Slot()
 							[
-								ItemContents
+								// The remainder of the space is reserved for the name.
+								SNew(SBox)
+								.WidthOverride(this, &SAssetTileItem::GetThumbnailBoxSize)
+								.HeightOverride(this, &SAssetTileItem::GetThumbnailBoxSize)
+								[
+									ItemContents
+								]
+							]
+
+							// Overlay name for tiny folders
+							+ SOverlay::Slot()
+							.HAlign(HAlign_Center)
+							.VAlign(VAlign_Center)
+							.Padding(10.f, 0.f)
+							[
+								SNew(SBorder)
+								.Padding(4.f, 2.f)
+								.BorderImage(FAppStyle::GetBrush("ContentBrowser.AssetTileItem.TinyFolderTextBorder"))
+								.Visibility(this, &SAssetTileItem::GetTinyFolderNameOverlayVisibility)
+								[
+									SNew(STextBlock)
+									.ColorAndOpacity(FStyleColors::White)
+									.OverflowPolicy(ETextOverflowPolicy::Ellipsis)
+									.Text(GetNameText())
+								]
 							]
 						]
 
@@ -2485,6 +2508,16 @@ FOptionalSize SAssetTileItem::GetThumbnailBoxSize() const
 	return UE::Editor::ContentBrowser::IsNewStyleEnabled() ? FOptionalSize(ThumbnailDimension.Get()) : FOptionalSize(ItemWidth.Get()- ThumbnailPadding);
 }
 
+EVisibility SAssetTileItem::GetTinyFolderNameOverlayVisibility() const
+{
+	if (!UE::Editor::ContentBrowser::IsNewStyleEnabled())
+	{
+		return EVisibility::Collapsed;
+	}
+	// Use the name area visibility so that if the logic change we will still show the overlay only when that's not visible
+	return GetNameAreaVisibility() != EVisibility::Visible && IsFolder() ? EVisibility::Visible : EVisibility::Collapsed;
+}
+
 EVisibility SAssetTileItem::GetAssetClassLabelVisibility() const
 {
 	return (!IsFolder() && bShowType)
diff --git a/Engine/Source/Editor/ContentBrowser/Public/AssetViewWidgets.h b/Engine/Source/Editor/ContentBrowser/Public/AssetViewWidgets.h
--- a/Engine/Source/Editor/ContentBrowser/Public/AssetViewWidgets.h
+++ b/Engine/Source/Editor/ContentBrowser/Public/AssetViewWidgets.h
@@ -724,6 +724,9 @@ protected:
 	/** Returns the size of the thumbnail widget */
 	FOptionalSize GetThumbnailBoxSize() const;
 
+	/** Gets the visibility of the folder name overlay on the folder icon */
+	EVisibility GetTinyFolderNameOverlayVisibility() const;
+
 	/** Gets the visibility of the asset class label in thumbnails */
 	EVisibility GetAssetClassLabelVisibility() const;
 
diff --git a/Engine/Source/Editor/EditorStyle/Private/StarshipStyle.cpp b/Engine/Source/Editor/EditorStyle/Private/StarshipStyle.cpp
--- a/Engine/Source/Editor/EditorStyle/Private/StarshipStyle.cpp
+++ b/Engine/Source/Editor/EditorStyle/Private/StarshipStyle.cpp
@@ -6813,11 +6813,12 @@ void FStarshipEditorStyle::FStyle::SetupContentBrowserStyle()
 		Set("ContentBrowser.AssetTileItem.NameAreaHoverBackground", new FSlateRoundedBoxBrush(FStyleColors::Hover, FVector4(0.0f, 0.0f, 4.0f, 4.0f)));
 		Set("ContentBrowser.AssetTileItem.NameAreaSelectedBackground", new FSlateRoundedBoxBrush(FStyleColors::Primary, FVector4(0.0f, 0.0f, 4.0f, 4.0f)));
 		Set("ContentBrowser.AssetTileItem.NameAreaSelectedHoverBackground", new FSlateRoundedBoxBrush(FStyleColors::PrimaryHover, FVector4(0.0f, 0.0f, 4.0f, 4.0f)));
+		Set("ContentBrowser.AssetTileItem.TinyFolderTextBorder", new FSlateRoundedBoxBrush(COLOR("#00000080"), 4.f, COLOR("#FFFFFF1A"), 1.f));
 
 		// Asset Thumbnail Border
-		Set("ContentBrowser.AssetTileItem.AssetBorderHoverBackground", new FSlateRoundedBoxBrush(FStyleColors::Transparent, FVector4(4.0f, 2.0f, 4.0f, 4.0f), FStyleColors::Hover, 2.f));
-		Set("ContentBrowser.AssetTileItem.AssetBorderSelectedBackground", new FSlateRoundedBoxBrush(FStyleColors::Transparent, FVector4(4.0f, 2.0f, 4.0f, 4.0f), FStyleColors::Primary, 2.f));
-		Set("ContentBrowser.AssetTileItem.AssetBorderSelectedHoverBackground", new FSlateRoundedBoxBrush(FStyleColors::Transparent, FVector4(4.0f, 2.0f, 4.0f, 4.0f), FStyleColors::PrimaryHover, 2.f));
+		Set("ContentBrowser.AssetTileItem.AssetBorderHoverBackground", new FSlateRoundedBoxBrush(FStyleColors::Transparent, FVector4(4.0f, 4.0f, 4.0f, 4.0f), FStyleColors::Hover, 2.f));
+		Set("ContentBrowser.AssetTileItem.AssetBorderSelectedBackground", new FSlateRoundedBoxBrush(FStyleColors::Transparent, FVector4(4.0f, 4.0f, 4.0f, 4.0f), FStyleColors::Primary, 2.f));
+		Set("ContentBrowser.AssetTileItem.AssetBorderSelectedHoverBackground", new FSlateRoundedBoxBrush(FStyleColors::Transparent, FVector4(4.0f, 4.0f, 4.0f, 4.0f), FStyleColors::PrimaryHover, 2.f));
 
 		// Tile Item Border
 		Set("ContentBrowser.AssetTileItem.AssetContent", new FSlateRoundedBoxBrush(FStyleColors::Secondary, FVector4(4.0f, 4.0f, 4.0f, 4.0f)));
