IContentBrowserSingleton: Fix deprecation warnings caused by default assignment operator

diff --git a/Engine/Source/Editor/ContentBrowser/Public/IContentBrowserSingleton.h b/Engine/Source/Editor/ContentBrowser/Public/IContentBrowserSingleton.h
--- a/Engine/Source/Editor/ContentBrowser/Public/IContentBrowserSingleton.h
+++ b/Engine/Source/Editor/ContentBrowser/Public/IContentBrowserSingleton.h
@@ -163,6 +163,10 @@ struct FContentBrowserConfig
 		, bCanShowLockButton(true)
 		, bCanSetAsPrimaryBrowser(true)
 	{ }
+
+	PRAGMA_DISABLE_DEPRECATION_WARNINGS
+		FContentBrowserConfig& operator=(const FContentBrowserConfig&) = default;
+	PRAGMA_ENABLE_DEPRECATION_WARNINGS
 };
 
 
