SLegacyContentSource: Fix incorrect #pragma in .cpp file

diff --git a/Engine/Source/Editor/ContentBrowser/Private/ContentSources/Widgets/SLegacyContentSource.cpp b/Engine/Source/Editor/ContentBrowser/Private/ContentSources/Widgets/SLegacyContentSource.cpp
--- a/Engine/Source/Editor/ContentBrowser/Private/ContentSources/Widgets/SLegacyContentSource.cpp
+++ b/Engine/Source/Editor/ContentBrowser/Private/ContentSources/Widgets/SLegacyContentSource.cpp
@@ -1,7 +1,5 @@
 // Copyright Epic Games, Inc. All Rights Reserved.
 
-#pragma once
-
 #include "SLegacyContentSource.h"
 
 namespace UE::Editor::ContentBrowser
