SContentSourcesView: Remove #pragma from cpp file

diff --git a/Engine/Source/Editor/ContentBrowser/Private/ContentSources/Widgets/SContentSourcesView.cpp b/Engine/Source/Editor/ContentBrowser/Private/ContentSources/Widgets/SContentSourcesView.cpp
--- a/Engine/Source/Editor/ContentBrowser/Private/ContentSources/Widgets/SContentSourcesView.cpp
+++ b/Engine/Source/Editor/ContentBrowser/Private/ContentSources/Widgets/SContentSourcesView.cpp
@@ -1,7 +1,5 @@
 // Copyright Epic Games, Inc. All Rights Reserved.
 
-#pragma once
-
 #include "SContentSourcesView.h"
 
 #include "ContentBrowserStyle.h"
