Added support for toolbar entry expansion behavior

diff --git a/Engine/Source/Developer/ToolMenus/Private/ToolMenus.cpp b/Engine/Source/Developer/ToolMenus/Private/ToolMenus.cpp
--- a/Engine/Source/Developer/ToolMenus/Private/ToolMenus.cpp
+++ b/Engine/Source/Developer/ToolMenus/Private/ToolMenus.cpp
@@ -2050,10 +2050,16 @@ void UToolMenus::PopulateToolBarBuilder(FToolBarBuilder& ToolBarBuilder, UToolMe
 		// previous alignment group.
 		if (bIsMiddleOrLast && NumAddedSections > 0)
 		{
+			static constexpr float AlmostZero = 0.1f; // Using 0.0f results in different behavior, so we just use a small, sub-1 pixel value.
+
 			FMenuEntryStyleParams StyleParams;
+
 			StyleParams.HorizontalAlignment = HAlign_Right;
 			StyleParams.SizeRule = FSizeParam::ESizeRule::SizeRule_StretchContent;
-			StyleParams.MinSize = 0.5f;
+			StyleParams.FillSize = 1.0f;
+			StyleParams.FillSizeMin = AlmostZero; // Allow the spacer to shrink to nothing if another widget needs the space.
+			StyleParams.MinimumSize = AlmostZero;
+			StyleParams.DesiredWidthOverride = StyleParams.DesiredHeightOverride = AlmostZero;
 
 			FMenuEntryResizeParams ResizeParams;
 			// Never show spacers in overflow menus.
diff --git a/Engine/Source/Editor/ContentBrowser/Private/Menus/ContentBrowserMenus.cpp b/Engine/Source/Editor/ContentBrowser/Private/Menus/ContentBrowserMenus.cpp
--- a/Engine/Source/Editor/ContentBrowser/Private/Menus/ContentBrowserMenus.cpp
+++ b/Engine/Source/Editor/ContentBrowser/Private/Menus/ContentBrowserMenus.cpp
@@ -110,9 +110,9 @@ FDelayedAutoRegisterHelper SContentBrowser::NavigationBarMenuRegistration(
 									FText::GetEmpty()));
 
 						PathBoxEntry.WidgetData.StyleParams.SizeRule = FSizeParam::SizeRule_StretchContent;
-						PathBoxEntry.WidgetData.StyleParams.MinSize = NavigationPathMinWidth;
-						PathBoxEntry.WidgetData.StyleParams.MaxSize = NavigationPathMaxWidth;
-						PathBoxEntry.WidgetData.StyleParams.DesiredWidth = NavigationPathMaxWidth;
+						PathBoxEntry.WidgetData.StyleParams.MinimumSize = NavigationPathMinWidth;
+						PathBoxEntry.WidgetData.StyleParams.MaximumSize = NavigationPathMaxWidth;
+						PathBoxEntry.WidgetData.StyleParams.DesiredWidthOverride = NavigationPathMaxWidth;
 					}
 				}
 			}));
@@ -383,9 +383,9 @@ FDelayedAutoRegisterHelper SContentBrowser::ToolBarMenuRegistration(
                                     FText::GetEmpty()));
 
                             SearchBoxEntry.WidgetData.StyleParams.SizeRule = FSizeParam::SizeRule_StretchContent;
-                            SearchBoxEntry.WidgetData.StyleParams.MinSize = SearchBoxMinWidth;
-                            SearchBoxEntry.WidgetData.StyleParams.MaxSize = SearchBoxMaxWidth;
-                            SearchBoxEntry.WidgetData.StyleParams.DesiredWidth = SearchBoxMaxWidth;
+                            SearchBoxEntry.WidgetData.StyleParams.MinimumSize = SearchBoxMinWidth;
+                            SearchBoxEntry.WidgetData.StyleParams.MaximumSize = SearchBoxMaxWidth;
+                            SearchBoxEntry.WidgetData.StyleParams.DesiredWidthOverride = SearchBoxMaxWidth;
                         }
                     }
                 }));
diff --git a/Engine/Source/Editor/ContentBrowser/Private/Widgets/ParentWidgetSwitcher.h b/Engine/Source/Editor/ContentBrowser/Private/Widgets/ParentWidgetSwitcher.h
--- a/Engine/Source/Editor/ContentBrowser/Private/Widgets/ParentWidgetSwitcher.h
+++ b/Engine/Source/Editor/ContentBrowser/Private/Widgets/ParentWidgetSwitcher.h
@@ -43,6 +43,12 @@ namespace UE::Editor::ContentBrowser::Private
 			return true;
 	    }
 
+		TSharedPtr<SWidget> GetParentWidget() const
+		{
+			check(ParentWidgets.IsValidIndex(CurrentSubjectWidgetParentIdx));
+			return ParentWidgets[CurrentSubjectWidgetParentIdx];
+		}
+
 	protected:
 		virtual int32 GetParentIndex() const = 0;
 
diff --git a/Engine/Source/Editor/ContentBrowser/Private/Widgets/WidgetSizeSwitcher.h b/Engine/Source/Editor/ContentBrowser/Private/Widgets/WidgetSizeSwitcher.h
new file mode 100644
--- /dev/null
+++ b/Engine/Source/Editor/ContentBrowser/Private/Widgets/WidgetSizeSwitcher.h
@@ -0,0 +1,85 @@
+// Copyright Epic Games, Inc. All Rights Reserved.
+
+#pragma once
+
+#include "CoreTypes.h"
+#include "Widgets/Layout/SBox.h"
+#include "Widgets/SNullWidget.h"
+#include "Widgets/SWidget.h"
+
+namespace UE::Editor::ContentBrowser::Private
+{
+	/**
+	 * Returns a size override to bind to based on a single subject widget's desired size, clamped between a provided size range.
+	 * If a reference widget is provided, it's size is used for the range max value.
+	 * Axis is 0 for x/width, 1 for y/height.
+	 */
+	template <EAxis::Type Axis>
+	class TWidgetDesiredSizeSwitcher
+	{
+		static_assert(Axis == EAxis::X || Axis == EAxis::Y, "Axis must be 0 (X/Width) or 1 (Y/Height).");
+
+	public:
+		virtual ~TWidgetDesiredSizeSwitcher() = default;
+
+		TWidgetDesiredSizeSwitcher(
+			const TSharedPtr<SWidget>& InSubjectWidget,
+			const TSharedPtr<SWidget>& InMaxSizeReferenceWidget,
+			const FInt16Range& InSizeRange)
+			: SubjectWidget(InSubjectWidget)
+			, MaxSizeReferenceWidget(InMaxSizeReferenceWidget)
+			, SizeRange(InSizeRange)
+		{
+		}
+
+		void SetMaxSizeReferenceWidget(const TSharedPtr<SWidget>& InWidget)
+		{
+			if (InWidget != MaxSizeReferenceWidget)
+			{
+				MaxSizeReferenceWidget = InWidget;
+			}
+		}
+
+		void SetSizeRange(const FInt16Range& InRange)
+		{
+			if (SizeRange != InRange)
+			{
+				SizeRange = InRange;
+			}
+		}
+
+		float GetDesiredSizeOverride()
+		{
+			const int16 CurrentSubjectWidgetDesiredSize = static_cast<int16>(FMath::FloorToInt(SubjectWidget->GetDesiredSize()[Axis - 1] - UE_KINDA_SMALL_NUMBER));
+			const int16 MaxSize =
+				MaxSizeReferenceWidget.IsValid()
+				? static_cast<int16>(FMath::FloorToInt(MaxSizeReferenceWidget->GetCachedGeometry().GetLocalSize()[Axis - 1] - UE_KINDA_SMALL_NUMBER))
+				: SizeRange.GetUpperBound().GetValue();
+
+			// Check if the subject or reference widget (max size) has changed
+			if (CurrentSubjectWidgetDesiredSize == LastSubjectWidgetDesiredSize
+				&& MaxSize == LastMaxSize)
+			{
+				return LastDesiredSizeOverride;
+			}
+
+			LastMaxSize = MaxSize;
+			LastSubjectWidgetDesiredSize = CurrentSubjectWidgetDesiredSize;
+
+			const int16 MinSize = SizeRange.GetLowerBound().GetValue();
+
+			LastDesiredSizeOverride = LastSubjectWidgetDesiredSize > MinSize ? MaxSize : MinSize;
+
+			return LastDesiredSizeOverride;
+		}
+
+	protected:
+		TSharedPtr<SWidget> SubjectWidget;
+		TSharedPtr<SWidget> MaxSizeReferenceWidget;
+		FInt16Range SizeRange;
+
+		int16 LastMaxSize = INDEX_NONE;
+		float LastDesiredSizeOverride = -1.0f;
+		int16 LastSubjectWidgetDesiredSize = INDEX_NONE;
+	};
+}
diff --git a/Engine/Source/Runtime/Slate/Private/Framework/MultiBox/MultiBox.cpp b/Engine/Source/Runtime/Slate/Private/Framework/MultiBox/MultiBox.cpp
--- a/Engine/Source/Runtime/Slate/Private/Framework/MultiBox/MultiBox.cpp
+++ b/Engine/Source/Runtime/Slate/Private/Framework/MultiBox/MultiBox.cpp
@@ -840,14 +840,10 @@ void SMultiBoxWidget::AddBlockWidget(const FMultiBlock& Block, TSharedPtr<SHoriz
 
 				FinalMultiBoxWidgets.Add(FinalWidgetWithHook.ToSharedPtr(), BlockRef.ToSharedPtr());
 
+				NewSlot.AutoWidth();
+
 				if (bOverride)
 				{
-					if (StyleParams.SizeRule.IsSet()
-						&& StyleParams.SizeRule.GetValue() == FSizeParam::SizeRule_Auto)
-					{
-						NewSlot.AutoWidth();
-					}
-
 					NewSlot
 					.HAlign(StyleParams.HorizontalAlignment)
 					.VAlign(StyleParams.VerticalAlignment.Get(VAlign_Fill));
@@ -857,28 +853,26 @@ void SMultiBoxWidget::AddBlockWidget(const FMultiBlock& Block, TSharedPtr<SHoriz
 						FSizeParam::ESizeRule SizeRule = StyleParams.SizeRule.GetValue();
 						if (SizeRule == FSizeParam::SizeRule_Stretch)
 						{
-							NewSlot.FillWidth(1.0f);
+							NewSlot.FillWidth(StyleParams.FillSize.Get(1.0f));
 						}
 						else if (SizeRule == FSizeParam::SizeRule_StretchContent)
 						{
-							NewSlot.FillContentWidth(1.0f);
+							NewSlot.FillContentWidth(
+								StyleParams.FillSize.Get(1.0f),
+								StyleParams.FillSizeMin.Get(StyleParams.FillSize.Get(1.0f)));
 						}
 					}
 
-					if (StyleParams.MinSize.IsSet())
+					if (StyleParams.MinimumSize.IsSet() || StyleParams.MinimumSize.IsBound())
 					{
-						NewSlot.MinWidth(StyleParams.MinSize.GetValue());
+						NewSlot.MinWidth(StyleParams.MinimumSize);
 					}
 
-					if (StyleParams.MaxSize.IsSet())
+					if (StyleParams.MaximumSize.IsSet() || StyleParams.MaximumSize.IsBound())
 					{
-						NewSlot.MaxWidth(StyleParams.MaxSize.GetValue());
+						NewSlot.MaxWidth(StyleParams.MaximumSize);
 					}
 				}
-				else
-				{
-					NewSlot.AutoWidth();
-				}
 			}
 		}
 		break;
diff --git a/Engine/Source/Runtime/Slate/Private/Framework/MultiBox/MultiBoxBuilder.cpp b/Engine/Source/Runtime/Slate/Private/Framework/MultiBox/MultiBoxBuilder.cpp
--- a/Engine/Source/Runtime/Slate/Private/Framework/MultiBox/MultiBoxBuilder.cpp
+++ b/Engine/Source/Runtime/Slate/Private/Framework/MultiBox/MultiBoxBuilder.cpp
@@ -740,8 +740,8 @@ void FToolBarBuilder::AddWidgetInternal(
 	TSharedRef<SWidget> ChildWidget = InWidget;
 	InWidget =
 		SNew( SBox )
-		.WidthOverride(InStyleParams.DesiredWidth)
-		.HeightOverride(InStyleParams.DesiredHeight)
+		.WidthOverride(InStyleParams.DesiredWidthOverride)
+		.HeightOverride(InStyleParams.DesiredHeightOverride)
 		.AddMetaData<FTagMetaData>(FTagMetaData(InTutorialHighlightName))
 		[
 			ChildWidget
diff --git a/Engine/Source/Runtime/Slate/Private/Framework/MultiBox/SWidgetBlock.cpp b/Engine/Source/Runtime/Slate/Private/Framework/MultiBox/SWidgetBlock.cpp
--- a/Engine/Source/Runtime/Slate/Private/Framework/MultiBox/SWidgetBlock.cpp
+++ b/Engine/Source/Runtime/Slate/Private/Framework/MultiBox/SWidgetBlock.cpp
@@ -58,14 +58,16 @@ bool FWidgetBlock::GetAlignmentOverrides(FMenuEntryStyleParams& OutAlignmentPara
 	if (StyleParams.SizeRule.IsSet())
 	{
 		OutAlignmentParameters.SizeRule = StyleParams.SizeRule;
+		OutAlignmentParameters.FillSize = StyleParams.FillSize;
+		OutAlignmentParameters.FillSizeMin = StyleParams.FillSizeMin;
 	}
 	else if (StyleParams.HorizontalAlignment == HAlign_Fill)
 	{
 		OutAlignmentParameters.SizeRule = FSizeParam::ESizeRule::SizeRule_Auto;
 	}
 
-	OutAlignmentParameters.MinSize = StyleParams.MinSize;
-	OutAlignmentParameters.MaxSize = StyleParams.MaxSize;
+	OutAlignmentParameters.MinimumSize = StyleParams.MinimumSize;
+	OutAlignmentParameters.MaximumSize = StyleParams.MaximumSize;
 
 	return true;
 }
diff --git a/Engine/Source/Runtime/Slate/Private/Framework/MultiBox/SWidgetBlock.h b/Engine/Source/Runtime/Slate/Private/Framework/MultiBox/SWidgetBlock.h
--- a/Engine/Source/Runtime/Slate/Private/Framework/MultiBox/SWidgetBlock.h
+++ b/Engine/Source/Runtime/Slate/Private/Framework/MultiBox/SWidgetBlock.h
@@ -25,7 +25,7 @@ public:
 	 * @param	InStyleParams				Optional additional style parameters. @see FMenuEntryStyleParams
 	 * @param	InIcon						Optional icon to be shown to the left of the label/content
 	 */
-	FWidgetBlock(TSharedRef<SWidget> InContent, const FText& InLabel, const TAttribute<FText>& InToolTipText = FText(), const FMenuEntryStyleParams& InStyleParams = {}, const TAttribute<FSlateIcon>& InIcon = TAttribute<FSlateIcon>());
+	FWidgetBlock(TSharedRef<SWidget> InContent, const FText& InLabel, const TAttribute<FText>& InToolTipText = FText(), const FMenuEntryStyleParams& InStyleParams = FMenuEntryStyleParams(), const TAttribute<FSlateIcon>& InIcon = TAttribute<FSlateIcon>());
 
 	/** FMultiBlock interface */
 	virtual void CreateMenuEntry(class FMenuBuilder& MenuBuilder) const override;
diff --git a/Engine/Source/Runtime/Slate/Public/Framework/MultiBox/MultiBoxDefs.h b/Engine/Source/Runtime/Slate/Public/Framework/MultiBox/MultiBoxDefs.h
--- a/Engine/Source/Runtime/Slate/Public/Framework/MultiBox/MultiBoxDefs.h
+++ b/Engine/Source/Runtime/Slate/Public/Framework/MultiBox/MultiBoxDefs.h
@@ -142,7 +142,14 @@ namespace EMultiBlockLocation
 /** Contains various Style parameters and overrides. Not all are applicable to a given entry */
 struct FMenuEntryStyleParams
 {
+	// Workaround for clang deprecation warnings for deprecated members
+PRAGMA_DISABLE_DEPRECATION_WARNINGS
 	FMenuEntryStyleParams() = default;
+	FMenuEntryStyleParams(FMenuEntryStyleParams&&) = default;
+	FMenuEntryStyleParams(const FMenuEntryStyleParams&) = default;
+	FMenuEntryStyleParams& operator=(FMenuEntryStyleParams&&) = default;
+	FMenuEntryStyleParams& operator=(const FMenuEntryStyleParams&) = default;
+PRAGMA_ENABLE_DEPRECATION_WARNINGS
 
 	/** If true, removes the padding from the left of the widget that lines it up with other menu items */
 	bool bNoIndent = false;
@@ -157,16 +164,38 @@ struct FMenuEntryStyleParams
 	TOptional<FSizeParam::ESizeRule> SizeRule;
 
 	/** (Optionally) override the minimum size. This will apply to the width or height, depending on the menu orientation */
+	UE_DEPRECATED(5.6, "Use MinimumSize instead")
 	TOptional<float> MinSize;
 
 	/** (Optionally) override the maximum size. This will apply to the width or height, depending on the menu orientation */
+	UE_DEPRECATED(5.6, "Use MaximumSize instead")
 	TOptional<float> MaxSize;
 
 	/** (Optionally) set the desired width override */
+	UE_DEPRECATED(5.6, "Use DesiredWidthOverride instead")
 	TOptional<float> DesiredWidth;
 
 	/** (Optionally) set the desired height override */
+	UE_DEPRECATED(5.6, "Use DesiredHeightOverride instead")
 	TOptional<float> DesiredHeight;
+
+	/** (Optionally) override the fill proportion when the SizeRule is Stretch or StretchContent, defaults to 1.0f */
+	TOptional<float> FillSize;
+
+	/** (Optionally) override the minimum fill proportion when the SizeRule is StretchContent, defaults to FillSize/1.0f */
+	TOptional<float> FillSizeMin;
+
+	/** (Optionally) override the minimum size. This will apply to the width or height, depending on the menu orientation */
+	TAttribute<float> MinimumSize;
+
+	/** (Optionally) override the maximum size. This will apply to the width or height, depending on the menu orientation */
+	TAttribute<float> MaximumSize;
+
+	/** (Optionally) set the desired width override */
+	TAttribute<FOptionalSize> DesiredWidthOverride;
+
+	/** (Optionally) set the desired height override */
+	TAttribute<FOptionalSize> DesiredHeightOverride;
 };
 
 struct FMenuEntryResizeParams
