[ContentBrowser] Search/Breadcrumb expansion behavior

diff --git a/Engine/Source/Editor/ContentBrowser/Private/Menus/ContentBrowserMenus.cpp b/Engine/Source/Editor/ContentBrowser/Private/Menus/ContentBrowserMenus.cpp
--- a/Engine/Source/Editor/ContentBrowser/Private/Menus/ContentBrowserMenus.cpp
+++ b/Engine/Source/Editor/ContentBrowser/Private/Menus/ContentBrowserMenus.cpp
@@ -109,10 +109,35 @@ FDelayedAutoRegisterHelper SContentBrowser::NavigationBarMenuRegistration(
 									ContentBrowser->LegacyContentSourceWidgets->NavigationBar.ToSharedRef(),
 									FText::GetEmpty()));
 
+						ContentBrowser->LegacyContentSourceWidgets->NavigationBarSizeSwitcher->SetSizeRange(FInt16Range(static_cast<int16>(FMath::FloorToInt(NavigationPathMaxWidth))));
+
 						PathBoxEntry.WidgetData.StyleParams.SizeRule = FSizeParam::SizeRule_StretchContent;
 						PathBoxEntry.WidgetData.StyleParams.MinimumSize = NavigationPathMinWidth;
-						PathBoxEntry.WidgetData.StyleParams.MaximumSize = NavigationPathMaxWidth;
-						PathBoxEntry.WidgetData.StyleParams.DesiredWidthOverride = NavigationPathMaxWidth;
+						PathBoxEntry.WidgetData.StyleParams.MaximumSize = TAttribute<float>::Create(
+							TAttribute<float>::FGetter::CreateSPLambda(
+								ContentBrowser,
+								[SizeSwitcher = ContentBrowser->LegacyContentSourceWidgets->NavigationBarSizeSwitcher.ToWeakPtr()]() -> float
+								{
+									if (SizeSwitcher.IsValid())
+									{
+										return SizeSwitcher.Pin()->GetDesiredSizeOverride();
+									}
+
+									return 0.0f;
+								}));
+
+						PathBoxEntry.WidgetData.StyleParams.DesiredWidthOverride = TAttribute<FOptionalSize>::Create(
+							TAttribute<FOptionalSize>::FGetter::CreateSPLambda(
+								ContentBrowser,
+								[SizeSwitcher = ContentBrowser->LegacyContentSourceWidgets->NavigationBarSizeSwitcher.ToWeakPtr()]() -> FOptionalSize
+								{
+									if (SizeSwitcher.IsValid())
+									{
+										return FOptionalSize(SizeSwitcher.Pin()->GetDesiredSizeOverride());
+									}
+
+									return FOptionalSize();
+								}));
 					}
 				}
 			}));
@@ -233,14 +258,14 @@ FDelayedAutoRegisterHelper SContentBrowser::ToolBarMenuRegistration(
                             .Icon(FAppStyle::Get().GetBrush("Icons.Plus"))
                             .Text(LOCTEXT("AddAssetButton", "Add"));
 
-                        InSection.AddEntry(
-                            FToolMenuEntry::InitWidget(
-                                "NewButton",
-                                NewButton,
-                                FText::GetEmpty(),
-                                true,
-                                false
-                            ));
+                    	InSection.AddEntry(
+							FToolMenuEntry::InitWidget(
+								"NewButton",
+								NewButton,
+								FText::GetEmpty(),
+								true,
+								false
+							));
                     }
                 }
             }));
@@ -302,14 +327,14 @@ FDelayedAutoRegisterHelper SContentBrowser::ToolBarMenuRegistration(
                                 ];
                         }
 
-                        InSection.AddEntry(
-                            FToolMenuEntry::InitWidget(
-                                "SaveButton",
-                                SaveButton.ToSharedRef(),
-                                FText::GetEmpty(),
-                                true,
-                                false
-                            ));
+                    	InSection.AddEntry(
+							FToolMenuEntry::InitWidget(
+								"SaveButton",
+								SaveButton.ToSharedRef(),
+								FText::GetEmpty(),
+								true,
+								false
+							));
                     }
                 }
             }));
@@ -382,10 +407,36 @@ FDelayedAutoRegisterHelper SContentBrowser::ToolBarMenuRegistration(
                                     ContentBrowser->LegacyContentSourceWidgets->SearchBoxPtr.ToSharedRef(),
                                     FText::GetEmpty()));
 
-                            SearchBoxEntry.WidgetData.StyleParams.SizeRule = FSizeParam::SizeRule_StretchContent;
-                            SearchBoxEntry.WidgetData.StyleParams.MinimumSize = SearchBoxMinWidth;
-                            SearchBoxEntry.WidgetData.StyleParams.MaximumSize = SearchBoxMaxWidth;
-                            SearchBoxEntry.WidgetData.StyleParams.DesiredWidthOverride = SearchBoxMaxWidth;
+                        	ContentBrowser->LegacyContentSourceWidgets->SearchBoxSizeSwitcher->SetSizeRange(FInt16Range(static_cast<int16>(FMath::FloorToInt(SearchBoxMaxWidth))));
+
+							SearchBoxEntry.WidgetData.StyleParams.SizeRule = FSizeParam::SizeRule_StretchContent;
+                        	SearchBoxEntry.WidgetData.StyleParams.FillSize = 1.0f;
+							SearchBoxEntry.WidgetData.StyleParams.MinimumSize = SearchBoxMinWidth;
+							SearchBoxEntry.WidgetData.StyleParams.MaximumSize = TAttribute<float>::Create(
+								TAttribute<float>::FGetter::CreateSPLambda(
+									ContentBrowser,
+									[SizeSwitcher = ContentBrowser->LegacyContentSourceWidgets->SearchBoxSizeSwitcher.ToWeakPtr()]() -> float
+									{
+										if (SizeSwitcher.IsValid())
+										{
+											return SizeSwitcher.Pin()->GetDesiredSizeOverride();
+										}
+
+										return 0.0f;
+									}));
+
+                        	SearchBoxEntry.WidgetData.StyleParams.DesiredWidthOverride = TAttribute<FOptionalSize>::Create(
+                        		TAttribute<FOptionalSize>::FGetter::CreateSPLambda(
+                        			ContentBrowser,
+                        			[SizeSwitcher = ContentBrowser->LegacyContentSourceWidgets->SearchBoxSizeSwitcher.ToWeakPtr()]() -> FOptionalSize
+                        			{
+                        				if (SizeSwitcher.IsValid())
+                        				{
+                        					return FOptionalSize(SizeSwitcher.Pin()->GetDesiredSizeOverride());
+                        				}
+
+                        				return FOptionalSize();
+                        			}));
                         }
                     }
                 }));
diff --git a/Engine/Source/Editor/ContentBrowser/Private/SContentBrowser.cpp b/Engine/Source/Editor/ContentBrowser/Private/SContentBrowser.cpp
--- a/Engine/Source/Editor/ContentBrowser/Private/SContentBrowser.cpp
+++ b/Engine/Source/Editor/ContentBrowser/Private/SContentBrowser.cpp
@@ -3380,11 +3380,27 @@ TSharedRef<SWidget> SContentBrowser::CreateLegacyAssetViewWidgets()
 			.Size(0.8f);
 		}
 
-		LegacyContentSourceWidgets->NavigationToolBarWidget = CreateNavigationToolBar(&InitConfig);
+		LegacyContentSourceWidgets->SearchBoxSizeSwitcher = MakeShared<TWidgetDesiredSizeSwitcher<EAxis::X>>(
+			LegacyContentSourceWidgets->SearchBoxPtr,
+			nullptr,
+			FInt16Range(100)); // @note: this is overridden in menu registration
+
+		TSharedRef<SWidget> ToolBarWidget = CreateToolBar(&InitConfig);
+
+		LegacyContentSourceWidgets->SearchBoxSizeSwitcher->SetMaxSizeReferenceWidget(ToolBarWidget.ToSharedPtr());
 
 		TSharedPtr<SBox> AssetViewNavigationToolBarContainer = nullptr;
 		TSharedPtr<SBox> SourceTreeAndAssetViewNavigationToolBarContainer = nullptr;
 
+		LegacyContentSourceWidgets->NavigationBarSizeSwitcher = MakeShared<TWidgetDesiredSizeSwitcher<EAxis::X>>(
+			LegacyContentSourceWidgets->NavigationBar,
+			nullptr,
+			FInt16Range(100)); // @note: this is overridden in menu registration
+
+		LegacyContentSourceWidgets->NavigationToolBarWidget = CreateNavigationToolBar(&InitConfig);
+
+		LegacyContentSourceWidgets->NavigationBarSizeSwitcher->SetMaxSizeReferenceWidget(LegacyContentSourceWidgets->NavigationToolBarWidget);
+
 		FinalWidget = 
 		SNew(SVerticalBox)
 
@@ -3471,7 +3487,7 @@ TSharedRef<SWidget> SContentBrowser::CreateLegacyAssetViewWidgets()
 									.VAlign(VAlign_Center)
 									.Padding(5, 0, 0, 0)
 									[
-										CreateToolBar(&InitConfig)
+										ToolBarWidget
 									]
 
 									+ SHorizontalBox::Slot()
@@ -3509,6 +3525,7 @@ TSharedRef<SWidget> SContentBrowser::CreateLegacyAssetViewWidgets()
 								.BorderImage(bIsDrawer ? FStyleDefaults::GetNoBrush() : FAppStyle::Get().GetBrush("Brushes.Panel"))
 								[
 									SAssignNew(AssetViewNavigationToolBarContainer, SBox)
+									.HAlign(HAlign_Fill)
 									.Padding(0)
 									[
 										LegacyContentSourceWidgets->NavigationToolBarWidget.ToSharedRef()
diff --git a/Engine/Source/Editor/ContentBrowser/Private/SContentBrowser.h b/Engine/Source/Editor/ContentBrowser/Private/SContentBrowser.h
--- a/Engine/Source/Editor/ContentBrowser/Private/SContentBrowser.h
+++ b/Engine/Source/Editor/ContentBrowser/Private/SContentBrowser.h
@@ -28,6 +28,7 @@
 #include "Widgets/Navigation/SBreadcrumbTrail.h"
 #include "Widgets/SCompoundWidget.h"
 #include "Widgets/ParentWidgetSwitcher.h"
+#include "Widgets/WidgetSizeSwitcher.h"
 
 namespace UE::Editor::ContentBrowser
 {
@@ -708,7 +709,10 @@ private:
 		
 		/** The text box used to search for assets */
 		TSharedPtr<SAssetSearchBox> SearchBoxPtr;
-		
+
+		/** Handles the search box sizing within the toolbar */
+		TSharedPtr<UE::Editor::ContentBrowser::Private::TWidgetDesiredSizeSwitcher<EAxis::X>> SearchBoxSizeSwitcher;
+
 		/** Combine breadcrumb/text-box widget for showing & changing location */
 		TSharedPtr<SNavigationBar> NavigationBar;
 
@@ -732,6 +736,9 @@ private:
 		/** Handles the navigation bar location switching */
 		TSharedPtr<UE::Editor::ContentBrowser::Private::TWidgetSizeParentSwitcher<Orient_Horizontal>> NavigationToolBarParentSwitcher;
 
+		/** Handles the navigation bar sizing within the navigation toolbar */
+		TSharedPtr<UE::Editor::ContentBrowser::Private::TWidgetDesiredSizeSwitcher<EAxis::X>> NavigationBarSizeSwitcher;
+
 		/** The number of fixed slots in the splitter, slots after this number are collection view slots */
 		int32 SourceTreeSplitterNumFixedSlots = 0;
 	};
diff --git a/Engine/Source/Editor/ContentBrowser/Private/Widgets/WidgetSizeSwitcher.h b/Engine/Source/Editor/ContentBrowser/Private/Widgets/WidgetSizeSwitcher.h
--- a/Engine/Source/Editor/ContentBrowser/Private/Widgets/WidgetSizeSwitcher.h
+++ b/Engine/Source/Editor/ContentBrowser/Private/Widgets/WidgetSizeSwitcher.h
@@ -53,7 +53,7 @@ namespace UE::Editor::ContentBrowser::Private
 			const int16 CurrentSubjectWidgetDesiredSize = static_cast<int16>(FMath::FloorToInt(SubjectWidget->GetDesiredSize()[Axis - 1] - UE_KINDA_SMALL_NUMBER));
 			const int16 MaxSize =
 				MaxSizeReferenceWidget.IsValid()
-				? static_cast<int16>(FMath::FloorToInt(MaxSizeReferenceWidget->GetCachedGeometry().GetLocalSize()[Axis - 1] - UE_KINDA_SMALL_NUMBER))
+				? static_cast<int16>(FMath::FloorToInt(MaxSizeReferenceWidget->GetDesiredSize()[Axis - 1] - UE_KINDA_SMALL_NUMBER))
 				: SizeRange.GetUpperBound().GetValue();
 
 			// Check if the subject or reference widget (max size) has changed
