[ContentBrowser] Menu Favorite/Path/Collection splitter header tweaks

diff --git a/Engine/Source/Editor/ContentBrowser/Private/ContentBrowserStyle.cpp b/Engine/Source/Editor/ContentBrowser/Private/ContentBrowserStyle.cpp
--- a/Engine/Source/Editor/ContentBrowser/Private/ContentBrowserStyle.cpp
+++ b/Engine/Source/Editor/ContentBrowser/Private/ContentBrowserStyle.cpp
@@ -76,8 +76,8 @@ namespace UE::ContentBrowser::Private
 
 		if (Editor::ContentBrowser::IsNewStyleEnabled())
 		{
-			Set("ContentBrowser.AssetTreeHeaderBrush", new FSlateColorBrush(FStyleColors::Header));
-			Set("ContentBrowser.AssetTreeBodyBrush", new FSlateColorBrush(FStyleColors::Header));
+			Set("ContentBrowser.AssetTreeHeaderBrush", new FSlateColorBrush(FStyleColors::Panel));
+			Set("ContentBrowser.AssetTreeBodyBrush", new FSlateColorBrush(FStyleColors::Recessed));
 		}
 		else
 		{
@@ -325,7 +325,7 @@ namespace UE::ContentBrowser::Private
 
 		const FSplitterStyle ContentBrowserSplitterStyle =
 			FSplitterStyle()
-			.SetHandleNormalBrush(FSlateColorBrush(FStyleColors::Transparent))
+			.SetHandleNormalBrush(FSlateColorBrush(FStyleColors::Recessed))
 			.SetHandleHighlightBrush(FSlateColorBrush(FStyleColors::Transparent));
 
 		Set("ContentBrowser.Splitter", ContentBrowserSplitterStyle);
diff --git a/Engine/Source/Editor/ContentBrowser/Private/Widgets/SContentBrowserSourceTree.cpp b/Engine/Source/Editor/ContentBrowser/Private/Widgets/SContentBrowserSourceTree.cpp
--- a/Engine/Source/Editor/ContentBrowser/Private/Widgets/SContentBrowserSourceTree.cpp
+++ b/Engine/Source/Editor/ContentBrowser/Private/Widgets/SContentBrowserSourceTree.cpp
@@ -172,7 +172,7 @@ namespace UE::Editor::ContentBrowser::Private
 
 		constexpr float SourceTreeSectionPadding = 2.0f;
 		constexpr float SourceTreeSectionHandleSize = 8.0f;
-		constexpr float SourceTreeHeaderHeight = 36.0f;
+		constexpr float SourceTreeHeaderHeight = 32.0f;
 
 		SAssignNew(Splitter, SSplitter)
 		.Style(&FContentBrowserStyle::Get().GetWidgetStyle<FSplitterStyle>("ContentBrowser.Splitter"))
@@ -186,6 +186,7 @@ namespace UE::Editor::ContentBrowser::Private
 		{
 			AddSlot()
 			.AreaWidget(Slot._AreaWidget)
+			.HeaderHeight(SourceTreeHeaderHeight)
 			.Size(Slot._Size)
 			.Visibility(Slot._Visibility);
 		}
@@ -238,7 +239,7 @@ namespace UE::Editor::ContentBrowser::Private
 
 		static const float VerticalHeaderPadding =
 			UE::Editor::ContentBrowser::IsNewStyleEnabled()
-			? HorizontalPadding - 4.0f
+			? HorizontalPadding - 6.0f
 			: 0.0f;
 
 		static const float SearchButtonRightPadding =
@@ -249,6 +250,7 @@ namespace UE::Editor::ContentBrowser::Private
 		static constexpr float EmptyBodyLabelPadding = 8.0f;
 
 		static const FMargin HeaderPadding = FMargin(HorizontalPadding, VerticalHeaderPadding);
+		static const FMargin ExpandableAreaPadding = FMargin(0.f, 1.f, 0.f , 0.f);
 
 		TSharedPtr<SWidget> SearchButtonWidget = SNullWidget::NullWidget;
 		if (InSearch.IsValid())
@@ -285,7 +287,7 @@ namespace UE::Editor::ContentBrowser::Private
 			.BodyBorderImage(FContentBrowserStyle::Get().GetBrush("ContentBrowser.AssetTreeBodyBrush"))
 			.HeaderPadding(HeaderPadding)
 			.Visibility(InArgs._Visibility)
-			.Padding(0.0f)
+			.Padding(ExpandableAreaPadding)
 			.AllowAnimatedTransition(true)
 			.OnAreaExpansionChanged(this, &SContentBrowserSourceTreeArea::OnAreaExpansionChanged)
 			.HeaderContent()
