[ContentBrowser] Folder Hover/Selection/SelectionHover style update

diff --git a/Engine/Source/Editor/ContentBrowser/Private/AssetViewWidgets.cpp b/Engine/Source/Editor/ContentBrowser/Private/AssetViewWidgets.cpp
--- a/Engine/Source/Editor/ContentBrowser/Private/AssetViewWidgets.cpp
+++ b/Engine/Source/Editor/ContentBrowser/Private/AssetViewWidgets.cpp
@@ -2184,9 +2184,10 @@ void SAssetTileItem::Construct( const FArguments& InArgs )
 		constexpr float NameAreaBoxLeftRightBotPadding = 4.f;
 		constexpr float NameAreaBoxTopPadding = 6.f;
 		constexpr float ClassNameMaxHeight = 14.f;
+		const FMargin FolderImagePadding = FMargin(2.f, 2.f, 2.f, 2.f);
 
 		ChildSlot
-		.Padding(FMargin(BorderPadding, BorderPadding, 0.f, 0.f))
+		.Padding(FMargin(BorderPadding, BorderPadding, 0.f, 1.f))
 		[
 			// Drop shadow border
 			SNew(SBorder)
@@ -2209,13 +2210,19 @@ void SAssetTileItem::Construct( const FArguments& InArgs )
 
 							// Thumbnail
 							+ SOverlay::Slot()
+							.Padding(IsFolder() ? FolderImagePadding : 0.f)
 							[
-								// The remainder of the space is reserved for the name.
-								SNew(SBox)
-								.WidthOverride(this, &SAssetTileItem::GetThumbnailBoxSize)
-								.HeightOverride(this, &SAssetTileItem::GetThumbnailBoxSize)
+								SNew(SBorder)
+								.Padding(0.f)
+								.BorderImage(IsFolder() ? TAttribute<const FSlateBrush*>(this, &SAssetTileItem::GetFolderSlotBackgroundImage) : FAppStyle::GetNoBrush())
 								[
-									ItemContents
+									// The remainder of the space is reserved for the name.
+									SNew(SBox)
+									.WidthOverride(this, &SAssetTileItem::GetThumbnailBoxSize)
+									.HeightOverride(this, &SAssetTileItem::GetThumbnailBoxSize)
+									[
+										ItemContents
+									]
 								]
 							]
 
@@ -2505,7 +2512,8 @@ FOptionalSize SAssetTileItem::GetStateIconImageSize() const
 
 FOptionalSize SAssetTileItem::GetThumbnailBoxSize() const
 {
-	return UE::Editor::ContentBrowser::IsNewStyleEnabled() ? FOptionalSize(ThumbnailDimension.Get()) : FOptionalSize(ItemWidth.Get()- ThumbnailPadding);
+	const float ThumbnailSizeOffset = IsFolder() ? -4.f : 0.f;
+	return UE::Editor::ContentBrowser::IsNewStyleEnabled() ? FOptionalSize(ThumbnailDimension.Get() + ThumbnailSizeOffset) : FOptionalSize(ItemWidth.Get()- ThumbnailPadding);
 }
 
 EVisibility SAssetTileItem::GetTinyFolderNameOverlayVisibility() const
@@ -2515,7 +2523,7 @@ EVisibility SAssetTileItem::GetTinyFolderNameOverlayVisibility() const
 		return EVisibility::Collapsed;
 	}
 	// Use the name area visibility so that if the logic change we will still show the overlay only when that's not visible
-	return GetNameAreaVisibility() != EVisibility::Visible && IsFolder() ? EVisibility::Visible : EVisibility::Collapsed;
+	return GetNameAreaVisibility() != EVisibility::Visible && IsFolder() ? EVisibility::HitTestInvisible : EVisibility::Collapsed;
 }
 
 EVisibility SAssetTileItem::GetAssetClassLabelVisibility() const
@@ -2567,6 +2575,20 @@ FSlateFontInfo SAssetTileItem::GetThumbnailFont() const
 	}
 }
 
+const FSlateBrush* SAssetTileItem::GetFolderSlotBackgroundImage() const
+{
+	const bool bIsSelected = IsSelected.IsBound() ? IsSelected.Execute() : false;
+	const bool bIsHoveredOrDraggedOver = IsHovered() || bDraggedOver;
+
+	if (bIsSelected || bIsHoveredOrDraggedOver)
+	{
+		static const FLazyName SelectedOrHovered("ContentBrowser.AssetTileItem.FolderAreaBackground"); // Panel
+		return FAppStyle::Get().GetBrush(SelectedOrHovered);
+	}
+
+	return FStyleDefaults::GetNoBrush();
+}
+
 const FSlateBrush* SAssetTileItem::GetFolderBackgroundImage() const
 {
 	const bool bIsSelected = IsSelected.IsBound() ? IsSelected.Execute() : false;
diff --git a/Engine/Source/Editor/ContentBrowser/Public/AssetViewWidgets.h b/Engine/Source/Editor/ContentBrowser/Public/AssetViewWidgets.h
--- a/Engine/Source/Editor/ContentBrowser/Public/AssetViewWidgets.h
+++ b/Engine/Source/Editor/ContentBrowser/Public/AssetViewWidgets.h
@@ -739,6 +739,7 @@ protected:
 	/** Gets the max size of the name area for an asset */
 	FOptionalSize GetMaxAssetNameHeight() const;
 
+	const FSlateBrush* GetFolderSlotBackgroundImage() const;
 	const FSlateBrush* GetFolderBackgroundImage() const;
 	const FSlateBrush* GetFolderBackgroundShadowImage() const;
 
diff --git a/Engine/Source/Editor/EditorStyle/Private/StarshipStyle.cpp b/Engine/Source/Editor/EditorStyle/Private/StarshipStyle.cpp
--- a/Engine/Source/Editor/EditorStyle/Private/StarshipStyle.cpp
+++ b/Engine/Source/Editor/EditorStyle/Private/StarshipStyle.cpp
@@ -6828,9 +6828,12 @@ void FStarshipEditorStyle::FStyle::SetupContentBrowserStyle()
 		Set("ContentBrowser.ListViewCodeFolderIcon", new IMAGE_BRUSH_SVG("Starship/ContentBrowser/FolderLargeCode", Icon64x64));
 		Set("ContentBrowser.ListViewPluginFolderIcon", new IMAGE_BRUSH_SVG("Starship/ContentBrowser/FolderLargePlugin", Icon64x64));
 
-		Set("ContentBrowser.AssetTileItem.FolderAreaHoveredBackground", new FSlateRoundedBoxBrush(FStyleColors::Secondary, 4.0f));
-		Set("ContentBrowser.AssetTileItem.FolderAreaSelectedBackground", new FSlateRoundedBoxBrush(FStyleColors::Secondary, 4.0f));
-		Set("ContentBrowser.AssetTileItem.FolderAreaSelectedHoverBackground", new FSlateRoundedBoxBrush(FStyleColors::Secondary, 4.0f));
+		// Folder Tile Item Border
+		Set("ContentBrowser.AssetTileItem.FolderAreaBackground", new FSlateRoundedBoxBrush(FStyleColors::Panel, 4.0f));
+		Set("ContentBrowser.AssetTileItem.FolderAreaHoveredBackground", new FSlateRoundedBoxBrush(FStyleColors::Hover, 4.0f));
+		Set("ContentBrowser.AssetTileItem.FolderAreaSelectedBackground", new FSlateRoundedBoxBrush(FStyleColors::Primary, 4.0f));
+		Set("ContentBrowser.AssetTileItem.FolderAreaSelectedHoverBackground", new FSlateRoundedBoxBrush(FStyleColors::PrimaryHover, 4.0f));
+
 		Set("ContentBrowser.AssetTileItem.ThumbnailAreaBackground", new FSlateRoundedBoxBrush(FStyleColors::Recessed, 4.0f));
 		Set("ContentBrowser.AssetTileItem.NameAreaBackground", new FSlateRoundedBoxBrush(FStyleColors::Secondary, FVector4(0.0f, 0.0f, 4.0f, 4.0f)));
 		Set("ContentBrowser.AssetTileItem.NameAreaHoverBackground", new FSlateRoundedBoxBrush(FStyleColors::Hover, FVector4(0.0f, 0.0f, 4.0f, 4.0f)));
