[ContentBrowser] Re-added missing source panel toggle in new style

diff --git a/Engine/Source/Editor/ContentBrowser/Private/SContentBrowser.cpp b/Engine/Source/Editor/ContentBrowser/Private/SContentBrowser.cpp
--- a/Engine/Source/Editor/ContentBrowser/Private/SContentBrowser.cpp
+++ b/Engine/Source/Editor/ContentBrowser/Private/SContentBrowser.cpp
@@ -1230,29 +1230,26 @@ void SContentBrowser::ExtendViewOptionsMenu(const FContentBrowserConfig* Config)
 						}
 					}
 
-					if (!UE::Editor::ContentBrowser::IsNewStyleEnabled())
+					if (bShowSourcesView)
 					{
-						if (bShowSourcesView)
-						{
-							FToolMenuSection& Section = InMenu->FindOrAddSection(UE::Editor::ContentBrowser::IsNewStyleEnabled() ? "Show" : "View");
-							Section.AddMenuEntry(
-								"ToggleSources",
-								UE::Editor::ContentBrowser::IsNewStyleEnabled()
-								? LOCTEXT("ToggleSourcesView_NewStyle", "Sources Panel")
-								: LOCTEXT("ToggleSourcesView", "Show Sources Panel"),
-								LOCTEXT("ToggleSourcesView_Tooltip", "Show or hide the sources panel"),
-								TAttribute<FSlateIcon>(),
-								FUIAction(
-									FExecuteAction::CreateLambda([ContentBrowser = Context->OwningContentBrowser]() { ContentBrowser.Pin()->SourcesViewExpandClicked(); }),
-									FCanExecuteAction(),
-									FIsActionChecked::CreateLambda([ContentBrowser = Context->OwningContentBrowser]() { return ContentBrowser.Pin()->bSourcesViewExpanded; })),
-								EUserInterfaceActionType::ToggleButton
-							)
-							.InsertPosition =
-								UE::Editor::ContentBrowser::IsNewStyleEnabled()
-								? FToolMenuInsert(NAME_None, EToolMenuInsertType::First)
-								: FToolMenuInsert();
-						}
+						FToolMenuSection& Section = InMenu->FindOrAddSection(UE::Editor::ContentBrowser::IsNewStyleEnabled() ? "Show" : "View");
+						Section.AddMenuEntry(
+							"ToggleSources",
+							UE::Editor::ContentBrowser::IsNewStyleEnabled()
+							? LOCTEXT("ToggleSourcesView_NewStyle", "Sources Panel")
+							: LOCTEXT("ToggleSourcesView", "Show Sources Panel"),
+							LOCTEXT("ToggleSourcesView_Tooltip", "Show or hide the sources panel"),
+							TAttribute<FSlateIcon>(),
+							FUIAction(
+								FExecuteAction::CreateLambda([ContentBrowser = Context->OwningContentBrowser]() { ContentBrowser.Pin()->SourcesViewExpandClicked(); }),
+								FCanExecuteAction(),
+								FIsActionChecked::CreateLambda([ContentBrowser = Context->OwningContentBrowser]() { return ContentBrowser.Pin()->bSourcesViewExpanded; })),
+							EUserInterfaceActionType::ToggleButton
+						)
+						.InsertPosition =
+							UE::Editor::ContentBrowser::IsNewStyleEnabled()
+							? FToolMenuInsert(NAME_None, EToolMenuInsertType::First)
+							: FToolMenuInsert();
 					}
 				}
 			}
