[ContentBrowser] Update Drop widget based on design

diff --git a/Engine/Source/Developer/CollectionManager/Private/CollectionContainer.cpp b/Engine/Source/Developer/CollectionManager/Private/CollectionContainer.cpp
--- a/Engine/Source/Developer/CollectionManager/Private/CollectionContainer.cpp
+++ b/Engine/Source/Developer/CollectionManager/Private/CollectionContainer.cpp
@@ -1546,7 +1546,7 @@ bool FCollectionContainer::AddToCollection(FName CollectionName, ECollectionShar
 			// Failed to add, all of the objects were already in the collection
 			if (OutError)
 			{
-				*OutError = LOCTEXT("Error_AlreadyInCollection", "All of the assets were already in the collection.");
+				*OutError = FText::Format(LOCTEXT("Error_AlreadyInCollection", "The selected {0}|plural(one=item,other=items) {0}|plural(one=has,other=have) already been added to '{1}'"), ObjectPaths.Num(), FText::FromName(CollectionName));
 			}
 			return false;
 		}
diff --git a/Engine/Source/Editor/ContentBrowser/Private/ContentBrowserStyle.cpp b/Engine/Source/Editor/ContentBrowser/Private/ContentBrowserStyle.cpp
--- a/Engine/Source/Editor/ContentBrowser/Private/ContentBrowserStyle.cpp
+++ b/Engine/Source/Editor/ContentBrowser/Private/ContentBrowserStyle.cpp
@@ -301,7 +301,8 @@ namespace UE::ContentBrowser::Private
 		Set("ContentBrowser.ContentDirty", new IMAGE_BRUSH("ContentBrowser/ContentDirty", CoreStyleConstants::Icon16x16));
 		Set("ContentBrowser.AssetDragDropTooltipBackground", new BOX_BRUSH("Old/Menu_Background", FMargin(8.0f / 64.0f)));
 		Set("ContentBrowser.CollectionTreeDragDropBorder", new BOX_BRUSH("Old/Window/ViewportDebugBorder", 0.8f));
-		Set("ContentBrowser.PopupMessageIcon", new IMAGE_BRUSH("Icons/alert", CoreStyleConstants::Icon32x32));
+		Set("ContentBrowser.PopupMessageIcon.Check", new CORE_IMAGE_BRUSH_SVG("Starship/Common/check-circle-solid", CoreStyleConstants::Icon16x16, FStyleColors::AccentGreen));
+		Set("ContentBrowser.PopupMessageIcon.Info", new CORE_IMAGE_BRUSH_SVG("Starship/Common/info-circle-solid", CoreStyleConstants::Icon16x16, FStyleColors::Foreground));
 		Set("ContentBrowser.NewFolderIcon", new CORE_IMAGE_BRUSH_SVG("Starship/Common/folder-plus", CoreStyleConstants::Icon16x16));
 		Set("ContentBrowser.ShowInExplorer", new IMAGE_BRUSH_SVG("Starship/ContentBrowser/show-in-explorer", CoreStyleConstants::Icon16x16));
 		Set("ContentBrowser.ReferenceViewer", new IMAGE_BRUSH_SVG("Starship/ContentBrowser/reference-viewer", CoreStyleConstants::Icon16x16));
diff --git a/Engine/Source/Editor/ContentBrowser/Private/ContentBrowserUtils.cpp b/Engine/Source/Editor/ContentBrowser/Private/ContentBrowserUtils.cpp
--- a/Engine/Source/Editor/ContentBrowser/Private/ContentBrowserUtils.cpp
+++ b/Engine/Source/Editor/ContentBrowser/Private/ContentBrowserUtils.cpp
@@ -93,6 +93,7 @@ public:
 	SLATE_BEGIN_ARGS( SContentBrowserPopup ){}
 
 		SLATE_ATTRIBUTE( FText, Message )
+		SLATE_ARGUMENT( ContentBrowserUtils::EDisplayMessageType, MessageType )
 
 	SLATE_END_ARGS()
 
@@ -100,6 +101,9 @@ public:
 	BEGIN_SLATE_FUNCTION_BUILD_OPTIMIZATION
 	void Construct( const FArguments& InArgs )
 	{
+		Message = InArgs._Message;
+		MessageType = InArgs._MessageType;
+
 		ChildSlot
 		[
 			SNew(SBorder)
@@ -116,7 +120,7 @@ public:
 				.Padding(0.f, 0.f, 4.f, 0.f)
 				[
 					SNew(SImage)
-					.Image(UE::ContentBrowser::Private::FContentBrowserStyle::Get().GetBrush("ContentBrowser.PopupMessageIcon"))
+					.Image(this, &SContentBrowserPopup::GetDisplayMessageIconBrush)
 				]
 
 				+SHorizontalBox::Slot()
@@ -132,9 +136,38 @@ public:
 	}
 	END_SLATE_FUNCTION_BUILD_OPTIMIZATION
 
-	static void DisplayMessage( const FText& Message, const FSlateRect& ScreenAnchor, TSharedRef<SWidget> ParentContent )
+	const FSlateBrush* GetDisplayMessageIconBrush() const
+	{
+		FName BrushName = FName(TEXT(""));
+		switch (MessageType)
+		{
+			case ContentBrowserUtils::EDisplayMessageType::Successful:
+				BrushName = FName(TEXT("ContentBrowser.PopupMessageIcon.Check"));
+				break;
+
+			case ContentBrowserUtils::EDisplayMessageType::Info:
+				BrushName = FName(TEXT("ContentBrowser.PopupMessageIcon.Info"));
+				break;
+
+			case ContentBrowserUtils::EDisplayMessageType::Warning:
+				BrushName = FName(TEXT("Icons.Warning.Solid"));
+				break;
+
+			case ContentBrowserUtils::EDisplayMessageType::Error:
+				BrushName = FName(TEXT("Icons.Error.Solid"));
+				break;
+
+			default:
+				break;
+		}
+		return UE::ContentBrowser::Private::FContentBrowserStyle::Get().GetBrush(BrushName);
+	}
+
+	static void DisplayMessage( const FText& Message, const FSlateRect& ScreenAnchor, TSharedRef<SWidget> ParentContent, ContentBrowserUtils::EDisplayMessageType InMessageType )
 	{
-		TSharedRef<SContentBrowserPopup> PopupContent = SNew(SContentBrowserPopup) .Message(Message);
+		TSharedRef<SContentBrowserPopup> PopupContent = SNew(SContentBrowserPopup)
+			.Message(Message)
+			.MessageType(InMessageType);
 
 		const FVector2D ScreenLocation = FVector2D(ScreenAnchor.Left, ScreenAnchor.Top);
 		const bool bFocusImmediately = true;
@@ -174,6 +207,8 @@ private:
 	}
 
 private:
+	ContentBrowserUtils::EDisplayMessageType MessageType = ContentBrowserUtils::EDisplayMessageType::Info;
+	TAttribute<FText> Message;
 	TWeakPtr<IMenu> Menu;
 };
 
@@ -306,9 +341,9 @@ private:
 };
 
 
-void ContentBrowserUtils::DisplayMessage(const FText& Message, const FSlateRect& ScreenAnchor, const TSharedRef<SWidget>& ParentContent)
+void ContentBrowserUtils::DisplayMessage(const FText& Message, const FSlateRect& ScreenAnchor, const TSharedRef<SWidget>& ParentContent, EDisplayMessageType InMessageType)
 {
-	SContentBrowserPopup::DisplayMessage(Message, ScreenAnchor, ParentContent);
+	SContentBrowserPopup::DisplayMessage(Message, ScreenAnchor, ParentContent, InMessageType);
 }
 
 void ContentBrowserUtils::DisplayConfirmationPopup(const FText& Message, const FText& YesString, const FText& NoString, const TSharedRef<SWidget>& ParentContent, const FOnClicked& OnYesClicked, const FOnClicked& OnNoClicked)
diff --git a/Engine/Source/Editor/ContentBrowser/Private/ContentBrowserUtils.h b/Engine/Source/Editor/ContentBrowser/Private/ContentBrowserUtils.h
--- a/Engine/Source/Editor/ContentBrowser/Private/ContentBrowserUtils.h
+++ b/Engine/Source/Editor/ContentBrowser/Private/ContentBrowserUtils.h
@@ -34,13 +34,23 @@ enum class EContentBrowserIsFolderVisibleFlags : uint8;
 enum class EContentBrowserItemAttributeFilter : uint8;
 enum class EContentBrowserItemCategoryFilter : uint8;
 
+
 namespace ContentBrowserUtils
 {
 	// Import the functions that were moved into the more common AssetViewUtils namespace
 	using namespace AssetViewUtils;
 
+	// Display Message type, mainly used to decide the icon to use
+	enum class EDisplayMessageType
+	{
+		Successful,
+		Info,
+		Warning,
+		Error
+	};
+
 	/** Displays a modeless message at the specified anchor. It is fine to specify a zero-size anchor, just use the top and left fields */
-	void DisplayMessage(const FText& Message, const FSlateRect& ScreenAnchor, const TSharedRef<SWidget>& ParentContent);
+	void DisplayMessage(const FText& Message, const FSlateRect& ScreenAnchor, const TSharedRef<SWidget>& ParentContent, EDisplayMessageType InMessageType);
 
 	/** Displays a modeless message asking yes or no type question */
 	void DisplayConfirmationPopup(const FText& Message, const FText& YesString, const FText& NoString, const TSharedRef<SWidget>& ParentContent, const FOnClicked& OnYesClicked, const FOnClicked& OnNoClicked = FOnClicked());
diff --git a/Engine/Source/Editor/ContentBrowser/Private/DragDropHandler.cpp b/Engine/Source/Editor/ContentBrowser/Private/DragDropHandler.cpp
--- a/Engine/Source/Editor/ContentBrowser/Private/DragDropHandler.cpp
+++ b/Engine/Source/Editor/ContentBrowser/Private/DragDropHandler.cpp
@@ -309,23 +309,25 @@ void HandleDragDropMoveOrCopy(const FContentBrowserItem& InDropTargetItem, const
 	}
 
 	// Show a message if the move or copy was successful
-	if (NumMovedOrCopiedItems > 0 && InParentWidget)
+	if (InParentWidget)
 	{
-		const FText Message = FText::Format(InMoveOrCopyMsg, NumMovedOrCopiedItems, FText::FromName(InDropTargetPath));
+		const bool bWasSuccessful = NumMovedOrCopiedItems > 0;
+		const ContentBrowserUtils::EDisplayMessageType MessageType = bWasSuccessful ? ContentBrowserUtils::EDisplayMessageType::Successful : ContentBrowserUtils::EDisplayMessageType::Info;
+		const FText Message = FText::Format(InMoveOrCopyMsg, NumMovedOrCopiedItems, InDropTargetItem.GetDisplayName());
 		const FVector2f& CursorPos = FSlateApplication::Get().GetCursorPos();
 		FSlateRect MessageAnchor(CursorPos.X, CursorPos.Y, CursorPos.X, CursorPos.Y);
-		ContentBrowserUtils::DisplayMessage(Message, MessageAnchor, InParentWidget.ToSharedRef());
+		ContentBrowserUtils::DisplayMessage(Message, MessageAnchor, InParentWidget.ToSharedRef(), MessageType);
 	}
 }
 
 void HandleDragDropMove(const FContentBrowserItem& InDropTargetItem, const TArray<FContentBrowserItem>& InDraggedItems, const TSharedPtr<SWidget>& InParentWidget)
 {
-	return HandleDragDropMoveOrCopy(InDropTargetItem, InDraggedItems, InParentWidget, LOCTEXT("ItemsDroppedMove", "{0} {0}|plural(one=item,other=items) moved to '{1}'"), &UContentBrowserDataSource::CanMoveItem, &UContentBrowserDataSource::BulkMoveItems);
+	return HandleDragDropMoveOrCopy(InDropTargetItem, InDraggedItems, InParentWidget, LOCTEXT("ItemsDroppedMove", "{0} {0}|plural(one=item,other=items) moved to '.../{1}'"), &UContentBrowserDataSource::CanMoveItem, &UContentBrowserDataSource::BulkMoveItems);
 }
 
 void HandleDragDropCopy(const FContentBrowserItem& InDropTargetItem, const TArray<FContentBrowserItem>& InDraggedItems, const TSharedPtr<SWidget>& InParentWidget)
 {
-	return HandleDragDropMoveOrCopy(InDropTargetItem, InDraggedItems, InParentWidget, LOCTEXT("ItemsDroppedCopy", "{0} {0}|plural(one=item,other=items) copied to '{1}'"), &UContentBrowserDataSource::CanCopyItem, &UContentBrowserDataSource::BulkCopyItems);
+	return HandleDragDropMoveOrCopy(InDropTargetItem, InDraggedItems, InParentWidget, LOCTEXT("ItemsDroppedCopy", "{0} {0}|plural(one=item,other=items) copied to '.../{1}'"), &UContentBrowserDataSource::CanCopyItem, &UContentBrowserDataSource::BulkCopyItems);
 }
 
 bool HandleDragDropOnItem(const FContentBrowserItem& InItem, const FDragDropEvent& InDragDropEvent, const TSharedRef<SWidget>& InParentWidget)
diff --git a/Engine/Source/Editor/ContentBrowser/Private/SAssetView.cpp b/Engine/Source/Editor/ContentBrowser/Private/SAssetView.cpp
--- a/Engine/Source/Editor/ContentBrowser/Private/SAssetView.cpp
+++ b/Engine/Source/Editor/ContentBrowser/Private/SAssetView.cpp
@@ -5410,7 +5410,7 @@ void SAssetView::AssetRenameCommit(const TSharedPtr<FAssetViewItem>& Item, const
 	else if (!ErrorMessage.IsEmpty())
 	{
 		// Prompt the user with the reason the rename/creation failed
-		ContentBrowserUtils::DisplayMessage(ErrorMessage, MessageAnchor, SharedThis(this));
+		ContentBrowserUtils::DisplayMessage(ErrorMessage, MessageAnchor, SharedThis(this), ContentBrowserUtils::EDisplayMessageType::Error);
 	}
 
 	RenamingAsset.Reset();
diff --git a/Engine/Source/Editor/ContentBrowser/Private/SCollectionView.cpp b/Engine/Source/Editor/ContentBrowser/Private/SCollectionView.cpp
--- a/Engine/Source/Editor/ContentBrowser/Private/SCollectionView.cpp
+++ b/Engine/Source/Editor/ContentBrowser/Private/SCollectionView.cpp
@@ -1078,7 +1078,8 @@ void SCollectionView::DeleteCollectionItems( const TArray<TSharedPtr<FCollection
 			ContentBrowserUtils::DisplayMessage(
 				FText::Format( LOCTEXT("CollectionDestroyFailed", "Failed to destroy collection. {0}"), Error),
 				MessageAnchor,
-				CollectionTreePtr.ToSharedRef()
+				CollectionTreePtr.ToSharedRef(),
+				ContentBrowserUtils::EDisplayMessageType::Error
 				);
 		}
 	}
@@ -1283,7 +1284,7 @@ FReply SCollectionView::HandleDragDropOnCollectionTree(const FGeometry& Geometry
 					NAME_None, ECollectionShareType::CST_All, &Error
 					))
 			{
-				ContentBrowserUtils::DisplayMessage(Error, Geometry.GetLayoutBoundingRect(), SharedThis(this));
+				ContentBrowserUtils::DisplayMessage(Error, Geometry.GetLayoutBoundingRect(), SharedThis(this), ContentBrowserUtils::EDisplayMessageType::Error);
 			}
 		}
 
@@ -1384,7 +1385,7 @@ FReply SCollectionView::HandleDragDropOnCollectionItem(TSharedRef<FCollectionIte
 					CollectionItem->CollectionName, CollectionItem->CollectionType, &Error
 					))
 			{
-				ContentBrowserUtils::DisplayMessage(Error, Geometry.GetLayoutBoundingRect(), SharedThis(this));
+				ContentBrowserUtils::DisplayMessage(Error, Geometry.GetLayoutBoundingRect(), SharedThis(this), ContentBrowserUtils::EDisplayMessageType::Error);
 			}
 		}
 
@@ -1405,22 +1406,32 @@ FReply SCollectionView::HandleDragDropOnCollectionItem(TSharedRef<FCollectionIte
 		const double BeginTimeSec = FPlatformTime::Seconds();
 		int32 NumAdded = 0;
 		FText Message;
+		ContentBrowserUtils::EDisplayMessageType MessageType = ContentBrowserUtils::EDisplayMessageType::Info;
 		if (CollectionItem->CollectionContainer->AddToCollection(
 				CollectionItem->CollectionName, CollectionItem->CollectionType, ObjectPaths, &NumAdded, &Message))
 		{
-			if (DroppedAssets.Num() == 1)
+			if (UE::Editor::ContentBrowser::IsNewStyleEnabled())
 			{
-				FFormatNamedArguments Args;
-				Args.Add(TEXT("AssetName"), FText::FromName(DroppedAssets[0].AssetName));
-				Args.Add(TEXT("CollectionName"), FText::FromName(CollectionItem->CollectionName));
-				Message = FText::Format(LOCTEXT("CollectionAssetAdded", "Added {AssetName} to {CollectionName}"), Args);
+				Message = FText::Format(LOCTEXT("AddingToCollection", "{0} {0}|plural(one=item,other=items) added to '{1}'"), NumAdded, FText::FromName(CollectionItem->CollectionName));
+				MessageType = ContentBrowserUtils::EDisplayMessageType::Successful;
 			}
 			else
 			{
-				FFormatNamedArguments Args;
-				Args.Add(TEXT("Number"), NumAdded);
-				Args.Add(TEXT("CollectionName"), FText::FromName(CollectionItem->CollectionName));
-				Message = FText::Format(LOCTEXT("CollectionAssetsAdded", "Added {Number} asset(s) to {CollectionName}"), Args);
+				if (DroppedAssets.Num() == 1)
+				{
+					FFormatNamedArguments Args;
+					Args.Add(TEXT("AssetName"), FText::FromName(DroppedAssets[0].AssetName));
+					Args.Add(TEXT("CollectionName"), FText::FromName(CollectionItem->CollectionName));
+					Message = FText::Format(LOCTEXT("CollectionAssetAdded", "Added {AssetName} to {CollectionName}"), Args);
+				}
+				else
+				{
+					FFormatNamedArguments Args;
+					Args.Add(TEXT("Number"), NumAdded);
+					Args.Add(TEXT("CollectionName"), FText::FromName(CollectionItem->CollectionName));
+					Message = FText::Format(LOCTEXT("CollectionAssetsAdded", "Added {Number} asset(s) to {CollectionName}"), Args);
+				}
+				MessageType = ContentBrowserUtils::EDisplayMessageType::Info;
 			}
 
 			const double DurationSec = FPlatformTime::Seconds() - BeginTimeSec;
@@ -1436,7 +1447,7 @@ FReply SCollectionView::HandleDragDropOnCollectionItem(TSharedRef<FCollectionIte
 		}
 
 		// Added items to the collection or failed. Either way, display the message.
-		ContentBrowserUtils::DisplayMessage(Message, Geometry.GetLayoutBoundingRect(), SharedThis(this));
+		ContentBrowserUtils::DisplayMessage(Message, Geometry.GetLayoutBoundingRect(), SharedThis(this), MessageType);
 
 		return FReply::Handled();
 	}
diff --git a/Engine/Source/Editor/ContentBrowser/Private/SPathView.cpp b/Engine/Source/Editor/ContentBrowser/Private/SPathView.cpp
--- a/Engine/Source/Editor/ContentBrowser/Private/SPathView.cpp
+++ b/Engine/Source/Editor/ContentBrowser/Private/SPathView.cpp
@@ -2390,7 +2390,7 @@ void SPathView::FolderNameChanged(const TSharedPtr<FTreeItem>& TreeItem,
 	{
 		// Display the reason why the folder was invalid
 		FSlateRect MessageAnchor(MessageLocation.X, MessageLocation.Y, MessageLocation.X, MessageLocation.Y);
-		ContentBrowserUtils::DisplayMessage(ErrorMessage, MessageAnchor, SharedThis(this));
+		ContentBrowserUtils::DisplayMessage(ErrorMessage, MessageAnchor, SharedThis(this), ContentBrowserUtils::EDisplayMessageType::Error);
 	}
 }
 
diff --git a/Engine/Source/Editor/ContentBrowser/Private/SourcesViewWidgets.cpp b/Engine/Source/Editor/ContentBrowser/Private/SourcesViewWidgets.cpp
--- a/Engine/Source/Editor/ContentBrowser/Private/SourcesViewWidgets.cpp
+++ b/Engine/Source/Editor/ContentBrowser/Private/SourcesViewWidgets.cpp
@@ -548,7 +548,7 @@ void SCollectionTreeItem::HandleNameCommitted( const FText& NewText, ETextCommit
 				if ( !OnNameChangeCommit.Execute(CollectionItemPtr, NewText.ToString(), bIsCommitted, WarningMessage) && ParentWidget.IsValid() && bIsCommitted )
 				{
 					// Failed to rename/create a collection, display a warning.
-					ContentBrowserUtils::DisplayMessage(WarningMessage, CachedGeometry.GetLayoutBoundingRect(), ParentWidget.ToSharedRef());
+					ContentBrowserUtils::DisplayMessage(WarningMessage, CachedGeometry.GetLayoutBoundingRect(), ParentWidget.ToSharedRef(), ContentBrowserUtils::EDisplayMessageType::Error);
 				}
 			}				
 		}
diff --git a/Engine/Source/Editor/ContentBrowserData/Public/ContentBrowserDataDragDropOp.h b/Engine/Source/Editor/ContentBrowserData/Public/ContentBrowserDataDragDropOp.h
--- a/Engine/Source/Editor/ContentBrowserData/Public/ContentBrowserDataDragDropOp.h
+++ b/Engine/Source/Editor/ContentBrowserData/Public/ContentBrowserDataDragDropOp.h
@@ -74,6 +74,9 @@ private:
 	/** Return the folder widget or nullptr if the BrushName is not valid */
 	TSharedPtr<SWidget> GetFolderWidget() const;
 
+	/** Return the folder widget for DragAndDrop or nullptr if there is no basic folder widget  */
+	TSharedPtr<SWidget> GetFolderWidgetDragAndDrop() const;
+
 private:
 	FName FolderBrushName;
 	FName FolderShadowBrushName;
diff --git a/Engine/Source/Editor/EditorStyle/Private/StarshipStyle.cpp b/Engine/Source/Editor/EditorStyle/Private/StarshipStyle.cpp
--- a/Engine/Source/Editor/EditorStyle/Private/StarshipStyle.cpp
+++ b/Engine/Source/Editor/EditorStyle/Private/StarshipStyle.cpp
@@ -1245,6 +1245,7 @@ void FStarshipEditorStyle::FStyle::SetupGeneralStyles()
 		Set( "AssetThumbnail.AssetBackground", new FSlateColorBrush(FStyleColors::Recessed));
 		Set( "AssetThumbnail.ClassBackground", new IMAGE_BRUSH( "Common/ClassBackground_64x", Icon64x64, FLinearColor(0.75f, 0.75f, 0.75f, 1.0f) ) );
 		Set( "AssetThumbnail.Font", DEFAULT_FONT( "Regular", 10 ) );
+		Set( "AssetThumbnail.StatusOverflowFont", DEFAULT_FONT( "Bold", 9 ) );
 		Set( "AssetThumbnail.FontSmall", DEFAULT_FONT( "Regular", 7 ) );
 		Set( "AssetThumbnail.ColorAndOpacity", FLinearColor(1.75f, 1.75f, 1.75f, 1.f) );
 		Set( "AssetThumbnail.HintFont", DEFAULT_FONT( "Regular", 8 ) );
@@ -6977,6 +6978,7 @@ void FStarshipEditorStyle::FStyle::SetupContentBrowserStyle()
 		Set( "ContentBrowser.ContentDirty", new IMAGE_BRUSH_SVG("Starship/Common/DirtyBadge", Icon16x16) );
 		Set( "ContentBrowser.AssetDragDropTooltipBackground", new BOX_BRUSH( "Old/Menu_Background", FMargin(8.0f/64.0f) ) );
 		Set( "ContentBrowser.DragDropBackground", new FSlateRoundedBoxBrush(FStyleColors::Input, 0.f, FStyleColors::InputOutline, 1.f));
+		Set( "ContentBrowser.DragDropAssetNumbersBorder", new FSlateRoundedBoxBrush(FStyleColors::Panel, 2.f));
 		Set( "ContentBrowser.ThumbnailDragDropBackground", new FSlateRoundedBoxBrush(FStyleColors::Transparent, 0.f, FStyleColors::InputOutline, 1.f));
 		Set( "ContentBrowser.CollectionTreeDragDropBorder", new BOX_BRUSH( "Old/Window/ViewportDebugBorder", 0.8f ) );
 		Set( "ContentBrowser.PopupMessageIcon", new IMAGE_BRUSH( "Icons/alert", Icon32x32) );
diff --git a/Engine/Source/Editor/UnrealEd/Private/AssetThumbnail.cpp b/Engine/Source/Editor/UnrealEd/Private/AssetThumbnail.cpp
--- a/Engine/Source/Editor/UnrealEd/Private/AssetThumbnail.cpp
+++ b/Engine/Source/Editor/UnrealEd/Private/AssetThumbnail.cpp
@@ -969,7 +969,7 @@ private:
 			.Visibility(this, &SAssetThumbnail::GetStatusOverflowVisibility)
 			[
 				SNew(STextBlock)
-				.Font(FSlateFontInfo(FCoreStyle::GetDefaultFont(), 9, "Bold"))
+				.Font(FAppStyle::GetFontStyle(Style, ".StatusOverflowFont"))
 				.Text(this, &SAssetThumbnail::GetStatusOverflowText)
 			];
 	}
diff --git a/Engine/Source/Editor/UnrealEd/Private/DragAndDrop/AssetDragDropOp.cpp b/Engine/Source/Editor/UnrealEd/Private/DragAndDrop/AssetDragDropOp.cpp
--- a/Engine/Source/Editor/UnrealEd/Private/DragAndDrop/AssetDragDropOp.cpp
+++ b/Engine/Source/Editor/UnrealEd/Private/DragAndDrop/AssetDragDropOp.cpp
@@ -131,6 +131,9 @@ TSharedPtr<SWidget> FAssetDragDropOp::GetDefaultDecorator() const
 
 	if (UE::Editor::ContentBrowser::IsNewStyleEnabled())
 	{
+		static const FMargin AssetNumberSlotMargin = FMargin(6, 0, 0, 6);
+		static constexpr int32 OverlayMinOrFixedSize = 16.f;
+
 		return SNew(SBorder)
 				.BorderImage(FAppStyle::GetBrush("ContentBrowser.DragDropBackground"))
 				.Padding(0.f)
@@ -154,17 +157,24 @@ TSharedPtr<SWidget> FAssetDragDropOp::GetDefaultDecorator() const
 							]
 
 							+SOverlay::Slot()
-							.HAlign(HAlign_Center)
-							.VAlign(VAlign_Top)
-							.Padding(FMargin(0, 4, 0, 0))
+							.HAlign(HAlign_Left)
+							.VAlign(VAlign_Bottom)
+							.Padding(AssetNumberSlotMargin)
 							[
-								SNew(SBorder)
-								.BorderImage(FAppStyle::GetBrush("Menu.Background"))
-								.Visibility(TotalCount > 1 ? EVisibility::Visible : EVisibility::Collapsed)
-								.Content()
+								SNew(SBox)
+								.MinDesiredWidth(OverlayMinOrFixedSize)
+								.HeightOverride(OverlayMinOrFixedSize)
 								[
-									SNew(STextBlock)
-									.Text(FText::AsNumber(TotalCount))
+									SNew(SBorder)
+									.BorderImage(FAppStyle::GetBrush("ContentBrowser.DragDropAssetNumbersBorder"))
+									.Visibility(TotalCount > 1 ? EVisibility::Visible : EVisibility::Collapsed)
+									.HAlign(HAlign_Center)
+									.VAlign(VAlign_Center)
+									[
+										SNew(STextBlock)
+										.Font(FAppStyle::GetFontStyle("AssetThumbnail.StatusOverflowFont"))
+										.Text(FText::Format(NSLOCTEXT("AssetDragDropOp", "AssetDragAndDropNumber", "+{0}"), FText::AsNumber(TotalCount)))
+									]
 								]
 							]
 						]
