[Content Browser] Horizontal/Vertical Filter Panel mode
https://jira.it.epicgames.com/browse/UE-240577

- Select Add Filter Menu
- Within the menu, select the Common Filters group
- From the Content Browser, select Settings > Filter Display > Horizontal
- The active filters horizontally align next to the Search All field

Check whether Horizontal/Vertical Filter Panel mode may be activated properly via Settings Menu of Content Browser.

diff --git a/Engine/Source/Editor/ContentBrowser/Private/ContentBrowserSingleton.cpp b/Engine/Source/Editor/ContentBrowser/Private/ContentBrowserSingleton.cpp
--- a/Engine/Source/Editor/ContentBrowser/Private/ContentBrowserSingleton.cpp
+++ b/Engine/Source/Editor/ContentBrowser/Private/ContentBrowserSingleton.cpp
@@ -1467,6 +1467,51 @@ void FContentBrowserSingleton::ToggleFolderFavorite(const TArray<FString>& Folde
 	}
 }
 
+void FContentBrowserSingleton::SetFilterLayout(const EFilterBarLayout InFilterBarLayout)
+{
+	if (!HasPrimaryContentBrowser())
+	{
+		ChooseNewPrimaryBrowser();
+	}
+
+	if (PrimaryContentBrowser.IsValid())
+	{
+		PrimaryContentBrowser.Pin()->SetFilterLayout(InFilterBarLayout);
+	}
+}
+
+EFilterBarLayout FContentBrowserSingleton::GetFilterLayout()
+{
+	EFilterBarLayout FilterBarLayout = EFilterBarLayout::Vertical;
+
+	if (!HasPrimaryContentBrowser())
+	{
+		ChooseNewPrimaryBrowser();
+	}
+
+	if (PrimaryContentBrowser.IsValid())
+	{
+		FilterBarLayout = PrimaryContentBrowser.Pin()->GetFilterLayout();
+	}
+
+	return FilterBarLayout;
+}
+
+TSharedPtr<SWidget> FContentBrowserSingleton::GetActiveFilterContainer()
+{
+	if (!HasPrimaryContentBrowser())
+	{
+		ChooseNewPrimaryBrowser();
+	}
+
+	if (PrimaryContentBrowser.IsValid())
+	{
+		return PrimaryContentBrowser.Pin()->GetActiveFilterContainer();
+	}
+
+	return nullptr;
+}
+
 void FContentBrowserSingleton::RegisterContentSourceFactory(const FName& InContentSourceName, const FContentSourceFactory& InContentSourceFactory)
 {
 	ContentSourceFactories.Emplace(InContentSourceName, InContentSourceFactory);
diff --git a/Engine/Source/Editor/ContentBrowser/Private/ContentBrowserSingleton.h b/Engine/Source/Editor/ContentBrowser/Private/ContentBrowserSingleton.h
--- a/Engine/Source/Editor/ContentBrowser/Private/ContentBrowserSingleton.h
+++ b/Engine/Source/Editor/ContentBrowser/Private/ContentBrowserSingleton.h
@@ -125,6 +125,9 @@ PRAGMA_ENABLE_DEPRECATION_WARNINGS
 	virtual bool IsAssetViewDoneFiltering() override;
 	virtual void SetSearchText(const FText& NewText) override;
 	virtual void ToggleFolderFavorite(const TArray<FString>& FolderPaths) override;
+	virtual void SetFilterLayout(EFilterBarLayout InFilterBarLayout) override;
+	virtual TSharedPtr<SWidget> GetActiveFilterContainer() override;
+	virtual EFilterBarLayout GetFilterLayout() override;
 	virtual void DockContentBrowserDrawer() override;
 	virtual void RegisterContentSourceFactory(const FName& InContentSourceName, const FContentSourceFactory& InContentSourceFactory) override;
 	virtual void UnregisterContentSourceFactory(const FName& InContentSourceName) override;
diff --git a/Engine/Source/Editor/ContentBrowser/Private/SContentBrowser.cpp b/Engine/Source/Editor/ContentBrowser/Private/SContentBrowser.cpp
--- a/Engine/Source/Editor/ContentBrowser/Private/SContentBrowser.cpp
+++ b/Engine/Source/Editor/ContentBrowser/Private/SContentBrowser.cpp
@@ -518,6 +518,44 @@ void SContentBrowser::ToggleFolderFavorite(const TArray<FString>& FolderPaths)
 	SaveAndShowNewFolderFavorites(FolderPathsAdded);
 }
 
+void SContentBrowser::SetFilterLayout(const EFilterBarLayout InFilterBarLayout) const
+{
+	if (const TSharedPtr<SFilterList>& FilterBar = LegacyContentSourceWidgets->FilterListPtr; FilterBar && LegacyContentSourceWidgets->AssetViewPtr)
+	{
+		FilterBar->SetFilterLayout(InFilterBarLayout);
+		LegacyContentSourceWidgets->AssetViewPtr->SetFilterBar(FilterBar);
+	}
+	else
+	{
+		UE_LOG(LogContentBrowser, Warning, TEXT("SetFilterLayout failed: %s is invalid. FilterListPtr: %s, AssetViewPtr: %s"),
+		       TEXT("SContentBrowser::SetFilterLayout"),
+		       LegacyContentSourceWidgets->FilterListPtr.IsValid() ? TEXT("Valid") : TEXT("Invalid"),
+		       LegacyContentSourceWidgets->AssetViewPtr.IsValid() ? TEXT("Valid") : TEXT("Invalid"));
+	}
+}
+
+EFilterBarLayout SContentBrowser::GetFilterLayout() const
+{
+	if (const TSharedPtr<SFilterList>& FilterBar = LegacyContentSourceWidgets->FilterListPtr)
+	{
+		return FilterBar->GetFilterLayout();
+	}
+
+	UE_LOG(LogContentBrowser, Warning, TEXT("GetFilterLayout: FilterListPtr is invalid, returning default layout."));
+	return EFilterBarLayout::Vertical;
+}
+
+TSharedPtr<SWidget> SContentBrowser::GetActiveFilterContainer() const
+{
+	if (const TSharedPtr<SFilterList>& FilterBar = LegacyContentSourceWidgets->FilterListPtr)
+	{
+		return FilterBar->GetActiveFilterContainer();
+	}
+
+	UE_LOG(LogContentBrowser, Warning, TEXT("GetFilterLayout: FilterListPtr is invalid, returning nullptr."));
+	return nullptr;
+}
+
 void SContentBrowser::SaveAndShowNewFolderFavorites(const TArray<FString>& FolderPaths)
 {
 	// If the legacy content source isn't active - the settings will get updated when it is made active
diff --git a/Engine/Source/Editor/ContentBrowser/Private/SContentBrowser.h b/Engine/Source/Editor/ContentBrowser/Private/SContentBrowser.h
--- a/Engine/Source/Editor/ContentBrowser/Private/SContentBrowser.h
+++ b/Engine/Source/Editor/ContentBrowser/Private/SContentBrowser.h
@@ -271,6 +271,15 @@ public:
 	/** Toggles the favorite status of an array of folders*/
 	void ToggleFolderFavorite(const TArray<FString>& FolderPaths);
 
+	/** Sets the Content Browser filter layout to either vertical or horizontal */
+	void SetFilterLayout(EFilterBarLayout InFilterBarLayout) const;
+
+	/** Gets the current Content Browser filter layout */
+	EFilterBarLayout GetFilterLayout() const;
+
+	/** Gets a pointer to the active filter container widget (Horizontal or Vertical), nullptr if neither is active */
+	TSharedPtr<SWidget> GetActiveFilterContainer() const;
+
 private:
 	class FCollectionSource;
 
diff --git a/Engine/Source/Editor/ContentBrowser/Private/SFilterList.cpp b/Engine/Source/Editor/ContentBrowser/Private/SFilterList.cpp
--- a/Engine/Source/Editor/ContentBrowser/Private/SFilterList.cpp
+++ b/Engine/Source/Editor/ContentBrowser/Private/SFilterList.cpp
@@ -648,6 +648,21 @@ void SFilterList::ToggleFilterPathsRecursively()
 	SetFilterPathsRecursively(!bFilterPathsRecursively);
 }
 
+TSharedPtr<SWidget> SFilterList::GetActiveFilterContainer() const
+{
+	if (HorizontalFilterBox.IsValid() && HorizontalFilterBox->GetChildren()->Num() > 0)
+	{
+		return HorizontalFilterBox;
+	}
+	if (VerticalFilterBox.IsValid() && VerticalFilterBox->GetChildren()->Num() > 0)
+	{
+		return VerticalFilterBox;
+	}
+
+	// If neither has children, return nullptr.
+	return nullptr;
+}
+
 void SFilterList::OnDeleteCustomTextFilter(const TSharedPtr<ICustomTextFilter<FAssetFilterType>> InFilter)
 {
 	SAssetFilterBar<FAssetFilterType>::OnDeleteCustomTextFilter(InFilter);
diff --git a/Engine/Source/Editor/ContentBrowser/Private/SFilterList.h b/Engine/Source/Editor/ContentBrowser/Private/SFilterList.h
--- a/Engine/Source/Editor/ContentBrowser/Private/SFilterList.h
+++ b/Engine/Source/Editor/ContentBrowser/Private/SFilterList.h
@@ -177,6 +177,9 @@ public:
 	/** Toggle whether the filters are being applied to paths recursively */
 	void ToggleFilterPathsRecursively();
 
+	/** Returns a pointer to the active filter container widget (Horizontal or Vertical), nullptr if neither is active */
+	TSharedPtr<SWidget> GetActiveFilterContainer() const;
+
 protected:
 	
 	/** Handler for when the add filter button was clicked */
diff --git a/Engine/Source/Editor/ContentBrowser/Public/IContentBrowserSingleton.h b/Engine/Source/Editor/ContentBrowser/Public/IContentBrowserSingleton.h
--- a/Engine/Source/Editor/ContentBrowser/Public/IContentBrowserSingleton.h
+++ b/Engine/Source/Editor/ContentBrowser/Public/IContentBrowserSingleton.h
@@ -20,6 +20,7 @@
 #include "AssetThumbnail.h"
 #include "ContentBrowserItemPath.h"
 #include "Misc/NamePermissionList.h"
+#include "Filters/SBasicFilterBar.h"
 
 class FViewport;
 class IPlugin;
@@ -826,6 +827,15 @@ public:
 	/** Toggle the Favorite status of the given folder list */
 	virtual void ToggleFolderFavorite(const TArray<FString>& FolderPaths) = 0;
 
+	/** Sets the Content Browser filter layout to either vertical or horizontal */
+	virtual void SetFilterLayout(EFilterBarLayout InFilterBarLayout) = 0;
+
+	/** Gets the current Content Browser filter layout */
+	virtual EFilterBarLayout GetFilterLayout() = 0;
+
+	/** Gets a pointer to the active filter container widget (Horizontal or Vertical), nullptr if neither is active */
+	virtual TSharedPtr<SWidget> GetActiveFilterContainer() = 0;
+
 	/** Docks the current content browser drawer as a tabbed content browser in a layout */
 	virtual void DockContentBrowserDrawer() = 0;
 
