[Utility Rig Crash Fix] ShowRedirector_Lambda of the AssetPicker was causing the filters list to go out of scope and crash since it utilized SharedRefs

#okforfithub public

diff --git a/Engine/Source/Editor/ContentBrowser/Private/SAssetPicker.cpp b/Engine/Source/Editor/ContentBrowser/Private/SAssetPicker.cpp
--- a/Engine/Source/Editor/ContentBrowser/Private/SAssetPicker.cpp
+++ b/Engine/Source/Editor/ContentBrowser/Private/SAssetPicker.cpp
@@ -121,7 +121,7 @@ void SAssetPicker::Construct( const FArguments& InArgs )
 
 	for (auto DelegateIt = InArgs._AssetPickerConfig.GetCurrentSelectionDelegates.CreateConstIterator(); DelegateIt; ++DelegateIt)
 	{
-		if ((*DelegateIt) != NULL)
+		if ((*DelegateIt) != nullptr)
 		{
 			(**DelegateIt) = FGetCurrentSelectionDelegate::CreateSP(this, &SAssetPicker::GetCurrentSelection);
 		}
@@ -129,7 +129,7 @@ void SAssetPicker::Construct( const FArguments& InArgs )
 
 	for(auto DelegateIt = InArgs._AssetPickerConfig.SyncToAssetsDelegates.CreateConstIterator(); DelegateIt; ++DelegateIt)
 	{
-		if((*DelegateIt) != NULL)
+		if((*DelegateIt) != nullptr)
 		{
 			(**DelegateIt) = FSyncToAssetsDelegate::CreateSP(this, &SAssetPicker::SyncToAssets);
 		}
@@ -137,7 +137,7 @@ void SAssetPicker::Construct( const FArguments& InArgs )
 
 	for (auto DelegateIt = InArgs._AssetPickerConfig.SetFilterDelegates.CreateConstIterator(); DelegateIt; ++DelegateIt)
 	{
-		if ((*DelegateIt) != NULL)
+		if ((*DelegateIt) != nullptr)
 		{
 			(**DelegateIt) = FSetARFilterDelegate::CreateSP(this, &SAssetPicker::SetNewBackendFilter);
 		}
@@ -145,7 +145,7 @@ void SAssetPicker::Construct( const FArguments& InArgs )
 
 	for (auto DelegateIt = InArgs._AssetPickerConfig.RefreshAssetViewDelegates.CreateConstIterator(); DelegateIt; ++DelegateIt)
 	{
-		if ((*DelegateIt) != NULL)
+		if ((*DelegateIt) != nullptr)
 		{
 			(**DelegateIt) = FRefreshAssetViewDelegate::CreateSP(this, &SAssetPicker::RefreshAssetView);
 		}
@@ -351,7 +351,7 @@ void SAssetPicker::Construct( const FArguments& InArgs )
 		.OnAssetToolTipClosing(InArgs._AssetPickerConfig.OnAssetToolTipClosing)
 		.FrontendFilters(FrontendFilters)
 		.TextFilter(TextFilter)
-		.ShowRedirectors_Lambda([this]() { return ContentBrowserUtils::ShouldShowRedirectors(FilterListPtr); })
+		.ShowRedirectors(InArgs._AssetPickerConfig.bAddFilterUI ? ContentBrowserUtils::ShouldShowRedirectors(FilterListPtr) : false)
 		.InitialContentSources(CurrentContentSources)
 		.InitialBackendFilter(CurrentBackendFilter)
 		.InitialViewType(InArgs._AssetPickerConfig.InitialAssetViewType)
