[AssetThumbnail] Thumbnail Edit Mode is missing the Reset icon after editing an asset's thumbnail

diff --git a/Engine/Source/Editor/ContentBrowser/Private/ContentBrowserStyle.cpp b/Engine/Source/Editor/ContentBrowser/Private/ContentBrowserStyle.cpp
--- a/Engine/Source/Editor/ContentBrowser/Private/ContentBrowserStyle.cpp
+++ b/Engine/Source/Editor/ContentBrowser/Private/ContentBrowserStyle.cpp
@@ -210,7 +210,7 @@ namespace UE::ContentBrowser::Private
 		Set("ContentBrowser.PrimitiveCube", new IMAGE_BRUSH("ContentBrowser/ThumbnailCube", CoreStyleConstants::Icon32x32));
 		Set("ContentBrowser.PrimitivePlane", new IMAGE_BRUSH("ContentBrowser/ThumbnailPlane", CoreStyleConstants::Icon32x32));
 		Set("ContentBrowser.PrimitiveCylinder", new IMAGE_BRUSH("ContentBrowser/ThumbnailCylinder", CoreStyleConstants::Icon32x32));
-		Set("ContentBrowser.ResetPrimitiveToDefault", new IMAGE_BRUSH_SVG("Starship/Common/Undo", CoreStyleConstants::Icon20x20));
+		Set("ContentBrowser.ResetPrimitiveToDefault", new CORE_IMAGE_BRUSH_SVG("Starship/Common/Undo", CoreStyleConstants::Icon20x20));
 
 		Set("ContentBrowser.TopBar.Font", FTextBlockStyle(NormalText)
 			.SetFont(DEFAULT_FONT("Bold", 11))
