Update context menu to set assets from EAssetAccessSpecifier::Public down to EAssetAccessSpecifier::EpicInternal

#rnx

diff --git a/Engine/Plugins/Editor/ContentBrowser/ContentBrowserAliasDataSource/Source/ContentBrowserAliasDataSource/Private/ContentBrowserAliasDataSource.cpp b/Engine/Plugins/Editor/ContentBrowser/ContentBrowserAliasDataSource/Source/ContentBrowserAliasDataSource/Private/ContentBrowserAliasDataSource.cpp
--- a/Engine/Plugins/Editor/ContentBrowser/ContentBrowserAliasDataSource/Source/ContentBrowserAliasDataSource/Private/ContentBrowserAliasDataSource.cpp
+++ b/Engine/Plugins/Editor/ContentBrowser/ContentBrowserAliasDataSource/Source/ContentBrowserAliasDataSource/Private/ContentBrowserAliasDataSource.cpp
@@ -957,14 +957,14 @@ bool UContentBrowserAliasDataSource::CanPrivatizeItem(const FContentBrowserItemD
 	return ContentBrowserAssetData::CanPrivatizeItem(AssetTools, AssetRegistry, this, InItem, OutErrorMsg);
 }
 
-bool UContentBrowserAliasDataSource::PrivatizeItem(const FContentBrowserItemData& InItem)
+bool UContentBrowserAliasDataSource::PrivatizeItem(const FContentBrowserItemData& InItem, const EAssetAccessSpecifier InAssetAccessSpecifier)
 {
-	return ContentBrowserAssetData::PrivatizeItems(AssetTools, AssetRegistry, this, MakeArrayView(&InItem, 1));
+	return ContentBrowserAssetData::PrivatizeItems(AssetTools, AssetRegistry, this, MakeArrayView(&InItem, 1), InAssetAccessSpecifier);
 }
 
-bool UContentBrowserAliasDataSource::BulkPrivatizeItems(TArrayView<const FContentBrowserItemData> InItems)
+bool UContentBrowserAliasDataSource::BulkPrivatizeItems(TArrayView<const FContentBrowserItemData> InItems, const EAssetAccessSpecifier InAssetAccessSpecifier)
 {
-	return ContentBrowserAssetData::PrivatizeItems(AssetTools, AssetRegistry, this, InItems);
+	return ContentBrowserAssetData::PrivatizeItems(AssetTools, AssetRegistry, this, InItems, InAssetAccessSpecifier);
 }
 
 bool UContentBrowserAliasDataSource::AppendItemReference(const FContentBrowserItemData& InItem, FString& InOutStr)
diff --git a/Engine/Plugins/Editor/ContentBrowser/ContentBrowserAliasDataSource/Source/ContentBrowserAliasDataSource/Public/ContentBrowserAliasDataSource.h b/Engine/Plugins/Editor/ContentBrowser/ContentBrowserAliasDataSource/Source/ContentBrowserAliasDataSource/Public/ContentBrowserAliasDataSource.h
--- a/Engine/Plugins/Editor/ContentBrowser/ContentBrowserAliasDataSource/Source/ContentBrowserAliasDataSource/Public/ContentBrowserAliasDataSource.h
+++ b/Engine/Plugins/Editor/ContentBrowser/ContentBrowserAliasDataSource/Source/ContentBrowserAliasDataSource/Public/ContentBrowserAliasDataSource.h
@@ -9,6 +9,7 @@
 #include "ContentBrowserLocalizedAlias.h"
 
 #include "AssetRegistry/PathTree.h"
+#include "UObject/Package.h"
 #include "ContentBrowserAliasDataSource.generated.h"
 
 class IAssetRegistry;
@@ -89,8 +90,8 @@ public:
 	virtual bool BulkSaveItems(TArrayView<const FContentBrowserItemData> InItems, const EContentBrowserItemSaveFlags InSaveFlags) override;
 
 	virtual bool CanPrivatizeItem(const FContentBrowserItemData& InItem, FText* OutErrorMsg) override;
-	virtual bool PrivatizeItem(const FContentBrowserItemData& InItem) override;
-	virtual bool BulkPrivatizeItems(TArrayView<const FContentBrowserItemData> InItems) override;
+	virtual bool PrivatizeItem(const FContentBrowserItemData& InItem, const EAssetAccessSpecifier InAssetAccessSpecifier = EAssetAccessSpecifier::Private) override;
+	virtual bool BulkPrivatizeItems(TArrayView<const FContentBrowserItemData> InItems, const EAssetAccessSpecifier InAssetAccessSpecifier = EAssetAccessSpecifier::Private) override;
 
 	virtual bool AppendItemReference(const FContentBrowserItemData& InItem, FString& InOutStr) override;
 	virtual bool AppendItemObjectPath(const FContentBrowserItemData& InItem, FString& InOutStr) override;
diff --git a/Engine/Plugins/Editor/ContentBrowser/ContentBrowserAssetDataSource/Source/ContentBrowserAssetDataSource/Private/ContentBrowserAssetDataCore.cpp b/Engine/Plugins/Editor/ContentBrowser/ContentBrowserAssetDataSource/Source/ContentBrowserAssetDataSource/Private/ContentBrowserAssetDataCore.cpp
--- a/Engine/Plugins/Editor/ContentBrowser/ContentBrowserAssetDataSource/Source/ContentBrowserAssetDataSource/Private/ContentBrowserAssetDataCore.cpp
+++ b/Engine/Plugins/Editor/ContentBrowser/ContentBrowserAssetDataSource/Source/ContentBrowserAssetDataSource/Private/ContentBrowserAssetDataCore.cpp
@@ -883,7 +883,7 @@ bool DeleteItems(IAssetTools* InAssetTools, IAssetRegistry* InAssetRegistry, con
 	return bDidDelete;
 }
 
-bool PrivatizeItems(IAssetTools* InAssetTools, IAssetRegistry* InAssetRegistry, const UContentBrowserDataSource* InOwnerDataSource, TArrayView<const FContentBrowserItemData> InItems)
+bool PrivatizeItems(IAssetTools* InAssetTools, IAssetRegistry* InAssetRegistry, const UContentBrowserDataSource* InOwnerDataSource, TArrayView<const FContentBrowserItemData> InItems, const EAssetAccessSpecifier InAssetAccessSpecifier)
 {
 	TArray<TSharedRef<const FContentBrowserAssetFileItemDataPayload>, TInlineAllocator<16>> AssetPayloads;
 
@@ -915,7 +915,7 @@ bool PrivatizeItems(IAssetTools* InAssetTools, IAssetRegistry* InAssetRegistry,
 		AssetsToPrivatize.Add(AssetPayload->GetAssetData());
 	}
 
-	return ObjectTools::PrivatizeAssets(AssetsToPrivatize) > 0;
+	return ObjectTools::PrivatizeAssets(AssetsToPrivatize, /*bShowConfirmation*/true, InAssetAccessSpecifier) > 0;
 }
 
 bool DeleteAssetFolderItems(TArrayView<const TSharedRef<const FContentBrowserAssetFolderItemDataPayload>> InFolderPayloads)
diff --git a/Engine/Plugins/Editor/ContentBrowser/ContentBrowserAssetDataSource/Source/ContentBrowserAssetDataSource/Private/ContentBrowserAssetDataSource.cpp b/Engine/Plugins/Editor/ContentBrowser/ContentBrowserAssetDataSource/Source/ContentBrowserAssetDataSource/Private/ContentBrowserAssetDataSource.cpp
--- a/Engine/Plugins/Editor/ContentBrowser/ContentBrowserAssetDataSource/Source/ContentBrowserAssetDataSource/Private/ContentBrowserAssetDataSource.cpp
+++ b/Engine/Plugins/Editor/ContentBrowser/ContentBrowserAssetDataSource/Source/ContentBrowserAssetDataSource/Private/ContentBrowserAssetDataSource.cpp
@@ -2783,14 +2783,14 @@ bool UContentBrowserAssetDataSource::CanPrivatizeItem(const FContentBrowserItemD
 	return ContentBrowserAssetData::CanPrivatizeItem(AssetTools, AssetRegistry, this, InItem, OutErrorMsg);
 }
 
-bool UContentBrowserAssetDataSource::PrivatizeItem(const FContentBrowserItemData& InItem)
+bool UContentBrowserAssetDataSource::PrivatizeItem(const FContentBrowserItemData& InItem, const EAssetAccessSpecifier InAssetAccessSpecifier)
 {
-	return ContentBrowserAssetData::PrivatizeItems(AssetTools, AssetRegistry, this, MakeArrayView(&InItem, 1));
+	return ContentBrowserAssetData::PrivatizeItems(AssetTools, AssetRegistry, this, MakeArrayView(&InItem, 1), InAssetAccessSpecifier);
 }
 
-bool UContentBrowserAssetDataSource::BulkPrivatizeItems(TArrayView<const FContentBrowserItemData> InItems)
+bool UContentBrowserAssetDataSource::BulkPrivatizeItems(TArrayView<const FContentBrowserItemData> InItems, const EAssetAccessSpecifier InAssetAccessSpecifier)
 {
-	return ContentBrowserAssetData::PrivatizeItems(AssetTools, AssetRegistry, this, InItems);
+	return ContentBrowserAssetData::PrivatizeItems(AssetTools, AssetRegistry, this, InItems, InAssetAccessSpecifier);
 }
 
 bool UContentBrowserAssetDataSource::CanRenameItem(const FContentBrowserItemData& InItem, const FString* InNewName, const IContentBrowserHideFolderIfEmptyFilter* HideFolderIfEmptyFilter, FText* OutErrorMsg)
diff --git a/Engine/Plugins/Editor/ContentBrowser/ContentBrowserAssetDataSource/Source/ContentBrowserAssetDataSource/Public/ContentBrowserAssetDataCore.h b/Engine/Plugins/Editor/ContentBrowser/ContentBrowserAssetDataSource/Source/ContentBrowserAssetDataSource/Public/ContentBrowserAssetDataCore.h
--- a/Engine/Plugins/Editor/ContentBrowser/ContentBrowserAssetDataSource/Source/ContentBrowserAssetDataSource/Public/ContentBrowserAssetDataCore.h
+++ b/Engine/Plugins/Editor/ContentBrowser/ContentBrowserAssetDataSource/Source/ContentBrowserAssetDataSource/Public/ContentBrowserAssetDataCore.h
@@ -3,6 +3,7 @@
 #pragma once
 
 #include "ContentBrowserItemData.h"
+#include "UObject/Package.h"
 
 class FContentBrowserAssetFileItemDataPayload;
 class FContentBrowserAssetFolderItemDataPayload;
@@ -134,7 +135,7 @@ CONTENTBROWSERASSETDATASOURCE_API FContentBrowserItemData CreateUnsupportedAsset
 
 	CONTENTBROWSERASSETDATASOURCE_API bool CanPrivatizeAssetFileItem(IAssetTools* InAssetTools, const FContentBrowserAssetFileItemDataPayload& InAssetPayload, FText* OutErrorMsg);
 
-	CONTENTBROWSERASSETDATASOURCE_API bool PrivatizeItems(IAssetTools* InAssetTools, IAssetRegistry* InAssetRegistry, const UContentBrowserDataSource* InOwnerDataSource, TArrayView<const FContentBrowserItemData> InItems);
+	CONTENTBROWSERASSETDATASOURCE_API bool PrivatizeItems(IAssetTools* InAssetTools, IAssetRegistry* InAssetRegistry, const UContentBrowserDataSource* InOwnerDataSource, TArrayView<const FContentBrowserItemData> InItems, const EAssetAccessSpecifier InAssetAccessSpecifier = EAssetAccessSpecifier::Private);
 
 	CONTENTBROWSERASSETDATASOURCE_API bool CanRenameItem(IAssetTools* InAssetTools, const UContentBrowserDataSource* InOwnerDataSource, const FContentBrowserItemData& InItem, const FString* InNewName, FText* OutErrorMsg);
 	CONTENTBROWSERASSETDATASOURCE_API bool CanRenameItem(IAssetTools* InAssetTools, const UContentBrowserDataSource* InOwnerDataSource, const FContentBrowserItemData& InItem, const bool bCheckUniqueName, const FString* InNewName, FText* OutErrorMsg);
diff --git a/Engine/Plugins/Editor/ContentBrowser/ContentBrowserAssetDataSource/Source/ContentBrowserAssetDataSource/Public/ContentBrowserAssetDataSource.h b/Engine/Plugins/Editor/ContentBrowser/ContentBrowserAssetDataSource/Source/ContentBrowserAssetDataSource/Public/ContentBrowserAssetDataSource.h
--- a/Engine/Plugins/Editor/ContentBrowser/ContentBrowserAssetDataSource/Source/ContentBrowserAssetDataSource/Public/ContentBrowserAssetDataSource.h
+++ b/Engine/Plugins/Editor/ContentBrowser/ContentBrowserAssetDataSource/Source/ContentBrowserAssetDataSource/Public/ContentBrowserAssetDataSource.h
@@ -307,9 +307,9 @@ public:
 
 	virtual bool CanPrivatizeItem(const FContentBrowserItemData& InItem, FText* OutErrorMsg);
 
-	virtual bool PrivatizeItem(const FContentBrowserItemData& InItem) override;
+	virtual bool PrivatizeItem(const FContentBrowserItemData& InItem, const EAssetAccessSpecifier InAssetAccessSpecifier = EAssetAccessSpecifier::Private) override;
 
-	virtual bool BulkPrivatizeItems(TArrayView<const FContentBrowserItemData> InItems) override;
+	virtual bool BulkPrivatizeItems(TArrayView<const FContentBrowserItemData> InItems, const EAssetAccessSpecifier InAssetAccessSpecifier = EAssetAccessSpecifier::Private) override;
 
 	virtual bool CanRenameItem(const FContentBrowserItemData& InItem, const FString* InNewName, const IContentBrowserHideFolderIfEmptyFilter* HideFolderIfEmptyFilter, FText* OutErrorMsg) override;
 
diff --git a/Engine/Source/Editor/ContentBrowser/Private/AssetContextMenu.cpp b/Engine/Source/Editor/ContentBrowser/Private/AssetContextMenu.cpp
--- a/Engine/Source/Editor/ContentBrowser/Private/AssetContextMenu.cpp
+++ b/Engine/Source/Editor/ContentBrowser/Private/AssetContextMenu.cpp
@@ -714,8 +714,8 @@ bool FAssetContextMenu::AddPublicStateMenuOptions(UToolMenu* Menu)
 				{
 					Section.AddMenuEntry(
 						"MarkSelectedAsInternal",
-						LOCTEXT("MarkSelectedAsInternal", "Mark Selected As Internal"),
-						LOCTEXT("MarkSelectedAsInternalTooltip", "Sets all selected assets to be internally available for reference by other plugins"),
+						LOCTEXT("MarkSelectedAsInternal", "Mark Selected As Epic Internal"),
+						LOCTEXT("MarkSelectedAsInternalTooltip", "Sets all selected assets to be available as Epic Internal for reference by other plugins"),
 						FSlateIcon(UE::ContentBrowser::Private::FContentBrowserStyle::Get().GetStyleSetName(), "ContentBrowser.AssetActions.PublicAssetToggle"),
 						FUIAction(
 							FExecuteAction::CreateSP(this, &FAssetContextMenu::ExecuteBulkSetAssetAccessSpecifier, EAssetAccessSpecifier::EpicInternal),
@@ -1334,7 +1334,31 @@ bool FAssetContextMenu::CanSetAssetAccessSpecifier(const EAssetAccessSpecifier I
 
 void FAssetContextMenu::ExecuteBulkSetAssetAccessSpecifier(EAssetAccessSpecifier DestScope)
 {
-	if (DestScope == EAssetAccessSpecifier::Public || DestScope == EAssetAccessSpecifier::EpicInternal)
+	TMap<UContentBrowserDataSource*, TArray<FContentBrowserItemData>> PrivatizeSourcesAndItems;
+
+	auto AddForBulkPrivatizeItems = [&PrivatizeSourcesAndItems](const FContentBrowserItem& InItem)
+	{	
+		FContentBrowserItem::FItemDataArrayView ItemDataArray = InItem.GetInternalItems();
+		for (const FContentBrowserItemData& ItemData : ItemDataArray)
+		{
+			if (UContentBrowserDataSource* ItemDataSource = ItemData.GetOwnerDataSource())
+			{
+				FText PrivateErrorMsg;
+				// CanPrivatizeItem checks !GetAssetFolderItemPayload(), CanModifyAssetFileItem(), !IsRunningPIE(), !IsRedirector()
+				if (ItemDataSource->CanPrivatizeItem(ItemData, &PrivateErrorMsg))
+				{
+					TArray<FContentBrowserItemData>& ItemsForSource = PrivatizeSourcesAndItems.FindOrAdd(ItemDataSource);
+					ItemsForSource.Add(ItemData);
+				}
+				else
+				{
+					AssetViewUtils::ShowErrorNotifcation(PrivateErrorMsg);
+				}
+			}
+		}
+	};
+
+	if (DestScope == EAssetAccessSpecifier::Public)
 	{
 		FScopedSlowTask SlowTask((float)SelectedFiles.Num(), LOCTEXT("SetAssetScope", "Loading assets and modifying scope..."));
 		SlowTask.MakeDialog(/*bShowCancelButton*/true);
@@ -1351,38 +1375,75 @@ void FAssetContextMenu::ExecuteBulkSetAssetAccessSpecifier(EAssetAccessSpecifier
 			FAssetData ItemAssetData;
 			if (SelectedItem.Legacy_TryGetAssetData(ItemAssetData))
 			{
-				SetAssetAccessSpecifier(ItemAssetData, DestScope, false);
+				if (UPackage* ItemAssetPackage = ItemAssetData.GetPackage())
+				{
+					ItemAssetPackage->SetAssetAccessSpecifier(DestScope);
+				}
 			}
 		}
 	}
-	else
+	else if (DestScope == EAssetAccessSpecifier::EpicInternal)
 	{
-		// Batch these by their data sources
-		TMap<UContentBrowserDataSource*, TArray<FContentBrowserItemData>> SourcesAndItems;
+		FScopedSlowTask SlowTask((float)SelectedFiles.Num(), LOCTEXT("SetAssetScope", "Loading assets and modifying scope..."));
+		SlowTask.MakeDialog(/*bShowCancelButton*/true);
+
 		for (const FContentBrowserItem& SelectedItem : SelectedFiles)
 		{
-			FContentBrowserItem::FItemDataArrayView ItemDataArray = SelectedItem.GetInternalItems();
-			for (const FContentBrowserItemData& ItemData : ItemDataArray)
+			SlowTask.EnterProgressFrame(1);
+
+			if (SlowTask.ShouldCancel())
 			{
-				if (UContentBrowserDataSource* ItemDataSource = ItemData.GetOwnerDataSource())
-				{
-					FText PrivateErrorMsg;
-					if (ItemDataSource->CanPrivatizeItem(ItemData, &PrivateErrorMsg))
-					{
-						TArray<FContentBrowserItemData>& ItemsForSource = SourcesAndItems.FindOrAdd(ItemDataSource);
-						ItemsForSource.Add(ItemData);
-					}
-					else
-					{
-						AssetViewUtils::ShowErrorNotifcation(PrivateErrorMsg);
-					}
-				}
+				break;
+			}
+
+			FAssetData ItemAssetData;
+			if (!SelectedItem.Legacy_TryGetAssetData(ItemAssetData))
+			{
+				break;
+			}
+
+			UPackage* ItemAssetPackage = ItemAssetData.GetPackage();
+			if (!ItemAssetPackage)
+			{
+				break;
+			}
+			
+			const EAssetAccessSpecifier OldAssetAccessSpecifier = ItemAssetPackage->GetAssetAccessSpecifier();
+			if (OldAssetAccessSpecifier == DestScope)
+			{
+				break;
+			}
+
+			if (OldAssetAccessSpecifier == EAssetAccessSpecifier::Private)
+			{
+				// Ok to raise EAssetAccessSpecifier::Private to DestScope without calling BulkPrivatizeItems
+				ItemAssetPackage->SetAssetAccessSpecifier(DestScope);
+			}
+			else
+			{
+				// Lowering EAssetAccessSpecifier::Public to EAssetAccessSpecifier::EpicInternal uses BulkPrivatizeItems workflow instead of directly calling SetAssetAccessSpecifier
+				AddForBulkPrivatizeItems(SelectedItem);
 			}
 		}
+	}
+	else if (DestScope == EAssetAccessSpecifier::Private)
+	{
+		// Batch these by their data sources
+		for (const FContentBrowserItem& SelectedItem : SelectedFiles)
+		{
+			AddForBulkPrivatizeItems(SelectedItem);
+		}
+	}
+	else
+	{
+		checkf(false, TEXT("Unknown EAssetAccessSpecifier dest scope: %d"), DestScope);
+	}
 
-		for (const auto& SourceAndItemsPair : SourcesAndItems)
+	if (!PrivatizeSourcesAndItems.IsEmpty())
+	{
+		for (const auto& SourceAndItemsPair : PrivatizeSourcesAndItems)
 		{
-			SourceAndItemsPair.Key->BulkPrivatizeItems(SourceAndItemsPair.Value);
+			SourceAndItemsPair.Key->BulkPrivatizeItems(SourceAndItemsPair.Value, DestScope);
 		}
 	}
 
diff --git a/Engine/Source/Editor/ContentBrowserData/Private/ContentBrowserDataSource.cpp b/Engine/Source/Editor/ContentBrowserData/Private/ContentBrowserDataSource.cpp
--- a/Engine/Source/Editor/ContentBrowserData/Private/ContentBrowserDataSource.cpp
+++ b/Engine/Source/Editor/ContentBrowserData/Private/ContentBrowserDataSource.cpp
@@ -428,17 +428,17 @@ bool UContentBrowserDataSource::CanPrivatizeItem(const FContentBrowserItemData&
 	return false;
 }
 
-bool UContentBrowserDataSource::PrivatizeItem(const FContentBrowserItemData& InItem)
+bool UContentBrowserDataSource::PrivatizeItem(const FContentBrowserItemData& InItem, const EAssetAccessSpecifier InAssetAccessSpecifier)
 {
 	return false;
 }
 
-bool UContentBrowserDataSource::BulkPrivatizeItems(TArrayView<const FContentBrowserItemData> InItems)
+bool UContentBrowserDataSource::BulkPrivatizeItems(TArrayView<const FContentBrowserItemData> InItems, const EAssetAccessSpecifier InAssetAccessSpecifier)
 {
 	bool bSuccess = false;
 	for (const FContentBrowserItemData& Item : InItems)
 	{
-		bSuccess |= PrivatizeItem(Item);
+		bSuccess |= PrivatizeItem(Item, InAssetAccessSpecifier);
 	}
 	return bSuccess;
 }
diff --git a/Engine/Source/Editor/ContentBrowserData/Public/ContentBrowserDataSource.h b/Engine/Source/Editor/ContentBrowserData/Public/ContentBrowserDataSource.h
--- a/Engine/Source/Editor/ContentBrowserData/Public/ContentBrowserDataSource.h
+++ b/Engine/Source/Editor/ContentBrowserData/Public/ContentBrowserDataSource.h
@@ -16,6 +16,7 @@
 #include "UObject/NameTypes.h"
 #include "UObject/Object.h"
 #include "UObject/ObjectMacros.h"
+#include "UObject/Package.h"
 #include "UObject/UObjectGlobals.h"
 
 #include "ContentBrowserDataSource.generated.h"
@@ -550,16 +551,17 @@ public:
 	*
 	* @return True if the item was marked private, false otherwise
 	*/
-	virtual bool PrivatizeItem(const FContentBrowserItemData& InItem);
+	virtual bool PrivatizeItem(const FContentBrowserItemData& InItem, const EAssetAccessSpecifier InAssetAccessSpecifier = EAssetAccessSpecifier::Private);
 
 	/**
-	* Attempt to mark the given items as private (NotExternallyReferenceable)
+	* Attempt to mark the given items as InAssetAccessSpecifier
 	* 
-	* @param InItems The items to be marked private.
+	* @param InItems The items to be marked InAssetAccessSpecifier.
+	* @param InAssetAccessSpecifier The AssetAccessSpecifier to be marked the assets as.
 	* 
-	* @return True if any items were marked private, false otherwise
+	* @return True if any items were marked InAssetAccessSpecifier, false otherwise
 	*/
-	virtual bool BulkPrivatizeItems(TArrayView<const FContentBrowserItemData> InItems);
+	virtual bool BulkPrivatizeItems(TArrayView<const FContentBrowserItemData> InItems, const EAssetAccessSpecifier InAssetAccessSpecifier = EAssetAccessSpecifier::Private);
 
 	/**
 	 * Query whether the given item is can be renamed, optionally providing error information if it cannot.
diff --git a/Engine/Source/Editor/UnrealEd/Private/AssetPrivatizeModel.cpp b/Engine/Source/Editor/UnrealEd/Private/AssetPrivatizeModel.cpp
--- a/Engine/Source/Editor/UnrealEd/Private/AssetPrivatizeModel.cpp
+++ b/Engine/Source/Editor/UnrealEd/Private/AssetPrivatizeModel.cpp
@@ -2,7 +2,7 @@
 
 #include "AssetPrivatizeModel.h"
 #include "Engine/Blueprint.h"
-
+#include "Misc/AssetAccessRestrictions.h"
 #include "ObjectTools.h"
 #include "AssetRegistry/AssetRegistryModule.h"
 
@@ -52,9 +52,12 @@ bool FPendingPrivateAsset::IsReferenceIllegal(const FName& InReference)
 	}
 	else if (TargetAssetAccessSpecifier == EAssetAccessSpecifier::EpicInternal)
 	{
-		// Will be implemented in future changes
-		checkf(false, TEXT("EAssetAccessSpecifier::EpicInternal cannot be used yet"));
-		return true;
+		if (UE::AssetAccessRestrictions::IsPathAllowedToReferenceEpicInternalAssets.IsBound())
+		{
+			return UE::AssetAccessRestrictions::IsPathAllowedToReferenceEpicInternalAssets.Execute(FNameBuilder(ReferenceMountPointName).ToView());
+		}
+
+		return false;
 	}
 
 	return true;
diff --git a/Engine/Source/Editor/UnrealEd/Private/ObjectTools.cpp b/Engine/Source/Editor/UnrealEd/Private/ObjectTools.cpp
--- a/Engine/Source/Editor/UnrealEd/Private/ObjectTools.cpp
+++ b/Engine/Source/Editor/UnrealEd/Private/ObjectTools.cpp
@@ -2800,7 +2800,7 @@ namespace ObjectTools
 		return NumPackagesToDelete + NumObjectsToDelete;
 	}
 
-	int32 PrivatizeAssets(const TArray<FAssetData>& AssetsToPrivatize, bool bShowConfirmation)
+	int32 PrivatizeAssets(const TArray<FAssetData>& AssetsToPrivatize, bool bShowConfirmation, const EAssetAccessSpecifier InAssetAccessSpecifier)
 	{
 		TArray<UObject*> ObjectsToPrivatize;
 		for (const FAssetData& AssetToPrivatize : AssetsToPrivatize)
@@ -2815,7 +2815,7 @@ namespace ObjectTools
 
 		if (!ObjectsToPrivatize.IsEmpty())
 		{
-			return PrivatizeObjects(ObjectsToPrivatize, bShowConfirmation);
+			return PrivatizeObjects(ObjectsToPrivatize, bShowConfirmation, EAllowCancelDuringPrivatize::AllowCancel, InAssetAccessSpecifier);
 		}
 
 		return 0;
@@ -3177,7 +3177,7 @@ namespace ObjectTools
 		return DeleteModel->GetDeletedObjectCount();
 	}
 
-	int32 PrivatizeObjects(const TArray<UObject*>& InObjectsToPrivatize, bool bShowConfirmation, EAllowCancelDuringPrivatize AllowCancelDuringPrivatize)
+	int32 PrivatizeObjects(const TArray<UObject*>& InObjectsToPrivatize, bool bShowConfirmation, EAllowCancelDuringPrivatize AllowCancelDuringPrivatize, const EAssetAccessSpecifier InAssetAccessSpecifier)
 	{
 		const FScopedBusyCursor BusyCursor;
 		TArray<UObject*> ObjectsToPrivatize = InObjectsToPrivatize;
@@ -3207,7 +3207,7 @@ namespace ObjectTools
 			return 0;
 		}
 
-		TSharedRef<FAssetPrivatizeModel> PrivatizeModel = MakeShared<FAssetPrivatizeModel>(ObjectsToPrivatize, EAssetAccessSpecifier::Private);
+		TSharedRef<FAssetPrivatizeModel> PrivatizeModel = MakeShared<FAssetPrivatizeModel>(ObjectsToPrivatize, InAssetAccessSpecifier);
 
 		if (bShowConfirmation)
 		{
diff --git a/Engine/Source/Editor/UnrealEd/Public/ObjectTools.h b/Engine/Source/Editor/UnrealEd/Public/ObjectTools.h
--- a/Engine/Source/Editor/UnrealEd/Public/ObjectTools.h
+++ b/Engine/Source/Editor/UnrealEd/Public/ObjectTools.h
@@ -331,7 +331,7 @@ namespace ObjectTools
 	* 
 	* @return The number of objects successfully privatized
 	*/
-	UNREALED_API int32 PrivatizeObjects(const TArray<UObject*>& InObjectsToPrivatize, bool bShowConfirmation = true, EAllowCancelDuringPrivatize AllowCancelDuringPrivatize = EAllowCancelDuringPrivatize::AllowCancel);
+	UNREALED_API int32 PrivatizeObjects(const TArray<UObject*>& InObjectsToPrivatize, bool bShowConfirmation = true, EAllowCancelDuringPrivatize AllowCancelDuringPrivatize = EAllowCancelDuringPrivatize::AllowCancel, const EAssetAccessSpecifier InAssetAccessSpecifier = EAssetAccessSpecifier::Private);
 
 	/**
 	* Deletes the list of objects without checking if they are still being used.  This should not be called directly
@@ -361,7 +361,7 @@ namespace ObjectTools
 	* 
 	* @return The number of assets successfully privatized
 	*/
-	UNREALED_API int32 PrivatizeAssets(const TArray<FAssetData>& AssetsToPrivatize, bool bShowConfirmation = true);
+	UNREALED_API int32 PrivatizeAssets(const TArray<FAssetData>& AssetsToPrivatize, bool bShowConfirmation = true, const EAssetAccessSpecifier InAssetAccessSpecifier = EAssetAccessSpecifier::Private);
 
 	/**
 	 * Delete a single object
diff --git a/Engine/Source/Runtime/CoreUObject/Private/Misc/AssetAccessRestrictions.cpp b/Engine/Source/Runtime/CoreUObject/Private/Misc/AssetAccessRestrictions.cpp
new file mode 100644
--- /dev/null
+++ b/Engine/Source/Runtime/CoreUObject/Private/Misc/AssetAccessRestrictions.cpp
@@ -0,0 +1,6 @@
+// Copyright Epic Games, Inc. All Rights Reserved.
+
+#include "Misc/AssetAccessRestrictions.h"
+
+
+TDelegate<bool(FStringView)> UE::AssetAccessRestrictions::IsPathAllowedToReferenceEpicInternalAssets;
diff --git a/Engine/Source/Runtime/CoreUObject/Public/Misc/AssetAccessRestrictions.h b/Engine/Source/Runtime/CoreUObject/Public/Misc/AssetAccessRestrictions.h
new file mode 100644
--- /dev/null
+++ b/Engine/Source/Runtime/CoreUObject/Public/Misc/AssetAccessRestrictions.h
@@ -0,0 +1,13 @@
+// Copyright Epic Games, Inc. All Rights Reserved.
+
+#pragma once
+
+#include "CoreGlobals.h"
+#include "CoreTypes.h"
+#include "Delegates/Delegate.h"
+
+namespace UE::AssetAccessRestrictions
+{
+	/** Is path allowed to reference Epic Internal assets */
+	extern COREUOBJECT_API TDelegate<bool(FStringView)> IsPathAllowedToReferenceEpicInternalAssets;
+}
\ No newline at end of file
