Fix the show other develloper assets toggle in the asset picker not working.
Renamed the filter name to reflect the new logic more accuratly.

diff --git a/Engine/Source/Editor/ContentBrowser/Private/SAssetPicker.cpp b/Engine/Source/Editor/ContentBrowser/Private/SAssetPicker.cpp
--- a/Engine/Source/Editor/ContentBrowser/Private/SAssetPicker.cpp
+++ b/Engine/Source/Editor/ContentBrowser/Private/SAssetPicker.cpp
@@ -290,13 +290,13 @@ void SAssetPicker::Construct( const FArguments& InArgs )
 		];
 
 		// Use the 'other developer' filter from the filter list widget. 
-		OtherDevelopersFilter = StaticCastSharedPtr<FFilter_HideOtherDevelopers>(FilterListPtr->GetFrontendFilter(TEXT("HideOtherDevelopersBackend")));
+		HideOtherDevelopersFilter = StaticCastSharedPtr<FFilter_HideOtherDevelopers>(FilterListPtr->GetFrontendFilter(TEXT("HideOtherDevelopersBackend")));
 	}
 	else
 	{
 		// Filter UI is off, but the 'other developer' filter is a built-in feature.
-		OtherDevelopersFilter = MakeShared<FFilter_HideOtherDevelopers>(nullptr, FName(SaveSettingsName));
-		OtherDevelopersFilter->SetActiveInCollection(OtherDevelopersFilter.ToSharedRef(), false, *FrontendFilters);
+		HideOtherDevelopersFilter = MakeShared<FFilter_HideOtherDevelopers>(nullptr, FName(SaveSettingsName));
+		HideOtherDevelopersFilter->SetActiveInCollection(HideOtherDevelopersFilter.ToSharedRef(), false, *FrontendFilters);
 	}
 
 	// Make game-specific filter
@@ -564,7 +564,7 @@ void SAssetPicker::OnFilterChanged()
 	}
 	else if (!IsShowingOtherDevelopersContent())
 	{
-		CustomPermissionLists.Add(OtherDevelopersFilter->GetPathPermissionList());
+		CustomPermissionLists.Add(HideOtherDevelopersFilter->GetPathPermissionList());
 	}
 
 	Filter.Append(CurrentBackendFilter);
@@ -702,13 +702,13 @@ bool SAssetPicker::IsShowingOtherDevelopersContent() const
 	if (FilterListPtr.IsValid())
 	{
 		// OtherDevelopersFilter is an inverse filter, so if it is active we are not showing other developers content
-		return FilterListPtr->IsFrontendFilterActive(OtherDevelopersFilter);
+		return FilterListPtr->IsFrontendFilterActive(HideOtherDevelopersFilter);
 	}
 	else
 	{
 		for (int32 i=0; i < FrontendFilters->Num(); ++i)
 		{
-			if (FrontendFilters->GetFilterAtIndex(i) == OtherDevelopersFilter)
+			if (FrontendFilters->GetFilterAtIndex(i) == HideOtherDevelopersFilter)
 			{
 				return false;
 			}
@@ -735,26 +735,28 @@ void SAssetPicker::HandleShowOtherDevelopersCheckStateChanged( ECheckBoxState In
 	if (FilterListPtr) // Filter UI enabled?
 	{
 		// Pin+activate or unpin+deactivate the filter. A widget is pinned on the filter UI. It allows the user to activate/deactive the filter independently of the 'checked' state.
-		FilterListPtr->SetFrontendFilterCheckState(OtherDevelopersFilter, InCheckboxState);
+		FilterListPtr->SetFrontendFilterCheckState(HideOtherDevelopersFilter, InCheckboxState);
 	}
 	else
 	{
 		if (InCheckboxState == ECheckBoxState::Checked)
 		{
-			OtherDevelopersFilter->SetActiveInCollection(OtherDevelopersFilter.ToSharedRef(), /* do show other developers content*/ false, *FrontendFilters);
+			HideOtherDevelopersFilter->SetActiveInCollection(HideOtherDevelopersFilter.ToSharedRef(), /* do show other developers content*/ false, *FrontendFilters);
 		}
 		else
 		{
-			OtherDevelopersFilter->SetActiveInCollection(OtherDevelopersFilter.ToSharedRef(), /* do not show other developers content*/ true, *FrontendFilters);
+			HideOtherDevelopersFilter->SetActiveInCollection(HideOtherDevelopersFilter.ToSharedRef(), /* do not show other developers content*/ true, *FrontendFilters);
 		}
 	}
+
+	OnFilterChanged();
 }
 
 ECheckBoxState SAssetPicker::GetShowOtherDevelopersCheckState() const
 {
 	if (FilterListPtr) // Filter UI enabled?
 	{
-		return FilterListPtr->GetFrontendFilterCheckState(OtherDevelopersFilter); // Tells whether the 'other developer' filter is pinned on the filter UI. (The filter itself may be active or not).
+		return FilterListPtr->GetFrontendFilterCheckState(HideOtherDevelopersFilter); // Tells whether the 'other developer' filter is pinned on the filter UI. (The filter itself may be active or not).
 	}
 	else
 	{
diff --git a/Engine/Source/Editor/ContentBrowser/Private/SAssetPicker.h b/Engine/Source/Editor/ContentBrowser/Private/SAssetPicker.h
--- a/Engine/Source/Editor/ContentBrowser/Private/SAssetPicker.h
+++ b/Engine/Source/Editor/ContentBrowser/Private/SAssetPicker.h
@@ -193,7 +193,7 @@ private:
 	/** Filters needed for filtering the assets */
 	TSharedPtr< FAssetFilterCollectionType > FilterCollection;
 	TSharedPtr<class FAssetTextFilter> TextFilter;
-	TSharedPtr<FFilter_HideOtherDevelopers> OtherDevelopersFilter;
+	TSharedPtr<FFilter_HideOtherDevelopers> HideOtherDevelopersFilter;
 
 	EAssetTypeCategories::Type DefaultFilterMenuExpansion;
 
