[ContentBrowser] WidgetSizeSwitcher ref widget size fix (loop was accumulating size)

diff --git a/Engine/Source/Editor/ContentBrowser/Private/Widgets/WidgetSizeSwitcher.h b/Engine/Source/Editor/ContentBrowser/Private/Widgets/WidgetSizeSwitcher.h
--- a/Engine/Source/Editor/ContentBrowser/Private/Widgets/WidgetSizeSwitcher.h
+++ b/Engine/Source/Editor/ContentBrowser/Private/Widgets/WidgetSizeSwitcher.h
@@ -11,7 +11,7 @@ namespace UE::Editor::ContentBrowser::Private
 {
 	/**
 	 * Returns a size override to bind to based on a single subject widget's desired size, clamped between a provided size range.
-	 * If a reference widget is provided, it's size is used for the range max value.
+	 * If a reference widget is provided, it's (allotted, not desired) size is used for the range max value.
 	 * Axis is 0 for x/width, 1 for y/height.
 	 */
 	template <EAxis::Type Axis>
@@ -53,7 +53,7 @@ namespace UE::Editor::ContentBrowser::Private
 			const int16 CurrentSubjectWidgetDesiredSize = static_cast<int16>(FMath::FloorToInt(SubjectWidget->GetDesiredSize()[Axis - 1] - UE_KINDA_SMALL_NUMBER));
 			const int16 MaxSize =
 				MaxSizeReferenceWidget.IsValid()
-				? static_cast<int16>(FMath::FloorToInt(MaxSizeReferenceWidget->GetDesiredSize()[Axis - 1] - UE_KINDA_SMALL_NUMBER))
+				? static_cast<int16>(FMath::FloorToInt(MaxSizeReferenceWidget->GetPaintSpaceGeometry().Size[Axis - 1] - UE_KINDA_SMALL_NUMBER))
 				: SizeRange.GetUpperBound().GetValue();
 
 			// Check if the subject or reference widget (max size) has changed
