[ContentBrowser] Unite ListView and ColumnView in one view only in the new CB 2.0

diff --git a/Engine/Source/Editor/ContentBrowser/Private/AssetViewWidgets.cpp b/Engine/Source/Editor/ContentBrowser/Private/AssetViewWidgets.cpp
--- a/Engine/Source/Editor/ContentBrowser/Private/AssetViewWidgets.cpp
+++ b/Engine/Source/Editor/ContentBrowser/Private/AssetViewWidgets.cpp
@@ -88,6 +88,7 @@
 #include "Widgets/SNullWidget.h"
 #include "Widgets/SOverlay.h"
 #include "Widgets/SToolTip.h"
+#include "Widgets/Layout/SWidgetSwitcher.h"
 #include "Widgets/Text/SInlineEditableTextBlock.h"
 
 class FDragDropEvent;
@@ -128,18 +129,18 @@ bool IsTopLevelFolder(const FName InFolderPath)
 // FAssetViewModeUtils
 ///////////////////////////////
 
-TSharedRef<SWidget> FAssetViewItemHelper::CreateListItemContents(SAssetListItem* const InListItem, const TSharedRef<SWidget>& InThumbnail, FName& OutItemShadowBorder)
+TSharedRef<SWidget> FAssetViewItemHelper::CreateListItemContents(SAssetListItem* const InListItem, const TSharedRef<SWidget>& InThumbnail, FName& OutItemShadowBorder, const FMargin& InFolderOverlayPadding)
 {
-	return CreateListTileItemContents(InListItem, InThumbnail, OutItemShadowBorder);
+	return CreateListTileItemContents(InListItem, InThumbnail, OutItemShadowBorder, InFolderOverlayPadding);
 }
 
-TSharedRef<SWidget> FAssetViewItemHelper::CreateTileItemContents(SAssetTileItem* const InTileItem, const TSharedRef<SWidget>& InThumbnail, FName& OutItemShadowBorder)
+TSharedRef<SWidget> FAssetViewItemHelper::CreateTileItemContents(SAssetTileItem* const InTileItem, const TSharedRef<SWidget>& InThumbnail, FName& OutItemShadowBorder, const FMargin& InFolderOverlayPadding)
 {
-	return CreateListTileItemContents(InTileItem, InThumbnail, OutItemShadowBorder);
+	return CreateListTileItemContents(InTileItem, InThumbnail, OutItemShadowBorder, InFolderOverlayPadding);
 }
 
 template <typename T>
-TSharedRef<SWidget> FAssetViewItemHelper::CreateListTileItemContents(T* const InTileOrListItem, const TSharedRef<SWidget>& InThumbnail, FName& OutItemShadowBorder)
+TSharedRef<SWidget> FAssetViewItemHelper::CreateListTileItemContents(T* const InTileOrListItem, const TSharedRef<SWidget>& InThumbnail, FName& OutItemShadowBorder, const FMargin& InFolderOverlayPadding)
 {
 	TSharedRef<SOverlay> ItemContentsOverlay = SNew(SOverlay);
 
@@ -168,7 +169,7 @@ TSharedRef<SWidget> FAssetViewItemHelper::CreateListTileItemContents(T* const In
 		// Folder base
 		// clang-format off
 		ItemContentsOverlay->AddSlot()
-		.Padding(FMargin(5))
+		.Padding(InFolderOverlayPadding)
 		[
 			SNew(SBorder)
 			.BorderImage(DropShadowImage)
@@ -1840,8 +1841,8 @@ void SAssetListItem::Construct( const FArguments& InArgs )
 	ItemHeight = InArgs._ItemHeight;
 	IsSelectedExclusively = InArgs._IsSelectedExclusively;
 	HighlightText = InArgs._HighlightText;
-
 	ThumbnailPadding = InArgs._ThumbnailPadding;
+	CurrentThumbnailSize = InArgs._CurrentThumbnailSize;
 
 	if ( AssetItem.IsValid() && AssetThumbnail.IsValid() )
 	{
@@ -1854,6 +1855,14 @@ void SAssetListItem::Construct( const FArguments& InArgs )
 		ThumbnailConfig.HighlightedText = InArgs._HighlightText;
 		ThumbnailConfig.HintColorAndOpacity = InArgs._ThumbnailHintColorAndOpacity;
 
+		if (UE::Editor::ContentBrowser::IsNewStyleEnabled())
+		{
+			ThumbnailConfig.GenericThumbnailSize = MakeAttributeSP(this, &SAssetListItem::GetGenericThumbnailSize);
+			ThumbnailConfig.AssetSystemInfoProvider = MakeShared<FAssetSystemContentBrowserInfoProvider>(AssetItem);
+			ThumbnailConfig.bAllowAssetStatusThumbnailOverlay = true;
+			ThumbnailConfig.bShowAssetChip = true;
+		}
+
 		if(!AssetItem->GetItem().IsSupported())
 		{
 			ThumbnailConfig.ClassThumbnailBrushOverride = FName("Icons.WarningWithColor.Thumbnail");
@@ -1874,6 +1883,15 @@ void SAssetListItem::Construct( const FArguments& InArgs )
 		}
 
 		ThumbnailWidget = AssetThumbnail->MakeThumbnailWidget(ThumbnailConfig);
+
+		if (UE::Editor::ContentBrowser::IsNewStyleEnabled())
+		{
+			// Use the same tooltip as the Thumbnail
+			if (const TSharedPtr<IToolTip>& ThumbnailTooltip = ThumbnailWidget->GetToolTip())
+			{
+				SetToolTip(ThumbnailTooltip);
+			}
+		}
 	}
 	else
 	{
@@ -1886,9 +1904,9 @@ void SAssetListItem::OnAssetDataChanged()
 {
 	SAssetViewItem::OnAssetDataChanged();
 
-	if (ClassTextWidget.IsValid())
+	if (PathTextWidget.IsValid())
 	{
-		ClassTextWidget->SetText(GetAssetClassText());
+		PathTextWidget->SetText(GetAssetPathText());
 	}
 
 	if (AssetThumbnail)
@@ -1937,11 +1955,61 @@ FSlateColor SAssetListItem::GetColumnTextColor(FIsSelected InIsSelected) const
 	return FSlateColor::UseForeground();
 }
 
-TSharedRef<SWidget> SAssetListItem::GenerateWidgetForColumn( const FName& ColumnName, FIsSelected InIsSelected )
+FText SAssetListItem::GetAssetPathText() const
+{
+	return AssetItem
+		? FText::AsCultureInvariant(AssetItem->GetItem().GetVirtualPath().ToString())
+		: FText();
+}
+
+FText SAssetListItem::GetAssetTagText(FName InTag) const
+{
+	if (AssetItem)
+	{
+		// Check custom type
+		{
+			FText TagText;
+			if (AssetItem->GetCustomColumnDisplayValue(InTag, TagText))
+			{
+				return TagText;
+			}
+		}
+
+		// Check display tags
+		{
+			const FTagDisplayItem* FoundTagItem = CachedDisplayTags.FindByPredicate([InTag](const FTagDisplayItem& TagItem)
+			{
+				return TagItem.TagKey == InTag;
+			});
+			
+			if (FoundTagItem)
+			{
+				return FoundTagItem->DisplayValue;
+			}
+		}
+
+	}
+	
+	return FText();
+}
+
+int32 SAssetListItem::GetWidgetSwitcherIndex() const
+{
+	return CurrentThumbnailSize.IsSet() && CurrentThumbnailSize.Get() >= EThumbnailSize::Large ? 1 : 0;
+}
+
+TSharedRef<SWidget> SAssetListItem::GenerateWidgetForColumn(const FName& ColumnName, FIsSelected InIsSelected)
 {
+	if (UE::Editor::ContentBrowser::IsNewStyleEnabled())
+	{
+		return GenerateWidgetForColumnNewStyle(ColumnName, InIsSelected);
+	}
+
 	TSharedPtr<SWidget> Content;
 
-	if ( ColumnName == "Name" )
+	using namespace UE::ContentBrowser::Private;
+
+	if (ColumnName == "Name")
 	{
 		FName ItemShadowBorderName;
 		TSharedRef<SWidget> ItemContents = FAssetViewItemHelper::CreateListItemContents(this, ThumbnailWidget.ToSharedRef(), ItemShadowBorderName);
@@ -1979,7 +2047,7 @@ TSharedRef<SWidget> SAssetListItem::GenerateWidgetForColumn( const FName& Column
 				.Padding(0, 1)
 				[
 					SAssignNew(InlineRenameWidget, SInlineEditableTextBlock)
-					.Font(UE::ContentBrowser::Private::FContentBrowserStyle::Get().GetFontStyle("ContentBrowser.AssetTileViewNameFont"))
+					.Font(FContentBrowserStyle::Get().GetFontStyle("ContentBrowser.AssetTileViewNameFont"))
 					.Text( GetNameText() )
 					.OnBeginTextEdit(this, &SAssetListItem::HandleBeginNameChange)
 					.OnTextCommitted(this, &SAssetListItem::HandleNameCommitted)
@@ -1996,7 +2064,7 @@ TSharedRef<SWidget> SAssetListItem::GenerateWidgetForColumn( const FName& Column
 				[
 					// Class
 					SAssignNew(ClassTextWidget, STextBlock)
-					.Font(UE::ContentBrowser::Private::FContentBrowserStyle::Get().GetFontStyle("ContentBrowser.AssetListViewClassFont"))
+					.Font(FContentBrowserStyle::Get().GetFontStyle("ContentBrowser.AssetListViewClassFont"))
 					.Text(GetAssetClassText())
 					.HighlightText(HighlightText)
 					.ColorAndOpacity(this, &SAssetListItem::GetColumnTextColor, InIsSelected)
@@ -2028,6 +2096,161 @@ TSharedRef<SWidget> SAssetListItem::GenerateWidgetForColumn( const FName& Column
 	return Content.ToSharedRef();
 }
 
+int32 SAssetListItem::GetGenericThumbnailSize() const
+{
+	const int32 ThumbnailSizeIdx = (int32)CurrentThumbnailSize.Get();
+	if (UE::Editor::ContentBrowser::IsNewStyleEnabled())
+	{
+		static int32 GenericThumbnailSizes[(int32)EThumbnailSize::MAX] = { 22, 24, 32, 48, 64, 80 };
+		return GenericThumbnailSizes[ThumbnailSizeIdx];
+	}
+	else
+	{
+		static int32 GenericThumbnailSizes[(int32)EThumbnailSize::MAX] = { 24, 32, 64, 128, 160, 200 };
+		return GenericThumbnailSizes[ThumbnailSizeIdx];
+	}
+}
+
+TSharedRef<SWidget> SAssetListItem::GenerateWidgetForColumnNewStyle(const FName& ColumnName, FIsSelected InIsSelected)
+{
+	TSharedPtr<SWidget> Content;
+
+	if (ColumnName == "Name")
+	{
+		FName ItemShadowBorderName;
+		TSharedRef<SWidget> ItemContents = FAssetViewItemHelper::CreateListItemContents(this, ThumbnailWidget.ToSharedRef(), ItemShadowBorderName);
+
+		const FSlateBrush* IconBrush;
+		if (IsFolder())
+		{
+			if (ContentBrowserUtils::IsItemDeveloperContent(AssetItem->GetItem()))
+			{
+				IconBrush = UE::ContentBrowser::Private::FContentBrowserStyle::Get().GetBrush("ContentBrowser.ColumnViewDeveloperFolderIcon");
+			}
+			else
+			{
+				IconBrush = UE::ContentBrowser::Private::FContentBrowserStyle::Get().GetBrush("ContentBrowser.ColumnViewFolderIcon");
+			}
+		}
+		else
+		{
+			if (!AssetItem->GetItem().IsSupported()) 
+			{
+				IconBrush = FAppStyle::GetBrush("Icons.WarningWithColor");
+			}
+			else
+			{
+				IconBrush = UE::ContentBrowser::Private::FContentBrowserStyle::Get().GetBrush("ContentBrowser.FileImage");
+			}
+		}
+
+		Content = SNew(SBorder)
+		.BorderImage(this, &SAssetViewItem::GetBorderImage)
+		.Padding(0.f)
+		.AddMetaData<FTagMetaData>(FTagMetaData(AssetItem->GetItem().GetVirtualPath()))
+		[
+			SNew(SHorizontalBox)
+
+			// Viewport
+			+SHorizontalBox::Slot()
+			.AutoWidth()
+			.VAlign(VAlign_Center)
+			[
+				SNew( SBox )
+				.WidthOverride( this, &SAssetListItem::GetThumbnailBoxSize )
+				.HeightOverride( this, &SAssetListItem::GetThumbnailBoxSize )
+				[
+					SNew(SWidgetSwitcher)
+					.WidgetIndex(this, &SAssetListItem::GetWidgetSwitcherIndex)
+
+					+ SWidgetSwitcher::Slot()
+					.VAlign(VAlign_Fill)
+					.HAlign(HAlign_Fill)
+					[
+						SNew(SImage)
+						.Image(IconBrush)
+						.ColorAndOpacity(this, &SAssetViewItem::GetAssetColor)
+					]
+
+					+ SWidgetSwitcher::Slot()
+					.VAlign(VAlign_Fill)
+					.HAlign(HAlign_Fill)
+					[
+						ItemContents
+					]
+				]
+			]
+
+			+SHorizontalBox::Slot()
+			.AutoWidth()
+			.Padding(6, 0, 0, 0)
+			.VAlign(VAlign_Center)
+			[
+				SAssignNew(InlineRenameWidget, SInlineEditableTextBlock)
+				.Font(UE::ContentBrowser::Private::FContentBrowserStyle::Get().GetFontStyle("ContentBrowser.AssetTileViewNameFont"))
+				.Text( GetNameText() )
+				.OnBeginTextEdit(this, &SAssetListItem::HandleBeginNameChange)
+				.OnTextCommitted(this, &SAssetListItem::HandleNameCommitted)
+				.OnVerifyTextChanged(this, &SAssetListItem::HandleVerifyNameChanged)
+				.HighlightText(HighlightText)
+				.IsSelected(IsSelectedExclusively)
+				.IsReadOnly(this, &SAssetListItem::IsNameReadOnly)
+				.ColorAndOpacity(this, &SAssetListItem::GetColumnTextColor, InIsSelected)
+			]
+		];
+
+		if(AssetItem.IsValid())
+		{
+			AssetItem->OnRenameRequested().BindSP(InlineRenameWidget.Get(), &SInlineEditableTextBlock::EnterEditingMode);
+			AssetItem->OnRenameCanceled().BindSP(InlineRenameWidget.Get(), &SInlineEditableTextBlock::ExitEditingMode);
+		}
+
+	}
+	else if (ColumnName == "RevisionControl")
+	{
+		Content = SNew(SBox)
+			.WidthOverride(16.0f)
+			.HeightOverride(16.0f)
+			.VAlign(VAlign_Center)
+			.HAlign(HAlign_Center)
+			[
+				GenerateSourceControlIconWidget()
+			];
+	}
+	else if (ColumnName == "Class")
+	{
+		Content = SNew(SBox)
+			.VAlign(VAlign_Center)
+			[
+				SAssignNew(ClassTextWidget, STextBlock)
+				.Text(GetAssetClassText())
+				.HighlightText(HighlightText)
+			];
+	}
+	else if (ColumnName == "Path")
+	{
+		Content =  SNew(SBox)
+			.VAlign(VAlign_Center)
+			[
+				SAssignNew(PathTextWidget, STextBlock)
+				.Text(GetAssetPathText())
+				.HighlightText(HighlightText)
+			];
+	}
+	else
+	{
+		Content =  SNew(SBox)
+			.VAlign(VAlign_Center)
+			[
+				SNew(STextBlock)
+				.Text(TAttribute<FText>::Create(TAttribute<FText>::FGetter::CreateSP(this, &SAssetListItem::GetAssetTagText, ColumnName)))
+			];
+	}
+	HandleSourceControlStateChanged();
+
+	return Content.ToSharedRef();
+}
+
 
 float SAssetListItem::GetExtraStateIconWidth() const
 {
@@ -2147,7 +2370,7 @@ void SAssetTileItem::Construct( const FArguments& InArgs )
 	}
 
 	FName ItemShadowBorderName;
-	TSharedRef<SWidget> ItemContents = FAssetViewItemHelper::CreateTileItemContents(this, Thumbnail.ToSharedRef(), ItemShadowBorderName);
+	TSharedRef<SWidget> ItemContents = FAssetViewItemHelper::CreateTileItemContents(this, Thumbnail.ToSharedRef(), ItemShadowBorderName, FMargin(5.f));
 
 	if (bShowType)
 	{
diff --git a/Engine/Source/Editor/ContentBrowser/Private/ContentBrowserStyle.cpp b/Engine/Source/Editor/ContentBrowser/Private/ContentBrowserStyle.cpp
--- a/Engine/Source/Editor/ContentBrowser/Private/ContentBrowserStyle.cpp
+++ b/Engine/Source/Editor/ContentBrowser/Private/ContentBrowserStyle.cpp
@@ -255,6 +255,7 @@ namespace UE::ContentBrowser::Private
 
 
 		Set("ContentBrowser.ColumnViewAssetIcon", new IMAGE_BRUSH("Icons/doc_16x", CoreStyleConstants::Icon16x16));
+		Set("ContentBrowser.FileImage", new CORE_IMAGE_BRUSH_SVG("Starship/Common/file", CoreStyleConstants::Icon16x16));
 
 		Set("ContentBrowser.ColumnViewFolderIcon", new CORE_IMAGE_BRUSH_SVG("Starship/Common/folder-closed", CoreStyleConstants::Icon16x16));
 		Set("ContentBrowser.ColumnViewDeveloperFolderIcon", new IMAGE_BRUSH_SVG("Starship/ContentBrowser/folder-developer", CoreStyleConstants::Icon16x16));
diff --git a/Engine/Source/Editor/ContentBrowser/Private/Menus/AssetViewMenus.cpp b/Engine/Source/Editor/ContentBrowser/Private/Menus/AssetViewMenus.cpp
--- a/Engine/Source/Editor/ContentBrowser/Private/Menus/AssetViewMenus.cpp
+++ b/Engine/Source/Editor/ContentBrowser/Private/Menus/AssetViewMenus.cpp
@@ -110,7 +110,7 @@ FDelayedAutoRegisterHelper SAssetView::AssetViewOptionsMenuRegistration(
 
 						ViewTypeSection.AddMenuEntry(
 							"TileView",
-							LOCTEXT("TileViewOption", "Tiles"),
+							LOCTEXT("TileViewOption", "Grid"),
 							LOCTEXT("TileViewOptionToolTip", "View assets as tiles in a grid."),
 							FSlateIcon(),
 							FUIAction(
@@ -134,18 +134,21 @@ FDelayedAutoRegisterHelper SAssetView::AssetViewOptionsMenuRegistration(
 							EUserInterfaceActionType::RadioButton
 						);
 
-						ViewTypeSection.AddMenuEntry(
-							"ColumnView",
-							LOCTEXT("ColumnViewOption", "Columns"),
-							LOCTEXT("ColumnViewOptionToolTip", "View assets in a list with columns of details."),
-							FSlateIcon(),
-							FUIAction(
-								FExecuteAction::CreateSP(AssetView, &SAssetView::SetCurrentViewTypeFromMenu, EAssetViewType::Column),
-								FCanExecuteAction(),
-								FIsActionChecked::CreateSP(AssetView, &SAssetView::IsCurrentViewType, EAssetViewType::Column)
-							),
-							EUserInterfaceActionType::RadioButton
-						);
+						if (!UE::Editor::ContentBrowser::IsNewStyleEnabled())
+						{
+							ViewTypeSection.AddMenuEntry(
+								"ColumnView",
+								LOCTEXT("ColumnViewOption", "Columns"),
+								LOCTEXT("ColumnViewOptionToolTip", "View assets in a list with columns of details."),
+								FSlateIcon(),
+								FUIAction(
+									FExecuteAction::CreateSP(AssetView, &SAssetView::SetCurrentViewTypeFromMenu, EAssetViewType::Column),
+									FCanExecuteAction(),
+									FIsActionChecked::CreateSP(AssetView, &SAssetView::IsCurrentViewType, EAssetViewType::Column)
+								),
+								EUserInterfaceActionType::RadioButton
+							);
+						}
 
 						const FText CustomViewLabel = AssetView->ViewExtender
 							? AssetView->ViewExtender->GetViewDisplayName()
@@ -263,8 +266,8 @@ FDelayedAutoRegisterHelper SAssetView::AssetViewOptionsMenuRegistration(
 
 						ThumbnailSection.AddEntry(FToolMenuEntry::InitSubMenu(
 							"ThumbnailSize",
-							LOCTEXT("ThumbnailSize", "Thumbnail Size"),
-							LOCTEXT("ThumbnailSizeToolTip", "Adjust the size of thumbnails."),
+							FToolUIActionChoice(),
+							Private::GetThumbnailSizeWidget(),
 							FNewToolMenuDelegate::CreateLambda(CreateThumbnailSizeSubMenu)
 						));
 
diff --git a/Engine/Source/Editor/ContentBrowser/Private/SAssetView.cpp b/Engine/Source/Editor/ContentBrowser/Private/SAssetView.cpp
--- a/Engine/Source/Editor/ContentBrowser/Private/SAssetView.cpp
+++ b/Engine/Source/Editor/ContentBrowser/Private/SAssetView.cpp
@@ -1140,17 +1140,19 @@ void SAssetView::Construct( const FArguments& InArgs )
 	// Use the shared ThumbnailPool for the rendering of thumbnails
 	AssetThumbnailPool = UThumbnailManager::Get().GetSharedThumbnailPool();
 	NumOffscreenThumbnails = 64;
-	ListViewThumbnailResolution = 128;
-	ListViewThumbnailSize = 64;
-	ListViewThumbnailPadding = 4;
+	ListViewThumbnailResolution = 256;
+	ListViewThumbnailPadding = UE::Editor::ContentBrowser::IsNewStyleEnabled() ? 2 : 4;
 	TileViewThumbnailResolution = 256;
+	TileViewThumbnailPadding = 9;
 
 	// Max Size for the thumbnail
 	const int32 MaxTileViewThumbnailSize = UE::Editor::ContentBrowser::IsNewStyleEnabled() ? 160 : 150;
 
 	TileViewThumbnailSize = MaxTileViewThumbnailSize;
 
-	TileViewThumbnailPadding = 9;
+	const int32 MaxListViewThumbnailViewSize = UE::Editor::ContentBrowser::IsNewStyleEnabled() ? 160 : 64;
+
+	ListViewThumbnailSize = MaxListViewThumbnailViewSize;
 
 	TileViewNameHeight = 50;
 
@@ -1264,9 +1266,16 @@ void SAssetView::Construct( const FArguments& InArgs )
 	OnExtendAssetViewOptionsMenuContext = InArgs._OnExtendAssetViewOptionsMenuContext;
 	AssetViewOptionsProfile = InArgs._AssetViewOptionsProfile;
 	
-	if ( InArgs._InitialViewType >= 0 && InArgs._InitialViewType < EAssetViewType::MAX )
+	if (InArgs._InitialViewType >= 0 && InArgs._InitialViewType < EAssetViewType::MAX)
 	{
-		CurrentViewType = InArgs._InitialViewType;
+		if (UE::Editor::ContentBrowser::IsNewStyleEnabled() && InArgs._InitialViewType == EAssetViewType::Column)
+		{
+			CurrentViewType = EAssetViewType::List;
+		}
+		else
+		{
+			CurrentViewType = InArgs._InitialViewType;
+		}
 	}
 	else
 	{
@@ -1509,7 +1518,9 @@ void SAssetView::Construct( const FArguments& InArgs )
 	// Generalizing this is a bit difficult because the column ID is not accessible or is not known
 	// Currently I assume this won't work, if this view mode is not column. Otherwise, I don't think sorting by path
 	// is a good idea. 
-	if (CurrentViewType == EAssetViewType::Column && bSortByPathInColumnView)
+	if ((CurrentViewType == EAssetViewType::Column ||
+		(UE::Editor::ContentBrowser::IsNewStyleEnabled() && CurrentViewType == EAssetViewType::List))
+		&& bSortByPathInColumnView)
 	{
 		SortManager->SetSortColumnId(EColumnSortPriority::Primary, SortManager->PathColumnId);
 		SortManager->SetSortColumnId(EColumnSortPriority::Secondary, SortManager->NameColumnId);
@@ -1874,7 +1885,7 @@ void SAssetView::ApplyHistoryData( const FHistoryData& History )
 
 TArray<TSharedPtr<FAssetViewItem>> SAssetView::GetSelectedViewItems() const
 {
-	switch ( GetCurrentViewType() )
+	switch (GetCurrentViewType())
 	{
 		case EAssetViewType::List: return ListView->GetSelectedItems();
 		case EAssetViewType::Tile: return TileView->GetSelectedItems();
@@ -2017,14 +2028,20 @@ void SAssetView::LoadSettings(const FString& IniFilename, const FString& IniSect
 	}
 
 	int32 ViewType = EAssetViewType::Tile;
-	if ( GConfig->GetInt(*IniSection, *GetCurrentViewTypeSettingPath(SettingsString), ViewType, IniFilename) )
+	if (GConfig->GetInt(*IniSection, *GetCurrentViewTypeSettingPath(SettingsString), ViewType, IniFilename))
 	{
 		// Clamp value to normal range and update state
-		if ( ViewType < 0 || ViewType >= EAssetViewType::MAX)
+		if (ViewType < 0 || ViewType >= EAssetViewType::MAX)
 		{
 			ViewType = EAssetViewType::Tile;
 		}
-		SetCurrentViewType( (EAssetViewType::Type)ViewType );
+
+		// If the new style is enabled do not allow the ColumnView
+		if (UE::Editor::ContentBrowser::IsNewStyleEnabled() && ViewType == 2 /** Column */)
+		{
+			ViewType = 1; // List
+		}
+		SetCurrentViewType((EAssetViewType::Type)ViewType);
 	}
 
 	float Zoom = 0;
@@ -2041,13 +2058,24 @@ void SAssetView::LoadSettings(const FString& IniFilename, const FString& IniSect
 		HiddenColumnNames = LoadedHiddenColumnNames;
 
 		// Also update the visibility of the columns we just loaded in (unless this is called before creation and ColumnView doesn't exist)
-		if(ColumnView)
+		if (ColumnView)
 		{
 			for (const SHeaderRow::FColumn& Column : ColumnView->GetHeaderRow()->GetColumns())
 			{
 				ColumnView->GetHeaderRow()->SetShowGeneratedColumn(Column.ColumnId, !HiddenColumnNames.Contains(Column.ColumnId.ToString()));
 			}
 		}
+
+		if (UE::Editor::ContentBrowser::IsNewStyleEnabled())
+		{
+			if (ListView)
+			{
+				for (const SHeaderRow::FColumn& ListColumn : ListView->GetHeaderRow()->GetColumns())
+				{
+					ListView->GetHeaderRow()->SetShowGeneratedColumn(ListColumn.ColumnId, !HiddenColumnNames.Contains(ListColumn.ColumnId.ToString()));
+				}
+			}
+		}
 	}
 }
 
@@ -2185,7 +2213,8 @@ void SAssetView::Tick( const FGeometry& AllottedGeometry, const double InCurrent
 		else
 		{
 			// Update the columns in the column view now that we know the majority type
-			if (CurrentViewType == EAssetViewType::Column)
+			if (CurrentViewType == EAssetViewType::Column ||
+				(UE::Editor::ContentBrowser::IsNewStyleEnabled() && CurrentViewType == EAssetViewType::List))
 			{
 				int32 HighestCount = 0;
 				FName HighestType;
@@ -2418,7 +2447,8 @@ void SAssetView::ProcessItemsPendingFilter(const double TickStartTime)
 	int32 OldCount = FilteredAssetItems.Num();
 	Items->UpdateItemFiltering(FrontendFilterHelper, bFlushAllPendingItems ? MAX_dbl : TickStartTime + MaxSecondsPerFrame, FilteredAssetItems);
 
-	if (CurrentViewType == EAssetViewType::Column)
+	if (CurrentViewType == EAssetViewType::Column ||
+		(UE::Editor::ContentBrowser::IsNewStyleEnabled() && CurrentViewType == EAssetViewType::List))
 	{
 		for (int32 i = OldCount; i < FilteredAssetItems.Num(); ++i)
 		{
@@ -2644,7 +2674,8 @@ FReply SAssetView::OnMouseWheel( const FGeometry& MyGeometry, const FPointerEven
 	{
 		if (UE::Editor::ContentBrowser::IsNewStyleEnabled())
 		{
-			if (MouseEvent.IsShiftDown())
+			// If shift is pressed or if we are in the List view the CTRL + Wheel should jump ThumbnailSize by design
+			if (MouseEvent.IsShiftDown() || CurrentViewType == EAssetViewType::List)
 			{
 				const int32 Delta = MouseEvent.GetWheelDelta() > 0 ? 1 : -1;
 				EThumbnailSize DesiredThumbnailSize = (EThumbnailSize)FMath::Clamp<int32>((int32)ThumbnailSize + Delta, 0, (int32)EThumbnailSize::MAX - 1);
@@ -2671,6 +2702,8 @@ FReply SAssetView::OnMouseWheel( const FGeometry& MyGeometry, const FPointerEven
 				else
 				{
 					ZoomScale = FMath::Clamp(ZoomScale + NewDelta, 0.f, 1.f);
+					// Always refresh the view when changing size otherwise some items may be missing sometimes
+					RefreshList();
 				}
 			}
 		}
@@ -2734,8 +2767,8 @@ TSharedRef<SAssetListView> SAssetView::CreateListView()
 
 	RevisionControlColumnIcon->AddLayer(TAttribute<const FSlateBrush*>::CreateSP(this, &SAssetView::GetRevisionControlColumnIconBadge));
 
-	return SNew(SAssetListView)
-		.SelectionMode( SelectionMode )
+	TSharedPtr<SAssetListView> NewListView = SNew(SAssetListView)
+		.SelectionMode(SelectionMode)
 		.ListItemsSource(&FilteredAssetItems)
 		.OnGenerateRow(this, &SAssetView::MakeListViewWidget)
 		.OnItemToString_Debug_Static(&FAssetViewItem::ItemToString_Debug)
@@ -2743,12 +2776,13 @@ TSharedRef<SAssetListView> SAssetView::CreateListView()
 		.OnContextMenuOpening(this, &SAssetView::OnGetContextMenuContent)
 		.OnMouseButtonDoubleClick(this, &SAssetView::OnListMouseButtonDoubleClick)
 		.OnSelectionChanged(this, &SAssetView::AssetSelectionChanged)
-		.ItemHeight(this, &SAssetView::GetListViewItemHeight)
 		.ScrollbarVisibility(bForceHideScrollbar ? EVisibility::Collapsed : EVisibility::Visible)
 		.HeaderRow
 		(
 			SNew(SHeaderRow)
-			.ResizeMode(ESplitterResizeMode::FixedSize)
+			.ResizeMode(UE::Editor::ContentBrowser::IsNewStyleEnabled() ? ESplitterResizeMode::Fill : ESplitterResizeMode::FixedSize)
+			.CanSelectGeneratedColumn(UE::Editor::ContentBrowser::IsNewStyleEnabled())
+			.OnHiddenColumnsListChanged(this, &SAssetView::OnHiddenColumnsChanged)
 
 			// Revision Control column, currently doesn't support sorting
 			+ SHeaderRow::Column(SortManager->RevisionControlColumnId)
@@ -2770,6 +2804,49 @@ TSharedRef<SAssetListView> SAssetView::CreateListView()
 			.DefaultLabel( LOCTEXT("Column_Name", "Name") )
 			.ShouldGenerateWidget(true) // Can't hide name column, so at least one column is visible
 		);
+
+	if (UE::Editor::ContentBrowser::IsNewStyleEnabled())
+	{
+		{
+			const bool bIsColumnVisible = !HiddenColumnNames.Contains(SortManager->RevisionControlColumnId.ToString());
+			NewListView->GetHeaderRow()->SetShowGeneratedColumn(SortManager->RevisionControlColumnId, bIsColumnVisible);
+		}
+
+		NewListView->GetHeaderRow()->SetOnGetMaxRowSizeForColumn(FOnGetMaxRowSizeForColumn::CreateRaw(NewListView.Get(), &SAssetColumnView::GetMaxRowSizeForColumn));
+
+		if (bShowTypeInColumnView)
+		{
+			NewListView->GetHeaderRow()->AddColumn(
+					SHeaderRow::Column(SortManager->ClassColumnId)
+					.FillWidth(160)
+					.SortMode(TAttribute<EColumnSortMode::Type>::Create(TAttribute< EColumnSortMode::Type >::FGetter::CreateSP(this, &SAssetView::GetColumnSortMode, SortManager->ClassColumnId)))
+					.SortPriority(TAttribute<EColumnSortPriority::Type>::Create(TAttribute<EColumnSortPriority::Type>::FGetter::CreateSP(this, &SAssetView::GetColumnSortPriority, SortManager->ClassColumnId)))
+					.OnSort(FOnSortModeChanged::CreateSP(this, &SAssetView::OnSortColumnHeader))
+					.DefaultLabel(LOCTEXT("Column_Class", "Type"))
+				);
+
+			const bool bIsColumnVisible = !HiddenColumnNames.Contains(SortManager->ClassColumnId.ToString());
+		
+			NewListView->GetHeaderRow()->SetShowGeneratedColumn(SortManager->ClassColumnId, bIsColumnVisible);
+		}
+
+		if (bShowPathInColumnView)
+		{
+			NewListView->GetHeaderRow()->AddColumn(
+					SHeaderRow::Column(SortManager->PathColumnId)
+					.FillWidth(160)
+					.SortMode(TAttribute<EColumnSortMode::Type>::Create(TAttribute< EColumnSortMode::Type >::FGetter::CreateSP(this, &SAssetView::GetColumnSortMode, SortManager->PathColumnId)))
+					.SortPriority(TAttribute<EColumnSortPriority::Type>::Create(TAttribute<EColumnSortPriority::Type>::FGetter::CreateSP(this, &SAssetView::GetColumnSortPriority, SortManager->PathColumnId)))
+					.OnSort(FOnSortModeChanged::CreateSP(this, &SAssetView::OnSortColumnHeader))
+					.DefaultLabel(LOCTEXT("Column_Path", "Path"))
+				);
+
+			const bool bIsColumnVisible = !HiddenColumnNames.Contains(SortManager->PathColumnId.ToString());
+		
+			NewListView->GetHeaderRow()->SetShowGeneratedColumn(SortManager->PathColumnId, bIsColumnVisible);
+		}
+	}
+	return NewListView.ToSharedRef();
 }
 
 TSharedRef<SAssetColumnView> SAssetView::CreateColumnView()
@@ -3631,7 +3708,8 @@ bool SAssetView::IsOrganizingFolders() const
 
 void SAssetView::SetMajorityAssetType(FName NewMajorityAssetType)
 {
-	if (CurrentViewType != EAssetViewType::Column)
+	if (CurrentViewType != EAssetViewType::Column &&
+		!(UE::Editor::ContentBrowser::IsNewStyleEnabled() && CurrentViewType == EAssetViewType::List))
 	{
 		return;
 	}
@@ -3657,9 +3735,18 @@ void SAssetView::SetMajorityAssetType(FName NewMajorityAssetType)
 		MajorityAssetType = NewMajorityAssetType;
 
 		TArray<FName> AddedColumns;
+		TSharedPtr<SListView<TSharedPtr<FAssetViewItem>>> ViewToUse = ColumnView;
+
+		if (UE::Editor::ContentBrowser::IsNewStyleEnabled())
+		{
+			if (CurrentViewType == EAssetViewType::List)
+			{
+				ViewToUse = ListView;
+			}
+		}
 
 		// Since the asset type has changed, remove all columns except name and class
-		const TIndirectArray<SHeaderRow::FColumn>& Columns = ColumnView->GetHeaderRow()->GetColumns();
+		const TIndirectArray<SHeaderRow::FColumn>& Columns = ViewToUse->GetHeaderRow()->GetColumns();
 
 		for ( int32 ColumnIdx = Columns.Num() - 1; ColumnIdx >= 0; --ColumnIdx )
 		{
@@ -3667,7 +3754,7 @@ void SAssetView::SetMajorityAssetType(FName NewMajorityAssetType)
 
 			if ( ColumnId != NAME_None && !IsFixedColumn(ColumnId) )
 			{
-				ColumnView->GetHeaderRow()->RemoveColumn(ColumnId);
+				ViewToUse->GetHeaderRow()->RemoveColumn(ColumnId);
 			}
 		}
 
@@ -3703,7 +3790,7 @@ void SAssetView::SetMajorityAssetType(FName NewMajorityAssetType)
 			}
 			AddedColumns.Add(TagName);
 
-			ColumnView->GetHeaderRow()->AddColumn(
+			ViewToUse->GetHeaderRow()->AddColumn(
 				SHeaderRow::Column(TagName)
 				.SortMode(TAttribute< EColumnSortMode::Type >::Create(TAttribute< EColumnSortMode::Type >::FGetter::CreateSP(this, &SAssetView::GetColumnSortMode, TagName)))
 				.SortPriority(TAttribute< EColumnSortPriority::Type >::Create(TAttribute< EColumnSortPriority::Type >::FGetter::CreateSP(this, &SAssetView::GetColumnSortPriority, TagName)))
@@ -3714,7 +3801,7 @@ void SAssetView::SetMajorityAssetType(FName NewMajorityAssetType)
 
 			const bool bIsColumnVisible = !HiddenColumnNames.Contains(TagName.ToString());
 		
-			ColumnView->GetHeaderRow()->SetShowGeneratedColumn(TagName, bIsColumnVisible);
+			ViewToUse->GetHeaderRow()->SetShowGeneratedColumn(TagName, bIsColumnVisible);
 
 			// If we found a tag the matches the column we are currently sorting on, there will be no need to change the column
 			for (int32 SortIdx = 0; SortIdx < CurrentSortOrder.Num(); SortIdx++)
@@ -3780,7 +3867,7 @@ void SAssetView::SetMajorityAssetType(FName NewMajorityAssetType)
 						}
 						AddedColumns.Add(TagPair.Key);
 
-						ColumnView->GetHeaderRow()->AddColumn(
+						ViewToUse->GetHeaderRow()->AddColumn(
 							SHeaderRow::Column(TagPair.Key)
 							.SortMode(TAttribute< EColumnSortMode::Type >::Create(TAttribute< EColumnSortMode::Type >::FGetter::CreateSP(this, &SAssetView::GetColumnSortMode, TagPair.Key)))
 							.SortPriority(TAttribute< EColumnSortPriority::Type >::Create(TAttribute< EColumnSortPriority::Type >::FGetter::CreateSP(this, &SAssetView::GetColumnSortPriority, TagPair.Key)))
@@ -3791,7 +3878,7 @@ void SAssetView::SetMajorityAssetType(FName NewMajorityAssetType)
 
 						const bool bIsColumnVisible = !HiddenColumnNames.Contains(TagPair.Key.ToString());
 		
-						ColumnView->GetHeaderRow()->SetShowGeneratedColumn(TagPair.Key, bIsColumnVisible);
+						ViewToUse->GetHeaderRow()->SetShowGeneratedColumn(TagPair.Key, bIsColumnVisible);
 						
 						// If we found a tag the matches the column we are currently sorting on, there will be no need to change the column
 						for (int32 SortIdx = 0; SortIdx < CurrentSortOrder.Num(); SortIdx++)
@@ -4627,8 +4714,22 @@ void SAssetView::SetCurrentViewType(EAssetViewType::Type NewType)
 		}
 		else if ( NewType == EAssetViewType::List )
 		{
-			CurrentThumbnailSize = ListViewThumbnailSize;
-			bPendingUpdateThumbnails = true;
+			if (UE::Editor::ContentBrowser::IsNewStyleEnabled())
+			{
+				if (ThumbnailSize >= EThumbnailSize::Large)
+				{
+					CurrentThumbnailSize = ListViewThumbnailSize;
+					bPendingUpdateThumbnails = true;
+				}
+				MajorityAssetType = NAME_None;
+				RefreshFilteredItems();
+				SortList();
+			}
+			else
+			{
+				CurrentThumbnailSize = ListViewThumbnailSize;
+				bPendingUpdateThumbnails = true;
+			}
 		}
 		else if ( NewType == EAssetViewType::Column )
 		{
@@ -4812,17 +4913,22 @@ TSharedRef<ITableRow> SAssetView::MakeListViewWidget(TSharedPtr<FAssetViewItem>
 
 	VisibleItems.Add(AssetItem);
 	bPendingUpdateThumbnails = true;
-
+	const FTableRowStyle& ListViewStyle = UE::Editor::ContentBrowser::IsNewStyleEnabled()
+		                               ? UE::ContentBrowser::Private::FContentBrowserStyle::Get().GetWidgetStyle<FTableRowStyle>("ContentBrowser.AssetListView.ColumnListTableRow")
+		                               : FAppStyle::GetWidgetStyle<FTableRowStyle>("ContentBrowser.AssetListView.ColumnListTableRow");
 	if (AssetItem->IsFolder())
 	{
 		return
-			SNew( SAssetListViewRow, OwnerTable )
-			.OnDragDetected( this, &SAssetView::OnDraggingAssetItem )
-			.Cursor( bAllowDragging ? EMouseCursor::GrabHand : EMouseCursor::Default )
+			SNew(SAssetListViewRow, OwnerTable)
+			.Style(&ListViewStyle)
+			.OnDragDetected(this, &SAssetView::OnDraggingAssetItem)
+			.Cursor(bAllowDragging ? EMouseCursor::GrabHand : EMouseCursor::Default)
+			.Padding(this, &SAssetView::GetListViewItemPadding)
 			.AssetListItem(
 				SNew(SAssetListItem)
 					.AssetItem(AssetItem)
 					.ItemHeight(this, &SAssetView::GetListViewItemHeight)
+					.CurrentThumbnailSize(this, &SAssetView::GetThumbnailSize)
 					.OnRenameBegin(this, &SAssetView::AssetRenameBegin)
 					.OnRenameCommit(this, &SAssetView::AssetRenameCommit)
 					.OnVerifyRenameCommit(this, &SAssetView::AssetVerifyRenameCommit)
@@ -4842,15 +4948,18 @@ TSharedRef<ITableRow> SAssetView::MakeListViewWidget(TSharedPtr<FAssetViewItem>
 		}
 		
 		return
-			SNew( SAssetListViewRow, OwnerTable )
-			.OnDragDetected( this, &SAssetView::OnDraggingAssetItem )
-			.Cursor( bAllowDragging ? EMouseCursor::GrabHand : EMouseCursor::Default )
+			SNew(SAssetListViewRow, OwnerTable)
+			.Style(&ListViewStyle)
+			.OnDragDetected(this, &SAssetView::OnDraggingAssetItem)
+			.Cursor(bAllowDragging ? EMouseCursor::GrabHand : EMouseCursor::Default)
+			.Padding(this, &SAssetView::GetListViewItemPadding)
 			.AssetListItem(
 				SNew(SAssetListItem)
 					.AssetThumbnail(AssetThumbnail)
 					.AssetItem(AssetItem)
 					.ThumbnailPadding((float)ListViewThumbnailPadding)
 					.ItemHeight(this, &SAssetView::GetListViewItemHeight)
+					.CurrentThumbnailSize(this, &SAssetView::GetThumbnailSize)
 					.OnRenameBegin(this, &SAssetView::AssetRenameBegin)
 					.OnRenameCommit(this, &SAssetView::AssetRenameCommit)
 					.OnVerifyRenameCommit(this, &SAssetView::AssetVerifyRenameCommit)
@@ -4858,9 +4967,9 @@ TSharedRef<ITableRow> SAssetView::MakeListViewWidget(TSharedPtr<FAssetViewItem>
 					.ShouldAllowToolTip(this, &SAssetView::ShouldAllowToolTips)
 					.HighlightText(HighlightedText)
 					.ThumbnailEditMode(this, &SAssetView::IsThumbnailEditMode)
-					.ThumbnailLabel( ThumbnailLabel )
+					.ThumbnailLabel(ThumbnailLabel)
 					.ThumbnailHintColorAndOpacity( this, &SAssetView::GetThumbnailHintColorAndOpacity )
-					.AllowThumbnailHintLabel( AllowThumbnailHintLabel )
+					.AllowThumbnailHintLabel(AllowThumbnailHintLabel)
 					.OnIsAssetValidForCustomToolTip(OnIsAssetValidForCustomToolTip)
 					.OnGetCustomAssetToolTip(OnGetCustomAssetToolTip)
 					.OnVisualizeAssetToolTip(OnVisualizeAssetToolTip)
@@ -5586,30 +5695,37 @@ void SAssetView::UpdateThumbnailSizeValue()
 	case EThumbnailSize::Tiny:
 		MinThumbnailSize = 64.f;
 		MaxThumbnailSize = 80.f;
+		ListViewItemHeight = 22.f;
 		break;
 	case EThumbnailSize::Small:
 		MinThumbnailSize = 80.f;
 		MaxThumbnailSize = 96.f;
+		ListViewItemHeight = 24.f;
 		break;
 	case EThumbnailSize::Medium:
 		MinThumbnailSize = 96.f;
 		MaxThumbnailSize = 112.f;
+		ListViewItemHeight = 32.f;
 		break;
 	case EThumbnailSize::Large:
 		MinThumbnailSize = 112.f;
 		MaxThumbnailSize = 128.f;
+		ListViewItemHeight = 48.f;
 		break;
 	case EThumbnailSize::XLarge:
 		MinThumbnailSize = 128.f;
 		MaxThumbnailSize = 136.f;
+		ListViewItemHeight = 64.f;
 		break;
 	case EThumbnailSize::Huge:
 		MinThumbnailSize = 136.f;
 		MaxThumbnailSize = 160.f;
+		ListViewItemHeight = 80.f;
 		break;
 	default:
 		MinThumbnailSize = 64.f;
 		MaxThumbnailSize = 80.f;
+		ListViewItemHeight = 22.f;
 		break;
 	}
 }
@@ -5666,7 +5782,16 @@ float SAssetView::GetSourceControlIconHeight() const
 
 float SAssetView::GetListViewItemHeight() const
 {
-	return (float)(ListViewThumbnailSize + ListViewThumbnailPadding * 2) * FMath::Lerp(MinThumbnailScale, MaxThumbnailScale, GetThumbnailScale());
+	return UE::Editor::ContentBrowser::IsNewStyleEnabled()
+		? ListViewItemHeight
+		: (float)(ListViewThumbnailSize + ListViewThumbnailPadding * 2) * FMath::Lerp(MinThumbnailScale, MaxThumbnailScale, GetThumbnailScale());
+}
+
+FMargin SAssetView::GetListViewItemPadding() const
+{
+	return UE::Editor::ContentBrowser::IsNewStyleEnabled() && GetThumbnailSize() >= EThumbnailSize::Large
+		? FMargin(0.f, (float)ListViewThumbnailPadding)
+		: FMargin(0.f);
 }
 
 float SAssetView::GetTileViewItemHeight() const
@@ -6011,30 +6136,59 @@ void SAssetView::ExportColumns()
 
 void SAssetView::OnHiddenColumnsChanged()
 {
-	// Early out if this is called before creation (due to loading config etc)
-	if(!ColumnView)
+	if (UE::Editor::ContentBrowser::IsNewStyleEnabled())
 	{
-		return;
-	}
+		// Early out if this is called before creation (due to loading config etc)
+		if (!ListView)
+		{
+			return;
+		}
 
-	// We can't directly update the hidden columns list, because some columns maybe hidden, but not created yet
-	TArray<FName> NewHiddenColumns = ColumnView->GetHeaderRow()->GetHiddenColumnIds();
+		// We can't directly update the hidden columns list, because some columns maybe hidden, but not created yet
+		TArray<FName> NewHiddenColumns = ListView->GetHeaderRow()->GetHiddenColumnIds();
 
-	// So instead for each column that currently exists, we update its visibility state in the HiddenColumnNames array
-	for (const SHeaderRow::FColumn& Column : ColumnView->GetHeaderRow()->GetColumns())
-	{
-		const bool bIsColumnVisible = NewHiddenColumns.Contains(Column.ColumnId);
+		// So instead for each column that currently exists, we update its visibility state in the HiddenColumnNames array
+		for (const SHeaderRow::FColumn& ListColumn : ListView->GetHeaderRow()->GetColumns())
+		{
+			const bool bIsColumnVisible = NewHiddenColumns.Contains(ListColumn.ColumnId);
 
-		if(bIsColumnVisible)
+			if(bIsColumnVisible)
+			{
+				HiddenColumnNames.AddUnique(ListColumn.ColumnId.ToString());
+			}
+			else
+			{
+				HiddenColumnNames.Remove(ListColumn.ColumnId.ToString());
+			}
+		}
+	}
+	else
+	{
+		// Early out if this is called before creation (due to loading config etc)
+		if (!ColumnView)
 		{
-			HiddenColumnNames.AddUnique(Column.ColumnId.ToString());
+			return;
 		}
-		else
+
+		// We can't directly update the hidden columns list, because some columns maybe hidden, but not created yet
+		TArray<FName> NewHiddenColumns = ColumnView->GetHeaderRow()->GetHiddenColumnIds();
+
+		// So instead for each column that currently exists, we update its visibility state in the HiddenColumnNames array
+		for (const SHeaderRow::FColumn& Column : ColumnView->GetHeaderRow()->GetColumns())
 		{
-			HiddenColumnNames.Remove(Column.ColumnId.ToString());
+			const bool bIsColumnVisible = NewHiddenColumns.Contains(Column.ColumnId);
+
+			if(bIsColumnVisible)
+			{
+				HiddenColumnNames.AddUnique(Column.ColumnId.ToString());
+			}
+			else
+			{
+				HiddenColumnNames.Remove(Column.ColumnId.ToString());
+			}
 		}
 	}
-	
+
 	if (FAssetViewInstanceConfig* Config = GetAssetViewConfig())
 	{
 		Config->HiddenColumns.Reset();
diff --git a/Engine/Source/Editor/ContentBrowser/Public/AssetViewWidgets.h b/Engine/Source/Editor/ContentBrowser/Public/AssetViewWidgets.h
--- a/Engine/Source/Editor/ContentBrowser/Public/AssetViewWidgets.h
+++ b/Engine/Source/Editor/ContentBrowser/Public/AssetViewWidgets.h
@@ -72,12 +72,12 @@ class SAssetTileItem;
 struct FAssetViewItemHelper
 {
 public:
-	static TSharedRef<SWidget> CreateListItemContents(SAssetListItem* const InListItem, const TSharedRef<SWidget>& InThumbnail, FName& OutItemShadowBorder);
-	static TSharedRef<SWidget> CreateTileItemContents(SAssetTileItem* const InTileItem, const TSharedRef<SWidget>& InThumbnail, FName& OutItemShadowBorder);
+	static TSharedRef<SWidget> CreateListItemContents(SAssetListItem* const InListItem, const TSharedRef<SWidget>& InThumbnail, FName& OutItemShadowBorder, const FMargin& InFolderOverlayPadding = FMargin(0.f));
+	static TSharedRef<SWidget> CreateTileItemContents(SAssetTileItem* const InTileItem, const TSharedRef<SWidget>& InThumbnail, FName& OutItemShadowBorder, const FMargin& InFolderOverlayPadding = FMargin(0.f));
 
 private:
 	template <typename T>
-	static TSharedRef<SWidget> CreateListTileItemContents(T* const InTileOrListItem, const TSharedRef<SWidget>& InThumbnail, FName& OutItemShadowBorder);
+	static TSharedRef<SWidget> CreateListTileItemContents(T* const InTileOrListItem, const TSharedRef<SWidget>& InThumbnail, FName& OutItemShadowBorder, const FMargin& InFolderOverlayPadding);
 };
 
 /** The tile view mode of the asset view */
@@ -369,6 +369,7 @@ public:
 	SLATE_BEGIN_ARGS( SAssetListItem )
 		: _ThumbnailPadding(0)
 		, _ThumbnailLabel( EThumbnailLabel::ClassName )
+		, _CurrentThumbnailSize(EThumbnailSize::Medium)
 		, _ThumbnailHintColorAndOpacity( FLinearColor( 0.0f, 0.0f, 0.0f, 0.0f ) )
 		, _ItemHeight(16)
 		, _ShouldAllowToolTip(true)
@@ -388,6 +389,9 @@ public:
 		/** The contents of the label displayed on the thumbnail */
 		SLATE_ARGUMENT( EThumbnailLabel::Type, ThumbnailLabel )
 
+		/** Current size of the thumbnail that was generated */
+		SLATE_ATTRIBUTE(EThumbnailSize, CurrentThumbnailSize)
+
 		/**  */
 		SLATE_ATTRIBUTE( FLinearColor, ThumbnailHintColorAndOpacity )
 
@@ -457,9 +461,16 @@ public:
 	bool CanDisplayPrimitiveTools() const { return false; }
 
 	/** Generates a widget for a particular column */
-	virtual TSharedRef<SWidget> GenerateWidgetForColumn( const FName& ColumnName, FIsSelected InIsSelected );
+	virtual TSharedRef<SWidget> GenerateWidgetForColumn(const FName& ColumnName, FIsSelected InIsSelected);
+
+protected:
+	/** Generic ThumbnailSize for the Thumbnail */
+	int32 GetGenericThumbnailSize() const;
 
 private:
+	/** New Style columns generation */
+	TSharedRef<SWidget> GenerateWidgetForColumnNewStyle(const FName& ColumnName, FIsSelected InIsSelected);
+	
 	/** Get the expected width of an extra state icon. */
 	float GetExtraStateIconWidth() const;
 
@@ -474,6 +485,16 @@ private:
 
 	/** Get the text color for the columns */
 	FSlateColor GetColumnTextColor(FIsSelected InIsSelected) const;
+	
+	/** Gets the path to this asset */
+	FText GetAssetPathText() const;
+
+	/** Gets the value for the specified tag in this asset */
+	FText GetAssetTagText(FName InTag) const;
+
+	/** Get the index to use for the Thumbnail WidgetSwitcher, based on the ThumbnailSize */
+	int32 GetWidgetSwitcherIndex() const;
+
 private:
 	/** The handle to the thumbnail that this item is rendering */
 	TSharedPtr<FAssetThumbnail> AssetThumbnail;
@@ -492,16 +513,25 @@ private:
 
 	/** Whether the item is selected in the view without anything else being selected*/
 	FIsSelected IsSelectedExclusively;
+	
+	/** Current thumbnail size when this widget was generated */
+	TAttribute<EThumbnailSize> CurrentThumbnailSize;
+
+	/** Path text widget, need to be updated when the AssetData changes */
+	TSharedPtr<STextBlock> PathTextWidget;
 };
 
 class SAssetListViewRow : public SMultiColumnTableRow< TSharedPtr<FAssetViewItem> >
 {
 public:
 	SLATE_BEGIN_ARGS( SAssetListViewRow )
+		: _Padding(FMargin(0.f))
 		{}
-
+		
+		SLATE_STYLE_ARGUMENT( FTableRowStyle, Style )
 		SLATE_EVENT( FOnDragDetected, OnDragDetected )
 		SLATE_ARGUMENT( TSharedPtr<SAssetListItem>, AssetListItem )
+		SLATE_ATTRIBUTE( FMargin, Padding )
 
 	SLATE_END_ARGS()
 
@@ -513,7 +543,8 @@ public:
 
 		SMultiColumnTableRow< TSharedPtr<FAssetViewItem> >::Construct( 
 			FSuperRowType::FArguments()
-				.Style(FAppStyle::Get(), "ContentBrowser.AssetListView.ColumnListTableRow")
+				.Style(InArgs._Style)
+				.Padding(InArgs._Padding)
 				.OnDragDetected(InArgs._OnDragDetected), 
 			InOwnerTableView);
 		Content = this->AssetListItem;
@@ -736,9 +767,6 @@ protected:
 	/** Returns the font to use for the thumbnail label */
 	FSlateFontInfo GetThumbnailFont() const;
 
-	/** Gets the max size of the name area for an asset */
-	FOptionalSize GetMaxAssetNameHeight() const;
-
 	const FSlateBrush* GetFolderSlotBackgroundImage() const;
 	const FSlateBrush* GetFolderBackgroundImage() const;
 	const FSlateBrush* GetFolderBackgroundShadowImage() const;
@@ -750,6 +778,7 @@ protected:
 
 	FOptionalSize GetNameAreaMaxDesiredHeight() const;
 
+	/** Generic ThumbnailSize for the Thumbnail */
 	int32 GetGenericThumbnailSize() const;
 
 	/** Gets the visibility of the SCC icons */
diff --git a/Engine/Source/Editor/ContentBrowser/Public/SAssetView.h b/Engine/Source/Editor/ContentBrowser/Public/SAssetView.h
--- a/Engine/Source/Editor/ContentBrowser/Public/SAssetView.h
+++ b/Engine/Source/Editor/ContentBrowser/Public/SAssetView.h
@@ -771,7 +771,6 @@ private:
 	/** Handler for column view widget creation */
 	TSharedRef<ITableRow> MakeColumnViewWidget(TSharedPtr<FAssetViewItem> AssetItem, const TSharedRef<STableViewBase>& OwnerTable);
 
-
 	/** Handler for when any asset item widget gets destroyed */
 	void AssetItemWidgetDestroyed(const TSharedPtr<FAssetViewItem>& Item);
 	
@@ -864,7 +863,10 @@ private:
 
 	/** Gets the scaled item height for the list view */
 	float GetListViewItemHeight() const;
-	
+
+	/** Gets the padding for the list view rows */
+	FMargin GetListViewItemPadding() const;
+
 	/** Gets the final scaled item height for the tile view */
 	float GetTileViewItemHeight() const;
 
@@ -1144,6 +1146,7 @@ private:
 	int32 ListViewThumbnailResolution;
 	int32 ListViewThumbnailSize;
 	int32 ListViewThumbnailPadding;
+	float ListViewItemHeight;
 	int32 TileViewThumbnailResolution;
 	int32 TileViewThumbnailSize;
 	int32 TileViewThumbnailPadding;
diff --git a/Engine/Source/Editor/UnrealEd/Private/AssetThumbnail.cpp b/Engine/Source/Editor/UnrealEd/Private/AssetThumbnail.cpp
--- a/Engine/Source/Editor/UnrealEd/Private/AssetThumbnail.cpp
+++ b/Engine/Source/Editor/UnrealEd/Private/AssetThumbnail.cpp
@@ -1661,6 +1661,7 @@ TSharedRef<SWidget> FAssetThumbnail::MakeThumbnailWidget( const FAssetThumbnailC
 			.AllowAssetStatusThumbnailOverlay(InConfig.bAllowAssetStatusThumbnailOverlay)
 			.ShowAssetChip(InConfig.bShowAssetChip)
 			.UseRoundAssetChip(InConfig.bUseRoundAssetChip)
+			.CanDisplayEditModePrimitiveTools(InConfig.bCanDisplayEditModePrimitiveTools)
 			.IsEditModeVisible(InConfig.IsEditModeVisible)
 			.SmallAssetChipBorder(InConfig.bSmallAssetChipBorder)
 			.AssetChipBorderImageOverride(InConfig.AssetChipBorderImageOverride)
