Support for EpicInternal in public published modules
#rnx

diff --git a/Engine/Source/Editor/ContentBrowser/Private/AssetContextMenu.cpp b/Engine/Source/Editor/ContentBrowser/Private/AssetContextMenu.cpp
--- a/Engine/Source/Editor/ContentBrowser/Private/AssetContextMenu.cpp
+++ b/Engine/Source/Editor/ContentBrowser/Private/AssetContextMenu.cpp
@@ -1593,6 +1593,8 @@ void FAssetContextMenu::CacheCanExecuteVars()
 	const TSharedPtr<FPathPermissionList>& ShowEpicInternalContentPermissionList = ContentBrowserModule.Get().GetShowEpicInternalContentPermissionList();
 	bool bAnyPathPassesShowEpicInternal = !ShowEpicInternalContentPermissionList->HasFiltering();
 
+	const TSharedPtr<FPathPermissionList>& ShowPrivateContentPermissionList = ContentBrowserModule.Get().GetShowPrivateContentPermissionList();
+
 	if (SelectedFiles.Num() == 1)
 	{
 		const FContentBrowserItem& SelectedItem = SelectedFiles[0];
@@ -1616,9 +1618,9 @@ void FAssetContextMenu::CacheCanExecuteVars()
 			if (IAssetTools::Get().CanAssetBePublic(AssetPath))
 			{
 				bCanExecuteSetPublicAsset = true;
-				if (ContentBrowserModule.Get().GetShowEpicInternalContentPermissionList()->PassesStartsWithFilter(AssetPath))
+				if (bIsEpicInternalAssetFeatureEnabled && (ShowEpicInternalContentPermissionList->PassesStartsWithFilter(AssetPath) || ShowPrivateContentPermissionList->PassesStartsWithFilter(AssetPath)))
 				{
-					bCanExecuteSetEpicInternalAsset = bIsEpicInternalAssetFeatureEnabled;
+					bCanExecuteSetEpicInternalAsset = true;
 				}
 				bCanExecuteSetPrivateAsset = true;
 			}
@@ -1660,9 +1662,12 @@ void FAssetContextMenu::CacheCanExecuteVars()
 					// If any selected can be set to Epic internal
 					if (!bAnyPathPassesShowEpicInternal && bIsEpicInternalAssetFeatureEnabled)
 					{
-						if (ShowEpicInternalContentPermissionList->PassesStartsWithFilter(AssetPath) && IAssetTools::Get().CanAssetBePublic(AssetPath))
+						if (IAssetTools::Get().CanAssetBePublic(AssetPath))
 						{
-							bAnyPathPassesShowEpicInternal = true;
+							if (ShowEpicInternalContentPermissionList->PassesStartsWithFilter(AssetPath) || ShowPrivateContentPermissionList->HasFiltering() && ShowPrivateContentPermissionList->PassesStartsWithFilter(AssetPath))
+							{
+								bAnyPathPassesShowEpicInternal = true;
+							}
 						}
 					}
 
