Compile fix
#rnx

diff --git a/Engine/Source/Editor/ContentBrowser/Private/AssetContextMenu.cpp b/Engine/Source/Editor/ContentBrowser/Private/AssetContextMenu.cpp
--- a/Engine/Source/Editor/ContentBrowser/Private/AssetContextMenu.cpp
+++ b/Engine/Source/Editor/ContentBrowser/Private/AssetContextMenu.cpp
@@ -1664,7 +1664,7 @@ void FAssetContextMenu::CacheCanExecuteVars()
 					{
 						if (IAssetTools::Get().CanAssetBePublic(AssetPath))
 						{
-							if (ShowEpicInternalContentPermissionList->PassesStartsWithFilter(AssetPath) || ShowPrivateContentPermissionList->HasFiltering() && ShowPrivateContentPermissionList->PassesStartsWithFilter(AssetPath))
+							if (ShowEpicInternalContentPermissionList->PassesStartsWithFilter(AssetPath) || (ShowPrivateContentPermissionList->HasFiltering() && ShowPrivateContentPermissionList->PassesStartsWithFilter(AssetPath)))
 							{
 								bAnyPathPassesShowEpicInternal = true;
 							}
