Added support for displaying Verse paths in asset tooltips.

diff --git a/Engine/Source/Editor/ContentBrowser/Private/AssetSystemContentBrowserInfoProvider.cpp b/Engine/Source/Editor/ContentBrowser/Private/AssetSystemContentBrowserInfoProvider.cpp
--- a/Engine/Source/Editor/ContentBrowser/Private/AssetSystemContentBrowserInfoProvider.cpp
+++ b/Engine/Source/Editor/ContentBrowser/Private/AssetSystemContentBrowserInfoProvider.cpp
@@ -9,11 +9,14 @@
 #include "AssetDefinitionAssetInfo.h"
 #include "AssetDefinitionRegistry.h"
 #include "AssetRegistry/AssetRegistryModule.h"
+#include "AssetToolsModule.h"
 #include "AssetViewTypes.h"
 #include "AutoReimport/AssetSourceFilenameCache.h"
 #include "CollectionManagerModule.h"
+#include "Containers/VersePath.h"
 #include "ContentBrowserDataSource.h"
 #include "ContentBrowserUtils.h"
+#include "IAssetTools.h"
 #include "ICollectionContainer.h"
 #include "ICollectionManager.h"
 #include "Misc/EngineBuildSettings.h"
@@ -78,6 +81,17 @@ void FAssetSystemContentBrowserInfoProvider::PopulateAssetInfo(TArray<FAssetDisp
 		}
 		OutAssetDisplayInfo.Add(PathInfo);
 
+		if (ItemAssetData.IsValid() && FAssetToolsModule::GetModule().Get().ShowingContentVersePath())
+		{
+			UE::Core::FVersePath VersePath = ItemAssetData.GetVersePath();
+			if (VersePath.IsValid())
+			{
+				FAssetDisplayInfo& VersePathInfo = OutAssetDisplayInfo.AddDefaulted_GetRef();
+				VersePathInfo.StatusTitle = LOCTEXT("TileViewTooltipVersePath", "Verse Path");
+				VersePathInfo.StatusDescription = FText::FromString(VersePath.ToString());
+			}
+		}
+
 		if (ItemAssetData.IsValid() && ItemAssetData.PackageName != NAME_None)
 		{
 			const FString PackagePathWithinRoot = ContentBrowserUtils::GetPackagePathWithinRoot(ItemAssetData.PackageName.ToString());
diff --git a/Engine/Source/Editor/ContentBrowser/Private/AssetViewWidgets.cpp b/Engine/Source/Editor/ContentBrowser/Private/AssetViewWidgets.cpp
--- a/Engine/Source/Editor/ContentBrowser/Private/AssetViewWidgets.cpp
+++ b/Engine/Source/Editor/ContentBrowser/Private/AssetViewWidgets.cpp
@@ -11,6 +11,7 @@
 #include "AssetRegistry/AssetData.h"
 #include "AssetTagItemTypes.h"
 #include "AssetThumbnail.h"
+#include "AssetToolsModule.h"
 #include "AssetViewTypes.h"
 #include "AssetViewUtils.h"
 #include "AutoReimport/AssetSourceFilenameCache.h"
@@ -21,6 +22,7 @@
 #include "Containers/StringFwd.h"
 #include "Containers/StringView.h"
 #include "Containers/UnrealString.h"
+#include "Containers/VersePath.h"
 #include "ContentBrowserDataSource.h"
 #include "ContentBrowserItem.h"
 #include "ContentBrowserItemData.h"
@@ -35,6 +37,7 @@
 #include "Framework/Text/TextLayout.h"
 #include "HAL/IConsoleManager.h"
 #include "HAL/PlatformApplicationMisc.h"
+#include "IAssetTools.h"
 #include "ICollectionContainer.h"
 #include "ICollectionManager.h"
 #include "ISourceControlModule.h"
@@ -832,6 +835,15 @@ TSharedRef<SWidget> SAssetViewItem::CreateToolTipWidget() const
 				AddToToolTipInfoBox(InfoBox, LOCTEXT("TileViewTooltipPath", "Path"), FText::FromName(AssetItem->GetItem().GetVirtualPath()), false);
 			}
 
+			if (ItemAssetData.IsValid() && FAssetToolsModule::GetModule().Get().ShowingContentVersePath())
+			{
+				UE::Core::FVersePath VersePath = ItemAssetData.GetVersePath();
+				if (VersePath.IsValid())
+				{
+					AddToToolTipInfoBox(InfoBox, LOCTEXT("TileViewTooltipVersePath", "Verse Path"), FText::FromString(VersePath.ToString()), false);
+				}
+			}
+
 			if (ItemAssetData.IsValid() && ItemAssetData.PackageName != NAME_None)
 			{
 				const FString PackagePathWithinRoot = ContentBrowserUtils::GetPackagePathWithinRoot(ItemAssetData.PackageName.ToString());
