[ContentBrowser] Favorites text scales correctly

diff --git a/Engine/Source/Editor/ContentBrowser/Private/Widgets/SContentBrowserSourceTree.cpp b/Engine/Source/Editor/ContentBrowser/Private/Widgets/SContentBrowserSourceTree.cpp
--- a/Engine/Source/Editor/ContentBrowser/Private/Widgets/SContentBrowserSourceTree.cpp
+++ b/Engine/Source/Editor/ContentBrowser/Private/Widgets/SContentBrowserSourceTree.cpp
@@ -364,7 +364,7 @@ namespace UE::Editor::ContentBrowser::Private
 					.Padding(0)
 					[
 						SNew(SBox)
-						.HAlign(HAlign_Center)
+						.HAlign(HAlign_Fill)
 						.VAlign(VAlign_Top)
 						.HeightOverride(MinBodyHeight)
 						.Padding(EmptyBodyLabelPadding)
@@ -373,6 +373,7 @@ namespace UE::Editor::ContentBrowser::Private
 							.Text(InArgs._EmptyBodyLabel)
 							.TextStyle(&FAppStyle::Get().GetWidgetStyle<FTextBlockStyle>("RichTextBlock.Italic"))
 							.AutoWrapText(true)
+							.Justification(ETextJustify::Center)
 							.DecoratorStyleSet(&FAppStyle::Get())
 							+ SRichTextBlock::ImageDecorator()
 						]
