[ContentBrowser] SourceTree collapse fix

diff --git a/Engine/Source/Editor/ContentBrowser/Private/Widgets/SContentBrowserSourceTree.cpp b/Engine/Source/Editor/ContentBrowser/Private/Widgets/SContentBrowserSourceTree.cpp
--- a/Engine/Source/Editor/ContentBrowser/Private/Widgets/SContentBrowserSourceTree.cpp
+++ b/Engine/Source/Editor/ContentBrowser/Private/Widgets/SContentBrowserSourceTree.cpp
@@ -76,6 +76,11 @@ namespace UE::Editor::ContentBrowser::Private
 			: 0.0f;
 	}
 
+	float SContentBrowserSourceTree::FSlot::GetHeaderHeight() const
+	{
+		return HeaderHeight;
+	}
+
 	float SContentBrowserSourceTree::FSlot::GetInitialSlotSize() const
 	{
 		return InitialSlotSize;
@@ -119,9 +124,49 @@ namespace UE::Editor::ContentBrowser::Private
 				[
 					InSlotAdded->GetEntryWidget().ToSharedRef()
 				];
+
+				// ... and re-calculate total header height
+				UpdateTotalHeaderHeight();
 			}};
 	}
 
+	int32 SContentBrowserSourceTree::RemoveSlot(const TSharedRef<SWidget>& InSlotWidget)
+	{
+		if (!Splitter.IsValid())
+		{
+			return INDEX_NONE;
+		}
+
+		if (Slots.Remove(InSlotWidget) != INDEX_NONE)
+		{
+			const int32 RemovedIndex = Splitter->RemoveSlot(InSlotWidget);
+
+			UpdateTotalHeaderHeight();
+
+			return RemovedIndex;
+		}
+
+		return INDEX_NONE;
+	}
+
+	void SContentBrowserSourceTree::ClearChildren()
+	{
+		Slots.Empty();
+		Splitter->ClearChildren();
+
+		UpdateTotalHeaderHeight();
+	}
+
+	int32 SContentBrowserSourceTree::NumSlots() const
+	{
+		return Slots.Num();
+	}
+
+	bool SContentBrowserSourceTree::IsValidSlotIndex(int32 InIndex) const
+	{
+		return Slots.IsValidIndex(InIndex);
+	}
+
 	float SContentBrowserSourceTree::GetMinHeight(TSharedPtr<SContentBrowserSourceTreeArea> InEntryWidget) const
 	{
 		for (int32 SlotIdx = 0; SlotIdx < Slots.Num(); ++SlotIdx)
@@ -135,6 +180,11 @@ namespace UE::Editor::ContentBrowser::Private
 		return 0.0f;
 	}
 
+	float SContentBrowserSourceTree::GetTotalHeaderHeight() const
+	{
+		return TotalHeaderHeight;
+	}
+
 	SSplitter::ESizeRule SContentBrowserSourceTree::GetExpandedSizeRule(TSharedPtr<SContentBrowserSourceTreeArea> InEntryWidget) const
 	{
 		for (int32 SlotIdx = 0; SlotIdx < Slots.Num(); ++SlotIdx)
@@ -166,6 +216,15 @@ namespace UE::Editor::ContentBrowser::Private
 		}
 	}
 
+	void SContentBrowserSourceTree::UpdateTotalHeaderHeight()
+	{
+		TotalHeaderHeight = 0.0f;
+		for (int32 SlotIdx = 0; SlotIdx < Slots.Num(); ++SlotIdx)
+		{
+			TotalHeaderHeight += Slots[SlotIdx].GetHeaderHeight();
+		}
+	}
+
 	void SContentBrowserSourceTree::Construct(const FArguments& InArgs)
 	{
 		using namespace UE::ContentBrowser::Private;
@@ -191,6 +250,11 @@ namespace UE::Editor::ContentBrowser::Private
 			.Visibility(Slot._Visibility);
 		}
 
+		const TAttribute<float> TotalHeaderHeightAttribute =
+			TAttribute<float>::Create(
+				TAttribute<float>::FGetter::CreateSP(
+					this, &SContentBrowserSourceTree::GetTotalHeaderHeight));
+
 		ChildSlot
 		.HAlign(HAlign_Fill)
 	    .VAlign(VAlign_Fill)
@@ -200,6 +264,7 @@ namespace UE::Editor::ContentBrowser::Private
 			+ SVerticalBox::Slot()
 			.HAlign(HAlign_Fill)
 			.VAlign(VAlign_Fill)
+			.MinHeight(TotalHeaderHeightAttribute)
 			.FillContentHeight(1.0f)
 			[
 				Splitter.ToSharedRef()
@@ -207,12 +272,10 @@ namespace UE::Editor::ContentBrowser::Private
 
 			+ SVerticalBox::Slot()
 			.HAlign(HAlign_Fill)
-			.VAlign(VAlign_Top)
+			.VAlign(VAlign_Fill)
 			.FillContentHeight(1.0f, 0.0f)
 			[
-				SNew(SColorBlock)
-				.Color(FLinearColor::Green)
-				.Size(FVector2D(0, 0))
+				SNullWidget::NullWidget
 			]
 		];
 	}
diff --git a/Engine/Source/Editor/ContentBrowser/Private/Widgets/SContentBrowserSourceTree.h b/Engine/Source/Editor/ContentBrowser/Private/Widgets/SContentBrowserSourceTree.h
--- a/Engine/Source/Editor/ContentBrowser/Private/Widgets/SContentBrowserSourceTree.h
+++ b/Engine/Source/Editor/ContentBrowser/Private/Widgets/SContentBrowserSourceTree.h
@@ -51,6 +51,7 @@ namespace UE::Editor::ContentBrowser::Private
 		public:
 			SSplitter::ESizeRule GetExpandedSizeRule() const;
 			float GetMinHeight() const;
+			float GetHeaderHeight() const;
 			float GetInitialSlotSize() const;
 			float GetSlotSize() const;
 			TSharedPtr<SContentBrowserSourceTreeArea> GetEntryWidget() const;
@@ -79,6 +80,7 @@ namespace UE::Editor::ContentBrowser::Private
 
 		SContentBrowserSourceTree()
 			: Slots(this)
+			, TotalHeaderHeight(0.0f)
 		{
 		}
 
@@ -93,15 +95,36 @@ namespace UE::Editor::ContentBrowser::Private
 
 		FScopedWidgetSlotArguments AddSlot(const int32 InAtIndex = INDEX_NONE);
 
+		/** Removes a slot from this panel which contains the specified SWidget
+		 *
+		 * @param InSlotWidget The widget to match when searching through the slots
+		 * @returns The index in the children array where the slot was removed and -1 if no slot was found matching the widget
+		 */
+		int32 RemoveSlot(const TSharedRef<SWidget>& InSlotWidget);
+
+		/** Removes all children from the panel. */
+		void ClearChildren();
+
+		/** @return the number of slots. */
+		int32 NumSlots() const;
+
+		/** @return if it's a valid index slot index. */
+		bool IsValidSlotIndex(int32 InIndex) const;
+
 	private:
 		float GetMinHeight(TSharedPtr<SContentBrowserSourceTreeArea> InEntryWidget) const;
+		float GetTotalHeaderHeight() const;
 		SSplitter::ESizeRule GetExpandedSizeRule(TSharedPtr<SContentBrowserSourceTreeArea> InEntryWidget) const;
 		float GetSlotSize(const int32 InSlotIdx) const;
 		void OnSlotResized(float InNewSize, const int32 InSlotIdx);
+		void UpdateTotalHeaderHeight();
 
 	private:
 		TSharedPtr<SSplitter> Splitter;
 		TPanelChildren<FSlot> Slots;
+
+		/** Cached total header height, recalculated one Slot Add/Remove. */
+		float TotalHeaderHeight;
 	};
 
 	/** Represents a single item in the source tree view, ie. "Favorites". */
