[ContentBrowser] ListView tweaks + AssetPicker default changes
Adding padding betweens ListView columns
Using Thumbnails even in small size
Add MiddleEllipsis to ListView
Change AssetPicker default to not create any Major Asset Columns

diff --git a/Engine/Source/Editor/ContentBrowser/Private/AssetViewWidgets.cpp b/Engine/Source/Editor/ContentBrowser/Private/AssetViewWidgets.cpp
--- a/Engine/Source/Editor/ContentBrowser/Private/AssetViewWidgets.cpp
+++ b/Engine/Source/Editor/ContentBrowser/Private/AssetViewWidgets.cpp
@@ -2126,6 +2126,7 @@ int32 SAssetListItem::GetGenericThumbnailSize() const
 TSharedRef<SWidget> SAssetListItem::GenerateWidgetForColumnNewStyle(const FName& ColumnName, FIsSelected InIsSelected)
 {
 	TSharedPtr<SWidget> Content;
+	const FMargin DefaultColumnPadding = FMargin(8.f, 0.f, 0.f, 0.f);
 
 	if (ColumnName == "Name")
 	{
@@ -2158,7 +2159,7 @@ TSharedRef<SWidget> SAssetListItem::GenerateWidgetForColumnNewStyle(const FName&
 
 		Content = SNew(SBorder)
 		.BorderImage(this, &SAssetViewItem::GetBorderImage)
-		.Padding(0.f)
+		.Padding(DefaultColumnPadding)
 		.AddMetaData<FTagMetaData>(FTagMetaData(AssetItem->GetItem().GetVirtualPath()))
 		[
 			SNew(SHorizontalBox)
@@ -2168,7 +2169,7 @@ TSharedRef<SWidget> SAssetListItem::GenerateWidgetForColumnNewStyle(const FName&
 			.AutoWidth()
 			.VAlign(VAlign_Center)
 			[
-				SNew( SBox )
+				SNew(SBox)
 				.WidthOverride( this, &SAssetListItem::GetThumbnailBoxSize )
 				.HeightOverride( this, &SAssetListItem::GetThumbnailBoxSize )
 				[
@@ -2194,19 +2195,19 @@ TSharedRef<SWidget> SAssetListItem::GenerateWidgetForColumnNewStyle(const FName&
 			]
 
 			+SHorizontalBox::Slot()
-			.AutoWidth()
 			.Padding(6, 0, 0, 0)
 			.VAlign(VAlign_Center)
 			[
 				SAssignNew(InlineRenameWidget, SInlineEditableTextBlock)
 				.Font(UE::ContentBrowser::Private::FContentBrowserStyle::Get().GetFontStyle("ContentBrowser.AssetTileViewNameFont"))
-				.Text( GetNameText() )
+				.Text(GetNameText())
 				.OnBeginTextEdit(this, &SAssetListItem::HandleBeginNameChange)
 				.OnTextCommitted(this, &SAssetListItem::HandleNameCommitted)
 				.OnVerifyTextChanged(this, &SAssetListItem::HandleVerifyNameChanged)
 				.HighlightText(HighlightText)
 				.IsSelected(IsSelectedExclusively)
 				.IsReadOnly(this, &SAssetListItem::IsNameReadOnly)
+				.OverflowPolicy(ETextOverflowPolicy::MiddleEllipsis)
 				.ColorAndOpacity(this, &SAssetListItem::GetColumnTextColor, InIsSelected)
 			]
 		];
@@ -2232,29 +2233,35 @@ TSharedRef<SWidget> SAssetListItem::GenerateWidgetForColumnNewStyle(const FName&
 	else if (ColumnName == "Class")
 	{
 		Content = SNew(SBox)
+			.Padding(DefaultColumnPadding)
 			.VAlign(VAlign_Center)
 			[
 				SAssignNew(ClassTextWidget, STextBlock)
 				.Text(GetAssetClassText())
+				.OverflowPolicy(ETextOverflowPolicy::MiddleEllipsis)
 				.HighlightText(HighlightText)
 			];
 	}
 	else if (ColumnName == "Path")
 	{
-		Content =  SNew(SBox)
+		Content = SNew(SBox)
+			.Padding(DefaultColumnPadding)
 			.VAlign(VAlign_Center)
 			[
 				SAssignNew(PathTextWidget, STextBlock)
 				.Text(GetAssetPathText())
+				.OverflowPolicy(ETextOverflowPolicy::MiddleEllipsis)
 				.HighlightText(HighlightText)
 			];
 	}
 	else
 	{
-		Content =  SNew(SBox)
+		Content = SNew(SBox)
+			.Padding(DefaultColumnPadding)
 			.VAlign(VAlign_Center)
 			[
 				SNew(STextBlock)
+				.OverflowPolicy(ETextOverflowPolicy::MiddleEllipsis)
 				.Text(TAttribute<FText>::Create(TAttribute<FText>::FGetter::CreateSP(this, &SAssetListItem::GetAssetTagText, ColumnName)))
 			];
 	}
diff --git a/Engine/Source/Editor/ContentBrowser/Private/SAssetPicker.cpp b/Engine/Source/Editor/ContentBrowser/Private/SAssetPicker.cpp
--- a/Engine/Source/Editor/ContentBrowser/Private/SAssetPicker.cpp
+++ b/Engine/Source/Editor/ContentBrowser/Private/SAssetPicker.cpp
@@ -377,6 +377,7 @@ void SAssetPicker::Construct( const FArguments& InArgs )
 		.AssetShowWarningText( InArgs._AssetPickerConfig.AssetShowWarningText)
 		.AllowFocusOnSync(false)	// Stop the asset view from stealing focus (we're in control of that)
 		.HiddenColumnNames(InArgs._AssetPickerConfig.HiddenColumnNames)
+		.AvoidMajorAssetColumnsCreation(InArgs._AssetPickerConfig.AvoidMajorAssetColumnsCreation)
 		.CustomColumns(InArgs._AssetPickerConfig.CustomColumns)
 		.OnSearchOptionsChanged(this, &SAssetPicker::HandleSearchSettingsChanged)
 		.InitialThumbnailSize(InArgs._AssetPickerConfig.InitialThumbnailSize)
diff --git a/Engine/Source/Editor/ContentBrowser/Private/SAssetView.cpp b/Engine/Source/Editor/ContentBrowser/Private/SAssetView.cpp
--- a/Engine/Source/Editor/ContentBrowser/Private/SAssetView.cpp
+++ b/Engine/Source/Editor/ContentBrowser/Private/SAssetView.cpp
@@ -1186,6 +1186,7 @@ void SAssetView::Construct( const FArguments& InArgs )
 	bForceShowPluginContent = InArgs._ForceShowPluginContent;
 	bForceHideScrollbar = InArgs._ForceHideScrollbar;
 	bShowDisallowedAssetClassAsUnsupportedItems = InArgs._ShowDisallowedAssetClassAsUnsupportedItems;
+	bAvoidMajorAssetColumnsCreation = InArgs._AvoidMajorAssetColumnsCreation;
 
 	bPendingUpdateThumbnails = false;
 	bShouldNotifyNextAssetSync = true;
@@ -3708,8 +3709,9 @@ bool SAssetView::IsOrganizingFolders() const
 
 void SAssetView::SetMajorityAssetType(FName NewMajorityAssetType)
 {
-	if (CurrentViewType != EAssetViewType::Column &&
-		!(UE::Editor::ContentBrowser::IsNewStyleEnabled() && CurrentViewType == EAssetViewType::List))
+	if ((CurrentViewType != EAssetViewType::Column &&
+		!(UE::Editor::ContentBrowser::IsNewStyleEnabled() && CurrentViewType == EAssetViewType::List)) ||
+		(UE::Editor::ContentBrowser::IsNewStyleEnabled() && bAvoidMajorAssetColumnsCreation))
 	{
 		return;
 	}
diff --git a/Engine/Source/Editor/ContentBrowser/Public/IContentBrowserSingleton.h b/Engine/Source/Editor/ContentBrowser/Public/IContentBrowserSingleton.h
--- a/Engine/Source/Editor/ContentBrowser/Public/IContentBrowserSingleton.h
+++ b/Engine/Source/Editor/ContentBrowser/Public/IContentBrowserSingleton.h
@@ -329,6 +329,9 @@ struct FAssetPickerConfig
 	/** If true, sort by path in column view. Only works if initial view type is Column */
 	bool bSortByPathInColumnView;
 
+	/** If true, will avoid adding the major assets column creation in the ListView */
+	bool AvoidMajorAssetColumnsCreation;
+
 	/** Can be used to freely change the Add Filter menu. This is executed on the dynamically instanced menu, not the registered menu. */
 	TDelegate<void(UToolMenu*)> OnExtendAddFilterMenu;
 	
@@ -364,6 +367,7 @@ struct FAssetPickerConfig
 		, bShowPathInColumnView(false)
 		, bShowTypeInColumnView(true)
 		, bSortByPathInColumnView(false)
+		, AvoidMajorAssetColumnsCreation(false)
 		, DefaultFilterMenuExpansion(EAssetTypeCategories::Basic)
 		, bUseSectionsForCustomFilterCategories(false)
 	{}
diff --git a/Engine/Source/Editor/ContentBrowser/Public/SAssetView.h b/Engine/Source/Editor/ContentBrowser/Public/SAssetView.h
--- a/Engine/Source/Editor/ContentBrowser/Public/SAssetView.h
+++ b/Engine/Source/Editor/ContentBrowser/Public/SAssetView.h
@@ -298,6 +298,9 @@ public:
 		/** Allow the asset view to display a custom view registered with the Content Browser module */
 		SLATE_ARGUMENT(bool, AllowCustomView)
 
+		/** If true avoid the creation of all the Major Asset columns */
+		SLATE_ARGUMENT(bool, AvoidMajorAssetColumnsCreation)
+
 		/** Called to check if an asset tag should be display in details view. */
 		SLATE_EVENT( FOnShouldDisplayAssetTag, OnAssetTagWantsToBeDisplayed )
 
@@ -1241,6 +1244,9 @@ private:
 	/** Whether or not to notify about newly selected items on on the next asset sync */
 	bool bShouldNotifyNextAssetSync : 1;
 
+	/** Whether or not the major asset columns should be skipped in the ListView */
+	bool bAvoidMajorAssetColumnsCreation : 1;
+
 	/** The current selection mode used by the asset view */
 	ESelectionMode::Type SelectionMode;
 
diff --git a/Engine/Source/Editor/PropertyEditor/Private/PropertyEditorHelpers.h b/Engine/Source/Editor/PropertyEditor/Private/PropertyEditorHelpers.h
--- a/Engine/Source/Editor/PropertyEditor/Private/PropertyEditorHelpers.h
+++ b/Engine/Source/Editor/PropertyEditor/Private/PropertyEditorHelpers.h
@@ -16,6 +16,24 @@ class FNotifyHook;
 class FObjectPropertyNode;
 class FDetailWidgetRow;
 
+namespace UE::Editor::ContentBrowser
+{
+	static bool IsNewStyleEnabled()
+	{
+		static bool bIsNewStyleEnabled = [&]()
+		{
+			if (const IConsoleVariable* CVar = IConsoleManager::Get().FindConsoleVariable(TEXT("ContentBrowser.EnableNewStyle")))
+			{
+				ensureAlwaysMsgf(!EnumHasAnyFlags(CVar->GetFlags(), ECVF_Default), TEXT("The CVar should have already been set from commandline, @see: UnrealEdGlobals.cpp, UE::Editor::ContentBrowser::EnableContentBrowserNewStyleCVarRegistration."));
+				return CVar->GetBool();
+			}
+			return false;
+		}();
+
+		return bIsNewStyleEnabled;
+	}
+}
+
 /** Property button enums. */
 namespace EPropertyButton
 {
diff --git a/Engine/Source/Editor/PropertyEditor/Private/UserInterface/PropertyEditor/SPropertyEditorAsset.cpp b/Engine/Source/Editor/PropertyEditor/Private/UserInterface/PropertyEditor/SPropertyEditorAsset.cpp
--- a/Engine/Source/Editor/PropertyEditor/Private/UserInterface/PropertyEditor/SPropertyEditorAsset.cpp
+++ b/Engine/Source/Editor/PropertyEditor/Private/UserInterface/PropertyEditor/SPropertyEditorAsset.cpp
@@ -47,24 +47,6 @@
 
 #define LOCTEXT_NAMESPACE "PropertyEditor"
 
-namespace UE::Editor::ContentBrowser
-{
-	static bool IsNewStyleEnabled()
-	{
-		static bool bIsNewStyleEnabled = false;
-		UE_CALL_ONCE([&]()
-		{
-			if (const IConsoleVariable* CVar = IConsoleManager::Get().FindConsoleVariable(TEXT("ContentBrowser.EnableNewStyle")))
-			{
-				ensureAlwaysMsgf(!EnumHasAnyFlags(CVar->GetFlags(), ECVF_Default), TEXT("The CVar should have already been set from commandline, @see: UnrealEdGlobals.cpp, UE::Editor::ContentBrowser::EnableContentBrowserNewStyleCVarRegistration."));
-				bIsNewStyleEnabled = CVar->GetBool();
-			}
-		});
-
-		return bIsNewStyleEnabled;
-	}
-}
-
 // Helper to retrieve the correct property that has the applicable metadata.
 static const FProperty* GetActualMetadataProperty(const FProperty* Property)
 {
diff --git a/Engine/Source/Editor/PropertyEditor/Private/UserInterface/PropertyEditor/SPropertyMenuAssetPicker.cpp b/Engine/Source/Editor/PropertyEditor/Private/UserInterface/PropertyEditor/SPropertyMenuAssetPicker.cpp
--- a/Engine/Source/Editor/PropertyEditor/Private/UserInterface/PropertyEditor/SPropertyMenuAssetPicker.cpp
+++ b/Engine/Source/Editor/PropertyEditor/Private/UserInterface/PropertyEditor/SPropertyMenuAssetPicker.cpp
@@ -14,6 +14,7 @@
 #include "Modules/ModuleManager.h"
 #include "PropertyEditorClipboard.h"
 #include "PropertyEditorCopyPastePrivate.h"
+#include "PropertyEditorHelpers.h"
 #include "Styling/SlateIconFinder.h"
 #include "UserInterface/PropertyEditor/SPropertyEditorAsset.h"
 #include "UserInterface/PropertyEditor/PropertyEditorAssetConstants.h"
@@ -139,6 +140,14 @@ void SPropertyMenuAssetPicker::Construct( const FArguments& InArgs )
 		AssetPickerConfig.OnAssetEnterPressed = FOnAssetEnterPressed::CreateSP(this, &SPropertyMenuAssetPicker::OnAssetEnterPressed);
 		// Use the list view by default
 		AssetPickerConfig.InitialAssetViewType = EAssetViewType::List;
+		// TODO: Remove "PropertyEditorHelpers.h" include when IsNewStyleEnabled will be removed
+		if (UE::Editor::ContentBrowser::IsNewStyleEnabled())
+		{
+			// Use the large size by default if the new style is enabled
+			AssetPickerConfig.InitialThumbnailSize = EThumbnailSize::Large;
+			// Avoid creating the MajorAsset specific columns
+			AssetPickerConfig.AvoidMajorAssetColumnsCreation = true;
+		}
 		// The initial selection should be the current value
 		AssetPickerConfig.InitialAssetSelection = CurrentObject;
 		// We'll do clearing ourselves
