[ContentBrowser] Content Browser new style does not sort by size correctly

diff --git a/Engine/Source/Editor/ContentBrowser/Private/AssetViewSortManager.cpp b/Engine/Source/Editor/ContentBrowser/Private/AssetViewSortManager.cpp
--- a/Engine/Source/Editor/ContentBrowser/Private/AssetViewSortManager.cpp
+++ b/Engine/Source/Editor/ContentBrowser/Private/AssetViewSortManager.cpp
@@ -193,12 +193,12 @@ protected:
 
 		if (ValueA < ValueB)
 		{
-			return !bAscending;
+			return bAscending;
 		}
 
-		if (ValueB > ValueA)
+		if (ValueA > ValueB)
 		{
-			return bAscending;
+			return !bAscending;
 		}
 
 		return FCompareFAssetItemBase::Compare(A, B);
