Add missing localization text to the copies tooltips in the CB.

diff --git a/Engine/Source/Editor/ContentBrowser/Private/AssetContextMenu.cpp b/Engine/Source/Editor/ContentBrowser/Private/AssetContextMenu.cpp
--- a/Engine/Source/Editor/ContentBrowser/Private/AssetContextMenu.cpp
+++ b/Engine/Source/Editor/ContentBrowser/Private/AssetContextMenu.cpp
@@ -815,51 +815,54 @@ FSlateIcon FAssetContextMenu::GetEditAssetEditorIcon(bool bInCanEdit, bool bInCa
 
 FText FAssetContextMenu::GetCopyTooltip(EAssetViewCopyType InCopyType) const
 {
-	static const FString CopyExportPathBaseTooltip = TEXT("Copies reference paths for the selected asset(s) to the clipboard:");
-	static const FString CopyObjectPathBaseTooltip = TEXT("Copies object paths for the selected asset(s) to the clipboard:");
-	static const FString CopyPackageNameBaseTooltip = TEXT("Copies package paths for the selected asset(s) to the clipboard:");
+	static const FText CopyExportPathBaseTooltip = LOCTEXT("CopyExportPathBaseTooltip", "Copies reference paths for the selected asset(s) to the clipboard:");
+	static const FText CopyObjectPathBaseTooltip = LOCTEXT("CopyObjectPathBaseTooltip", "Copies object paths for the selected asset(s) to the clipboard:");
+	static const FText CopyPackageNameBaseTooltip = LOCTEXT("CopyPackageNameBaseTooltip", "Copies package paths for the selected asset(s) to the clipboard:");
+
+	TArray<FText> TextsToJoin;
+	TextsToJoin.Reserve(2);
 
-	FText FinalTooltip = FText::GetEmpty();
 	switch (InCopyType)
 	{
 		case EAssetViewCopyType::ExportTextPath:
 			{
-				FString ExportTextPathTooltip = CopyExportPathBaseTooltip;
-				AppendSelectionInformationForCopy(ExportTextPathTooltip, InCopyType);
-				FinalTooltip = FText::FromString(ExportTextPathTooltip);
+				TextsToJoin.Add(CopyExportPathBaseTooltip);
 			}
 			break;
 
 		case EAssetViewCopyType::ObjectPath:
 			{
-				FString ObjectPathTooltip = CopyObjectPathBaseTooltip;
-				AppendSelectionInformationForCopy(ObjectPathTooltip, InCopyType);
-				FinalTooltip = FText::FromString(ObjectPathTooltip);
+				TextsToJoin.Add(CopyObjectPathBaseTooltip);
 			}
 			break;
 
 		case EAssetViewCopyType::PackageName:
 			{
-				FString PackageNameTooltip = CopyPackageNameBaseTooltip;
-				AppendSelectionInformationForCopy(PackageNameTooltip, InCopyType);
-				FinalTooltip = FText::FromString(PackageNameTooltip);
+				TextsToJoin.Add(CopyPackageNameBaseTooltip);
 			}
 			break;
 	}
 
-	return FinalTooltip;
+	if (!TextsToJoin.IsEmpty())
+	{
+		TextsToJoin.Add(GetSelectionInformationForCopy(InCopyType));
+	}
+
+	return FText::Join(LOCTEXT("CopyTooltipDelimiter", "\n"), TextsToJoin);
 }
 
-void FAssetContextMenu::AppendSelectionInformationForCopy(FString& OutTooltip, EAssetViewCopyType InCopyType) const
+FText FAssetContextMenu::GetSelectionInformationForCopy(EAssetViewCopyType InCopyType) const
 {
-	FString AdditionalInformation = TEXT("");
+	FText SelectionInformationForCopy;
 
 	if (!SelectedFiles.IsEmpty())
 	{
 		const FContentBrowserItem& FirstSelectedFile = SelectedFiles[0];
 
-		switch (InCopyType)
 		{
+			FString AdditionalInformation = TEXT("");
+			switch (InCopyType)
+			{
 			case EAssetViewCopyType::ExportTextPath:
 				FirstSelectedFile.AppendItemReference(AdditionalInformation);
 				break;
@@ -871,19 +874,25 @@ void FAssetContextMenu::AppendSelectionInformationForCopy(FString& OutTooltip, E
 			case EAssetViewCopyType::PackageName:
 				FirstSelectedFile.AppendItemPackageName(AdditionalInformation);
 				break;
+			}
+
+			SelectionInformationForCopy = FText::FromString(MoveTemp(AdditionalInformation));
 		}
 
 		if (SelectedFiles.Num() > 1)
 		{
-			AdditionalInformation += LINE_TERMINATOR;
-			AdditionalInformation += FString::Format(TEXT("+{0} more"), { FString::FormatAsNumber(SelectedFiles.Num() - 1) });
+			FText MoreFileCount = FText::Format(LOCTEXT("AdditionalFileCountTooltip", "+{0} more"), FText::AsNumber(SelectedFiles.Num() - 1));
+
+			TArray<FText> SelectionInformationTexts;
+			SelectionInformationTexts.Reserve(2);
+			SelectionInformationTexts.Add(MoveTemp(SelectionInformationForCopy));
+			SelectionInformationTexts.Add(MoveTemp(MoreFileCount));
+
+			SelectionInformationForCopy = FText::Join(LOCTEXT("JoinAdditionalFileCountTooltip", "\n"), SelectionInformationTexts);
 		}
 	}
 
-	if (!AdditionalInformation.IsEmpty())
-	{
-		OutTooltip += LINE_TERMINATOR + AdditionalInformation;
-	}
+	return SelectionInformationForCopy;
 }
 
 bool FAssetContextMenu::AddAssetTypeMenuOptions(UToolMenu* Menu)
diff --git a/Engine/Source/Editor/ContentBrowser/Private/AssetContextMenu.h b/Engine/Source/Editor/ContentBrowser/Private/AssetContextMenu.h
--- a/Engine/Source/Editor/ContentBrowser/Private/AssetContextMenu.h
+++ b/Engine/Source/Editor/ContentBrowser/Private/AssetContextMenu.h
@@ -114,7 +114,7 @@ private:
 	FText GetCopyTooltip(EAssetViewCopyType InCopyType) const;
 
 	/** Append information on the path for the Copy tooltip based on the current selection */
-	void AppendSelectionInformationForCopy(FString& OutTooltip, EAssetViewCopyType InCopyType) const;
+	FText GetSelectionInformationForCopy(EAssetViewCopyType InCopyType) const;
 
 	bool AddPublicStateMenuOptions(UToolMenu* Menu);
 
