[ContentBrowser] AssetPicker default value change
AssetPicker will avoid the creation of columns in the new ListView if the default ViewType is not Column and if it doesn't have any Hidden or Custom columns in the config

diff --git a/Engine/Source/Editor/ContentBrowser/Private/SAssetPicker.cpp b/Engine/Source/Editor/ContentBrowser/Private/SAssetPicker.cpp
--- a/Engine/Source/Editor/ContentBrowser/Private/SAssetPicker.cpp
+++ b/Engine/Source/Editor/ContentBrowser/Private/SAssetPicker.cpp
@@ -332,7 +332,19 @@ void SAssetPicker::Construct( const FArguments& InArgs )
 	{
 		TextFilter = MakeShared<FAssetTextFilter>();
 	}
-	
+
+	// Avoid MajorAssetColumns creation if the ViewType is set to column and there are no Hidden or Custom columns
+	bool AvoidMajorAssetColumnsCreation = InArgs._AssetPickerConfig.AvoidMajorAssetColumnsCreation;
+	if (UE::Editor::ContentBrowser::IsNewStyleEnabled())
+	{
+		if (InArgs._AssetPickerConfig.InitialAssetViewType != EAssetViewType::Column &&
+			InArgs._AssetPickerConfig.HiddenColumnNames.IsEmpty() &&
+			InArgs._AssetPickerConfig.CustomColumns.IsEmpty())
+		{
+			AvoidMajorAssetColumnsCreation = true;
+		}
+	}
+
 	// clang-format off
 	VerticalBox->AddSlot()
 	.FillHeight(1.f)
@@ -377,7 +389,7 @@ void SAssetPicker::Construct( const FArguments& InArgs )
 		.AssetShowWarningText( InArgs._AssetPickerConfig.AssetShowWarningText)
 		.AllowFocusOnSync(false)	// Stop the asset view from stealing focus (we're in control of that)
 		.HiddenColumnNames(InArgs._AssetPickerConfig.HiddenColumnNames)
-		.AvoidMajorAssetColumnsCreation(InArgs._AssetPickerConfig.AvoidMajorAssetColumnsCreation)
+		.AvoidMajorAssetColumnsCreation(AvoidMajorAssetColumnsCreation)
 		.CustomColumns(InArgs._AssetPickerConfig.CustomColumns)
 		.OnSearchOptionsChanged(this, &SAssetPicker::HandleSearchSettingsChanged)
 		.InitialThumbnailSize(InArgs._AssetPickerConfig.InitialThumbnailSize)
