[ContentBrowser] ViewType command for menu entries and Grid/List View changes
Added commands for the View Type in the Content Browser menu
Grid view will automatically change into List when the size goes below the tiny size
List view will automatically change into Grid when the size goes above the huge size

diff --git a/Engine/Source/Editor/ContentBrowser/Private/ContentBrowserCommands.cpp b/Engine/Source/Editor/ContentBrowser/Private/ContentBrowserCommands.cpp
--- a/Engine/Source/Editor/ContentBrowser/Private/ContentBrowserCommands.cpp
+++ b/Engine/Source/Editor/ContentBrowser/Private/ContentBrowserCommands.cpp
@@ -27,6 +27,8 @@ void FContentBrowserCommands::RegisterCommands()
 	// AssetView commands
 	UI_COMMAND(AssetViewCopyObjectPath, "Copy Selected Object Path", "Copy the selected object path", EUserInterfaceActionType::Button, FInputChord(EModifierKey::Control | EModifierKey::Shift, EKeys::C));
 	UI_COMMAND(AssetViewCopyPackageName, "Copy Selected Package Name", "Copy the selected package name", EUserInterfaceActionType::Button, FInputChord(EModifierKey::Control | EModifierKey::Alt, EKeys::C));
+	UI_COMMAND(GridViewShortcut, "Grid", "View assets as tiles in a grid.", EUserInterfaceActionType::RadioButton, FInputChord());
+	UI_COMMAND(ListViewShortcut, "List", "View assets in a list with thumbnails.", EUserInterfaceActionType::RadioButton, FInputChord());
 }
 
 
diff --git a/Engine/Source/Editor/ContentBrowser/Private/ContentBrowserCommands.h b/Engine/Source/Editor/ContentBrowser/Private/ContentBrowserCommands.h
--- a/Engine/Source/Editor/ContentBrowser/Private/ContentBrowserCommands.h
+++ b/Engine/Source/Editor/ContentBrowser/Private/ContentBrowserCommands.h
@@ -42,4 +42,6 @@ public:
 	TSharedPtr<FUICommandInfo> EditPath;
 	TSharedPtr<FUICommandInfo> AssetViewCopyObjectPath;
 	TSharedPtr<FUICommandInfo> AssetViewCopyPackageName;
+	TSharedPtr<FUICommandInfo> GridViewShortcut;
+	TSharedPtr<FUICommandInfo> ListViewShortcut;
 };
diff --git a/Engine/Source/Editor/ContentBrowser/Private/Menus/AssetViewMenus.cpp b/Engine/Source/Editor/ContentBrowser/Private/Menus/AssetViewMenus.cpp
--- a/Engine/Source/Editor/ContentBrowser/Private/Menus/AssetViewMenus.cpp
+++ b/Engine/Source/Editor/ContentBrowser/Private/Menus/AssetViewMenus.cpp
@@ -1,5 +1,6 @@
 // Copyright Epic Games, Inc. All Rights Reserved.
 
+#include "ContentBrowserCommands.h"
 #include "ContentBrowserMenuContexts.h"
 #include "ContentBrowserMenuUtils.h"
 #include "ContentBrowserStyle.h"
@@ -108,31 +109,40 @@ FDelayedAutoRegisterHelper SAssetView::AssetViewOptionsMenuRegistration(
 							"AssetViewType",
 							LOCTEXT("ViewTypeHeading", "View Type"));
 
-						ViewTypeSection.AddMenuEntry(
-							"TileView",
-							LOCTEXT("TileViewOption", "Grid"),
-							LOCTEXT("TileViewOptionToolTip", "View assets as tiles in a grid."),
-							FSlateIcon(),
-							FUIAction(
-								FExecuteAction::CreateSP(AssetView, &SAssetView::SetCurrentViewTypeFromMenu, EAssetViewType::Tile),
-								FCanExecuteAction(),
-								FIsActionChecked::CreateSP(AssetView, &SAssetView::IsCurrentViewType, EAssetViewType::Tile)
-							),
-							EUserInterfaceActionType::RadioButton
-						);
+						// Use commands in the new CB
+						if (UE::Editor::ContentBrowser::IsNewStyleEnabled())
+						{
+							ViewTypeSection.AddMenuEntryWithCommandList(FContentBrowserCommands::Get().GridViewShortcut, AssetView->Commands);
+							ViewTypeSection.AddMenuEntryWithCommandList(FContentBrowserCommands::Get().ListViewShortcut, AssetView->Commands);
+						}
+						else
+						{
+							ViewTypeSection.AddMenuEntry(
+								"TileView",
+								LOCTEXT("TileViewOption", "Grid"),
+								LOCTEXT("TileViewOptionToolTip", "View assets as tiles in a grid."),
+								FSlateIcon(),
+								FUIAction(
+									FExecuteAction::CreateSP(AssetView, &SAssetView::SetCurrentViewTypeFromMenu, EAssetViewType::Tile),
+									FCanExecuteAction(),
+									FIsActionChecked::CreateSP(AssetView, &SAssetView::IsCurrentViewType, EAssetViewType::Tile)
+								),
+								EUserInterfaceActionType::RadioButton
+							);
 
-						ViewTypeSection.AddMenuEntry(
-							"ListView",
-							LOCTEXT("ListViewOption", "List"),
-							LOCTEXT("ListViewOptionToolTip", "View assets in a list with thumbnails."),
-							FSlateIcon(),
-							FUIAction(
-								FExecuteAction::CreateSP(AssetView, &SAssetView::SetCurrentViewTypeFromMenu, EAssetViewType::List),
-								FCanExecuteAction(),
-								FIsActionChecked::CreateSP(AssetView, &SAssetView::IsCurrentViewType, EAssetViewType::List)
-							),
-							EUserInterfaceActionType::RadioButton
-						);
+							ViewTypeSection.AddMenuEntry(
+								"ListView",
+								LOCTEXT("ListViewOption", "List"),
+								LOCTEXT("ListViewOptionToolTip", "View assets in a list with thumbnails."),
+								FSlateIcon(),
+								FUIAction(
+									FExecuteAction::CreateSP(AssetView, &SAssetView::SetCurrentViewTypeFromMenu, EAssetViewType::List),
+									FCanExecuteAction(),
+									FIsActionChecked::CreateSP(AssetView, &SAssetView::IsCurrentViewType, EAssetViewType::List)
+								),
+								EUserInterfaceActionType::RadioButton
+							);
+						}
 
 						if (!UE::Editor::ContentBrowser::IsNewStyleEnabled())
 						{
diff --git a/Engine/Source/Editor/ContentBrowser/Private/SAssetView.cpp b/Engine/Source/Editor/ContentBrowser/Private/SAssetView.cpp
--- a/Engine/Source/Editor/ContentBrowser/Private/SAssetView.cpp
+++ b/Engine/Source/Editor/ContentBrowser/Private/SAssetView.cpp
@@ -2675,10 +2675,14 @@ FReply SAssetView::OnMouseWheel( const FGeometry& MyGeometry, const FPointerEven
 	{
 		if (UE::Editor::ContentBrowser::IsNewStyleEnabled())
 		{
+			const int32 Delta = MouseEvent.GetWheelDelta() > 0 ? 1 : -1;
+			bool IsLessThanMinSize = (int32)ThumbnailSize + Delta < 0;
+			bool IsMoreThanMaxSize = (int32)ThumbnailSize + Delta == (int32)EThumbnailSize::MAX;
+			bool bWasSizeSupposedToChange = false;
 			// If shift is pressed or if we are in the List view the CTRL + Wheel should jump ThumbnailSize by design
 			if (MouseEvent.IsShiftDown() || CurrentViewType == EAssetViewType::List)
 			{
-				const int32 Delta = MouseEvent.GetWheelDelta() > 0 ? 1 : -1;
+				bWasSizeSupposedToChange = true;
 				EThumbnailSize DesiredThumbnailSize = (EThumbnailSize)FMath::Clamp<int32>((int32)ThumbnailSize + Delta, 0, (int32)EThumbnailSize::MAX - 1);
 
 				if (DesiredThumbnailSize != ThumbnailSize)
@@ -2688,10 +2692,11 @@ FReply SAssetView::OnMouseWheel( const FGeometry& MyGeometry, const FPointerEven
 			}
 			else
 			{
-				const float NewDelta = MouseEvent.GetWheelDelta() * 0.4f;
+				const float NewDelta = (float)Delta * 0.4f;
 				if ((ZoomScale == 1.f && NewDelta > 0) || (ZoomScale == 0.f && NewDelta < 0))
 				{
-					int32 Step = (int32)FMath::Sign(NewDelta);
+					bWasSizeSupposedToChange = true;
+					const int32 Step = (int32)FMath::Sign(NewDelta);
 					EThumbnailSize OldSize = ThumbnailSize;
 					ThumbnailSize = (EThumbnailSize)FMath::Clamp<int32>(((int32)ThumbnailSize + Step), 0, (int32)EThumbnailSize::MAX - 1);
 					if (OldSize != ThumbnailSize)
@@ -2707,6 +2712,22 @@ FReply SAssetView::OnMouseWheel( const FGeometry& MyGeometry, const FPointerEven
 					RefreshList();
 				}
 			}
+
+			// Switch the view automatically
+			if (bWasSizeSupposedToChange)
+			{
+				if (CurrentViewType == EAssetViewType::List && IsMoreThanMaxSize)
+				{
+					ZoomScale = 0.f;
+					SetCurrentViewType(EAssetViewType::Tile);
+					OnThumbnailSizeChanged(EThumbnailSize::Tiny);
+				}
+				else if (CurrentViewType == EAssetViewType::Tile && IsLessThanMinSize)
+				{
+					SetCurrentViewType(EAssetViewType::List);
+                    OnThumbnailSizeChanged(EThumbnailSize::Huge);
+				}
+			}
 		}
 		else
 		{
@@ -3525,6 +3546,21 @@ void SAssetView::BindCommands()
 		FCanExecuteAction::CreateSP(this, &SAssetView::IsAssetPathSelected)
 	));
 
+	if (UE::Editor::ContentBrowser::IsNewStyleEnabled())
+	{
+		Commands->MapAction(FContentBrowserCommands::Get().GridViewShortcut, FUIAction(
+			FExecuteAction::CreateSP(this, &SAssetView::SetCurrentViewTypeFromMenu, EAssetViewType::Tile),
+			FCanExecuteAction(),
+			FIsActionChecked::CreateSP(this, &SAssetView::IsCurrentViewType, EAssetViewType::Tile)
+		));
+
+		Commands->MapAction(FContentBrowserCommands::Get().ListViewShortcut, FUIAction(
+			FExecuteAction::CreateSP(this, &SAssetView::SetCurrentViewTypeFromMenu, EAssetViewType::List),
+			FCanExecuteAction(),
+			FIsActionChecked::CreateSP(this, &SAssetView::IsCurrentViewType, EAssetViewType::List)
+		));
+	}
+
 	FInputBindingManager::Get().RegisterCommandList(FContentBrowserCommands::Get().GetContextName(), Commands.ToSharedRef());
 }
 
