[ContentBrowser] ToolBar wrapping behavior

diff --git a/Engine/Source/Editor/ContentBrowser/Private/ContentBrowserStyle.cpp b/Engine/Source/Editor/ContentBrowser/Private/ContentBrowserStyle.cpp
--- a/Engine/Source/Editor/ContentBrowser/Private/ContentBrowserStyle.cpp
+++ b/Engine/Source/Editor/ContentBrowser/Private/ContentBrowserStyle.cpp
@@ -74,17 +74,18 @@ namespace UE::ContentBrowser::Private
 
 		Set("ContentBrowser.DefaultFolderColor", FStyleColors::AccentFolder);
 
-		if (Editor::ContentBrowser::IsNewStyleEnabled())
-		{
-			Set("ContentBrowser.AssetTreeHeaderBrush", new FSlateColorBrush(FStyleColors::Panel));
-			Set("ContentBrowser.AssetTreeBodyBrush", new FSlateColorBrush(FStyleColors::Recessed));
-		}
-		else
+		FSlateBrush* AssetTreeHeaderBrush = new FSlateColorBrush(FStyleColors::Panel);
+		FSlateBrush* AssetTreeBodyBrush = new FSlateColorBrush(FStyleColors::Recessed);
+
+		if (!Editor::ContentBrowser::IsNewStyleEnabled())
 		{
-			Set("ContentBrowser.AssetTreeHeaderBrush", new FSlateColorBrush(FStyleColors::Header));
-			Set("ContentBrowser.AssetTreeBodyBrush", new FSlateColorBrush(FStyleColors::Recessed));
+			AssetTreeHeaderBrush = new FSlateColorBrush(FStyleColors::Panel);
+			AssetTreeBodyBrush = new FSlateColorBrush(FStyleColors::Recessed);
 		}
 
+		Set("ContentBrowser.AssetTreeHeaderBrush", AssetTreeHeaderBrush);
+		Set("ContentBrowser.AssetTreeBodyBrush", AssetTreeBodyBrush);
+
 		FExpandableAreaStyle AssetTreeExpandableAreaStyle = FStarshipCoreStyle::GetCoreStyle().GetWidgetStyle<FExpandableAreaStyle>("ExpandableArea");
 
 		Set("ContentBrowser.AssetTreeExpandableArea", AssetTreeExpandableAreaStyle);
@@ -336,11 +337,19 @@ namespace UE::ContentBrowser::Private
 
 		ContentBrowserToolBarStyle.SetSeparatorBrush(FSlateNoResource());
 		ContentBrowserToolBarStyle.SetSeparatorPadding(FMargin(4.0f, 0.0f));
+
 		if (UE::Editor::ContentBrowser::IsNewStyleEnabled())
 		{
+			// Hide separators
+			ContentBrowserToolBarStyle.SetSeparatorThickness(0.0f);
+			ContentBrowserToolBarStyle.SetSeparatorPadding(FMargin(0.0f, 0.0f));
+
+			ContentBrowserToolBarStyle.SetBlockPadding(FMargin(2.0f, 0.0f)); // Effectively makes custom widgets look like buttons
+
 			ContentBrowserToolBarStyle.SetAllowWrapButton(false); // Never show the wrap button, just clip
 			ContentBrowserToolBarStyle.SetIconSize(FVector2D(16.0f, 16.0f));
 			ContentBrowserToolBarStyle.SetBackgroundPadding(FMargin(4.0f, 2.0f, 4.0f, 2.0f));
+			ContentBrowserToolBarStyle.SetAllowWrappingDefault(false);
 		}
 		else
 		{
@@ -410,6 +419,18 @@ namespace UE::ContentBrowser::Private
 			Set("ContentBrowser.Sources.CreateIcon", new IMAGE_BRUSH_SVG("Starship/Common/PlaceActors", DefaultIconSize));
 		}
 
+		// Filter View, mostly matches AssetTree
+		{
+			static constexpr float FilterViewHorizontalPadding = 10.0f;
+			static constexpr float FilterViewVerticalHeaderPadding = FilterViewHorizontalPadding - 6.0f;
+
+			Set("ContentBrowser.VerticalFilterViewHeaderBrush", AssetTreeHeaderBrush);
+			Set("ContentBrowser.VerticalFilterViewHeaderPadding", FMargin(FilterViewHorizontalPadding, FilterViewVerticalHeaderPadding));
+			Set("ContentBrowser.VerticalFilterViewHeaderTextHeight", CoreStyleConstants::Icon24x24.Y); // Based on button height when dictated by an icon, matches AssetTree header
+
+			Set("ContentBrowser.VerticalFilterViewBodyBrush", AssetTreeHeaderBrush); // Body is the same as the header brush/color
+		}
+
 		FSlateStyleRegistry::RegisterSlateStyle(*this);
 	}
 
diff --git a/Engine/Source/Editor/ContentBrowser/Private/Menus/AssetViewMenus.cpp b/Engine/Source/Editor/ContentBrowser/Private/Menus/AssetViewMenus.cpp
--- a/Engine/Source/Editor/ContentBrowser/Private/Menus/AssetViewMenus.cpp
+++ b/Engine/Source/Editor/ContentBrowser/Private/Menus/AssetViewMenus.cpp
@@ -546,6 +546,8 @@ FDelayedAutoRegisterHelper SAssetView::ToolBarMenuExtensionRegistration(
 		{
 			FToolMenuSection& SortSection = ToolBarMenu->FindOrAddSection("Sort");
 			SortSection.InsertPosition = FToolMenuInsert("Search", EToolMenuInsertType::After);
+			SortSection.ResizeParams.Wrapping.Allow = true;
+			SortSection.ResizeParams.Wrapping.Priority = 10;
 
 			SortSection.AddDynamicEntry(
 				"Sort",
diff --git a/Engine/Source/Editor/ContentBrowser/Private/Menus/ContentBrowserMenus.cpp b/Engine/Source/Editor/ContentBrowser/Private/Menus/ContentBrowserMenus.cpp
--- a/Engine/Source/Editor/ContentBrowser/Private/Menus/ContentBrowserMenus.cpp
+++ b/Engine/Source/Editor/ContentBrowser/Private/Menus/ContentBrowserMenus.cpp
@@ -6,6 +6,7 @@
 #include "SActionButton.h"
 #include "SAssetSearchBox.h"
 #include "SContentBrowser.h"
+#include "SFilterList.h"
 #include "SNavigationBar.h"
 #include "SPositiveActionButton.h"
 #include "ToolMenu.h"
@@ -231,12 +232,13 @@ FDelayedAutoRegisterHelper SContentBrowser::ToolBarMenuRegistration(
         FToolMenuOwnerScoped OwnerScoped(UE_MODULE_NAME);
 
 		static const FName ToolBarMenuName("ContentBrowser.ToolBar");
-        UToolMenu* ToolBar = UToolMenus::Get()->RegisterMenu(ToolBarMenuName, NAME_None, EMultiBoxType::SlimHorizontalToolBar);
+        UToolMenu* ToolBar = UToolMenus::Get()->RegisterMenu(ToolBarMenuName, NAME_None, EMultiBoxType::SlimWrappingToolBar);
         ToolBar->SetStyleSet(&UE::ContentBrowser::Private::FContentBrowserStyle::Get());
         ToolBar->StyleName = ToolBarMenuName; // Style name is the same as the menu entry name
 
-        ToolBar->FindOrAddSection("New")
-        .AddDynamicEntry(
+		FToolMenuSection& NewSection = ToolBar->FindOrAddSection("New");
+
+        NewSection.AddDynamicEntry(
             "New",
             FNewToolMenuSectionDelegate::CreateLambda([](FToolMenuSection& InSection)
             {
@@ -258,7 +260,7 @@ FDelayedAutoRegisterHelper SContentBrowser::ToolBarMenuRegistration(
                             .Icon(FAppStyle::Get().GetBrush("Icons.Plus"))
                             .Text(LOCTEXT("AddAssetButton", "Add"));
 
-                    	InSection.AddEntry(
+						InSection.AddEntry(
 							FToolMenuEntry::InitWidget(
 								"NewButton",
 								NewButton,
@@ -270,8 +272,9 @@ FDelayedAutoRegisterHelper SContentBrowser::ToolBarMenuRegistration(
                 }
             }));
 
-        ToolBar->FindOrAddSection("Save", { }, FToolMenuInsert("New", EToolMenuInsertType::After))
-        .AddDynamicEntry(
+		FToolMenuSection& SaveSection = ToolBar->FindOrAddSection("Save", {}, FToolMenuInsert("New", EToolMenuInsertType::After));
+		
+        SaveSection.AddDynamicEntry(
             "Save",
             FNewToolMenuSectionDelegate::CreateLambda([](FToolMenuSection& InSection)
             {
@@ -341,46 +344,53 @@ FDelayedAutoRegisterHelper SContentBrowser::ToolBarMenuRegistration(
 
         if (UE::Editor::ContentBrowser::IsNewStyleEnabled())
         {
-            ToolBar->FindOrAddSection("Filter", { }, FToolMenuInsert("Save", EToolMenuInsertType::After))
-            .AddDynamicEntry(
-                "Filter",
-                FNewToolMenuSectionDelegate::CreateLambda(
-                [](FToolMenuSection& InSection)
-                {
-                    if (const UContentBrowserToolbarMenuContext* Context = InSection.FindContext<UContentBrowserToolbarMenuContext>())
-                    {
-                        if (Context->ContentBrowser.IsValid())
-                        {
-                            const TSharedRef<SContentBrowser> ContentBrowser = Context->ContentBrowser.Pin().ToSharedRef();
-
-                            if (!ContentBrowser->FilterComboButton.IsValid())
-                            {
-                                return;
-                            }
-
-                            FToolMenuEntry& FilterEntry = InSection.AddEntry(
-                                FToolMenuEntry::InitWidget(
-                                    "Filter",
-                                    SNew(SBox)
-                                    .HAlign(HAlign_Left)
-                                    .VAlign(VAlign_Center)
-                                    [
-                                        ContentBrowser->FilterComboButton.ToSharedRef()
-                                    ],
-                                    LOCTEXT("Filter", "Filter")
-                                ));
+			FToolMenuSection& FilterSection = ToolBar->FindOrAddSection("Filter", {}, FToolMenuInsert("Save", EToolMenuInsertType::After));
+        	FilterSection.ResizeParams.Wrapping.Allow = true;
+        	FilterSection.ResizeParams.Wrapping.Priority = 10;
+
+        	FilterSection.AddDynamicEntry(
+				"Filter",
+				FNewToolMenuSectionDelegate::CreateLambda(
+					[](FToolMenuSection& InSection)
+					{
+						if (const UContentBrowserToolbarMenuContext* Context = InSection.FindContext<UContentBrowserToolbarMenuContext>())
+						{
+							if (Context->ContentBrowser.IsValid())
+							{
+								const TSharedRef<SContentBrowser> ContentBrowser = Context->ContentBrowser.Pin().ToSharedRef();
 
-                            FilterEntry.WidgetData.StyleParams.HorizontalAlignment = HAlign_Left;
-                            FilterEntry.WidgetData.StyleParams.SizeRule = FSizeParam::ESizeRule::SizeRule_Auto;
+								if (!ContentBrowser->FilterComboButton.IsValid())
+								{
+									return;
+								}
+
+								FToolMenuEntry& FilterEntry = InSection.AddEntry(
+									FToolMenuEntry::InitWidget(
+										"Filter",
+										SNew(SBox)
+										.HAlign(HAlign_Left)
+										.VAlign(VAlign_Center)
+										[
+											ContentBrowser->FilterComboButton.ToSharedRef()
+										],
+										LOCTEXT("Filter", "Filter")
+									));
+
+								FilterEntry.WidgetData.StyleParams.HorizontalAlignment = HAlign_Left;
+								FilterEntry.WidgetData.StyleParams.SizeRule = FSizeParam::ESizeRule::SizeRule_Auto;
+
+								// ... no label when in the toolbar
+								FilterEntry.ToolBarData.LabelOverride = FText::GetEmpty();
+							}
+						}
+					}));
 
-                            // ... no label when in the toolbar
-                            FilterEntry.ToolBarData.LabelOverride = FText::GetEmpty();
-                        }
-                    }
-                }));
+			FToolMenuSection& SearchSection = ToolBar->FindOrAddSection("Search", {}, FToolMenuInsert("Filter", EToolMenuInsertType::After));
+        	SearchSection.ResizeParams.Wrapping.Allow = true;
+        	SearchSection.ResizeParams.Wrapping.Mode = UE::Slate::PrioritizedWrapBox::EWrapMode::Parent;
+        	SearchSection.ResizeParams.Wrapping.Priority = 10;
 
-            ToolBar->FindOrAddSection("Search", { }, FToolMenuInsert("Filter", EToolMenuInsertType::After))
-            .AddDynamicEntry(
+        	SearchSection.AddDynamicEntry(
                 "Search",
                 FNewToolMenuSectionDelegate::CreateLambda(
                 [](FToolMenuSection& InSection)
@@ -441,6 +451,62 @@ FDelayedAutoRegisterHelper SContentBrowser::ToolBarMenuRegistration(
                     }
                 }));
 
+        	// Filters (plural) is the individual tag widgets, not the dropdown menu
+			FToolMenuSection& FiltersSection = ToolBar->FindOrAddSection("Filters", {}, FToolMenuInsert("Sort", EToolMenuInsertType::After));
+        	FiltersSection.ResizeParams.Wrapping.Allow = true;
+        	FiltersSection.ResizeParams.Wrapping.Priority = 20;
+        	FiltersSection.ResizeParams.Wrapping.Mode = UE::Slate::PrioritizedWrapBox::EWrapMode::Parent; // Always wrap when too long, don't wait to hit the Preferred size threshold
+        	FiltersSection.ResizeParams.Wrapping.VerticalOverflowBehavior = UE::Slate::PrioritizedWrapBox::EVerticalOverflowBehavior::ExpandProportional;
+
+        	FiltersSection.AddDynamicEntry(
+                "Filters",
+                FNewToolMenuSectionDelegate::CreateLambda(
+                [](FToolMenuSection& InSection)
+                {
+                    if (const UContentBrowserToolbarMenuContext* Context = InSection.FindContext<UContentBrowserToolbarMenuContext>())
+                    {
+                        if (Context->ContentBrowser.IsValid())
+                        {
+                            const TSharedRef<SContentBrowser> ContentBrowser = Context->ContentBrowser.Pin().ToSharedRef();
+
+                            // Should have been initialized in SContentBrowser::CreateAssetView
+                            if (!ensure(ContentBrowser->LegacyContentSourceWidgets->FilterListPtr.IsValid()))
+                            {
+                                return;
+                            }
+
+                            FToolMenuEntry& FilterListEntry =
+                                InSection.AddEntry(
+                                FToolMenuEntry::InitWidget(
+                                    "FilterList",
+                                    SNew(SBox)
+                                    .Padding(FMargin(-3))
+                                    [
+                                    	ContentBrowser->LegacyContentSourceWidgets->FilterListPtr.ToSharedRef()
+                                    ],
+                                    FText::GetEmpty()));
+
+                        	const TAttribute<EVisibility> FilterListVisibility = TAttribute<EVisibility>::Create(
+                        		TAttribute<EVisibility>::FGetter::CreateSPLambda(
+                        			ContentBrowser,
+                        			[WeakContentBrowser = ContentBrowser.ToWeakPtr()]() -> EVisibility
+                        			{
+                        				if (WeakContentBrowser.IsValid())
+                        				{
+                        					return WeakContentBrowser.Pin()->GetFilterLayout() == EFilterBarLayout::Horizontal ? EVisibility::Visible : EVisibility::Collapsed;
+                        				}
+
+                        				return EVisibility::Collapsed;
+                        			}));
+
+                        	FilterListEntry.Visibility = FilterListVisibility;
+							FilterListEntry.WidgetData.StyleParams.SizeRule = FSizeParam::SizeRule_StretchContent;
+                        	FilterListEntry.WidgetData.StyleParams.FillSizeMin = 0.1f;
+                        	FilterListEntry.WidgetData.StyleParams.FillSize = 1.0f;
+                        }
+                    }
+                }));
+
             FToolMenuSection& SettingsSection = ToolBar->FindOrAddSection("Settings", { }, FToolMenuInsert("Search", EToolMenuInsertType::After));
             SettingsSection.Alignment = EToolMenuSectionAlign::Last;
 
diff --git a/Engine/Source/Editor/ContentBrowser/Private/SContentBrowser.cpp b/Engine/Source/Editor/ContentBrowser/Private/SContentBrowser.cpp
--- a/Engine/Source/Editor/ContentBrowser/Private/SContentBrowser.cpp
+++ b/Engine/Source/Editor/ContentBrowser/Private/SContentBrowser.cpp
@@ -172,6 +172,24 @@ struct FSlateBrush;
 
 #define LOCTEXT_NAMESPACE "ContentBrowser"
 
+namespace UE::Editor::ContentBrowser::Private
+{
+	// Find and return the slot index containing a widget with the given tag. Will return INDEX_NONE if not found.
+	static int32 FindSlotByWidgetTag(const TSharedRef<SSplitter>& InSplitter, const FName InTag)
+	{
+		for (int32 SlotIndex = 0; SlotIndex < InSplitter->GetChildren()->Num(); ++SlotIndex)
+		{
+			TSharedRef<SWidget> WidgetInSlot = InSplitter->SlotAt(SlotIndex).GetWidget();
+			if (WidgetInSlot->GetTag() == InTag)
+			{
+				return SlotIndex;
+			}
+		}
+
+		return INDEX_NONE;
+	}
+}
+
 const FString SContentBrowser::SettingsIniSection = TEXT("ContentBrowser");
 
 class SContentBrowser::FCollectionSource
@@ -662,92 +680,105 @@ void SContentBrowser::OnFilterBarLayoutChanging(EFilterBarLayout NewLayout)
 {
 	if (UE::Editor::ContentBrowser::IsNewStyleEnabled())
 	{
+		using namespace UE::ContentBrowser::Private;
+		
+		// Identify filter view locations by widget tag, so we don't assume slot index
+		static FName HorizontalFilterViewTagName = "HorizontalFilterView";
+		static FName VerticalFilterViewTagName = "VerticalFilterView";
+	
 		if (NewLayout == EFilterBarLayout::Horizontal)
 		{
-			LegacyContentSourceWidgets->AssetViewBorder->SetContent(
-
-				SNew(SVerticalBox)
-
-				+ SVerticalBox::Slot()
-				.Padding(0.0f)
-				.AutoHeight()
-				[
-					LegacyContentSourceWidgets->FilterListPtr.ToSharedRef()
-				]
-
-				+ SVerticalBox::Slot()
-				.FillHeight(1.0f)
-				.Padding(0, 0)
-				[
-					LegacyContentSourceWidgets->AssetViewPtr.ToSharedRef()
-				]
-			);
+			const int32 FoundVerticalFilterViewSlotIndex = UE::Editor::ContentBrowser::Private::FindSlotByWidgetTag(
+				LegacyContentSourceWidgets->PathAssetSplitterPtr.ToSharedRef(),
+				VerticalFilterViewTagName);
+			
+			// Remove from Vertical layout (if needed)
+			if (FoundVerticalFilterViewSlotIndex != INDEX_NONE)
+			{
+				LegacyContentSourceWidgets->PathAssetSplitterPtr->RemoveAt(FoundVerticalFilterViewSlotIndex);
+			}
 		}
 		else
 		{
-			LegacyContentSourceWidgets->AssetViewBorder->SetContent(
+			const int32 HorizontalFilterViewSlotIndex = UE::Editor::ContentBrowser::Private::FindSlotByWidgetTag(
+				LegacyContentSourceWidgets->PathAssetSplitterPtr.ToSharedRef(),
+				HorizontalFilterViewTagName);
 
-				SNew(SSplitter)
-				.PhysicalSplitterHandleSize(2.0f)
+			// Remove from Horizontal layout (if needed)
+			if (HorizontalFilterViewSlotIndex != INDEX_NONE)
+			{
+				LegacyContentSourceWidgets->PathAssetSplitterPtr->RemoveAt(HorizontalFilterViewSlotIndex);
+			}
 
-				/* Filters in an SScrollBox */
-				+ SSplitter::Slot()
-				.MinSize(95.f)
-				.Resizable(true)
-				.SizeRule(SSplitter::SizeToContent)
-				.OnSlotResized(this, &SContentBrowser::OnFilterBoxColumnResized)
+			// This can differ to the desired index, depending on widget state
+			const int32 FoundVerticalFilterViewSlotIndex = UE::Editor::ContentBrowser::Private::FindSlotByWidgetTag(
+				LegacyContentSourceWidgets->PathAssetSplitterPtr.ToSharedRef(),
+				VerticalFilterViewTagName);
+
+			// Check for Existing
+			if (FoundVerticalFilterViewSlotIndex != INDEX_NONE)
+			{
+				LegacyContentSourceWidgets->PathAssetSplitterPtr->RemoveAt(FoundVerticalFilterViewSlotIndex);
+			}
+
+			constexpr int32 DesiredVerticalFilterViewSlotIndex = 1;
+
+			LegacyContentSourceWidgets->PathAssetSplitterPtr->AddSlot(DesiredVerticalFilterViewSlotIndex)
+			.MinSize(95.0f)
+			.Resizable(true)
+			.SizeRule(SSplitter::SizeToContent)
+			.OnSlotResized(this, &SContentBrowser::OnFilterBoxColumnResized)
+			[
+				// Vertical Filter View
+				SNew(SBox)
+				.Tag(VerticalFilterViewTagName)
+				.WidthOverride(this, &SContentBrowser::GetFilterViewBoxWidthOverride)
+				// Don't take up space when there are no filters
+				.Visibility_Lambda([this]
+				{
+					return LegacyContentSourceWidgets->FilterListPtr->HasAnyFilters() ? EVisibility::Visible : EVisibility::Collapsed;
+				})
 				[
-					SNew(SBox)
-					.WidthOverride(this, &SContentBrowser::GetFilterViewBoxWidthOverride)
-					// Don't take up space when there are no filters
-					.Visibility_Lambda([this]
-					{
-						return LegacyContentSourceWidgets->FilterListPtr->HasAnyFilters() ? EVisibility::Visible : EVisibility::Collapsed;
-					})
-					[
-						SNew(SVerticalBox)
+					SNew(SVerticalBox)
 
-						// Header
-						+ SVerticalBox::Slot()
-						.Padding(0.0f, 2.0f)
-						.AutoHeight()
+					// Header
+					+ SVerticalBox::Slot()
+					.Padding(0.0f, 2.0f)
+					.AutoHeight()
+					[
+						SNew(SBorder)
+						.BorderImage(FContentBrowserStyle::Get().GetBrush("ContentBrowser.VerticalFilterViewHeaderBrush"))
+						.Padding(FContentBrowserStyle::Get().GetMargin("ContentBrowser.VerticalFilterViewHeaderPadding"))
+						.Content()
 						[
-
-							SNew(SBorder)
-							.BorderImage(FAppStyle::Get().GetBrush("Brushes.Header"))
-							.Padding(FMargin(8.0f, 6.0f))
-							.Content()
+							// Enforce widget height
+							SNew(SBox)
+							.HeightOverride(FContentBrowserStyle::Get().GetFloat("ContentBrowser.VerticalFilterViewHeaderTextHeight"))
+							.Padding(0)
+							.VAlign(VAlign_Center)
 							[
 								SNew(STextBlock)
 								.Text(LOCTEXT("FilterListVerticalHeader", "Filters"))
 								.TextStyle(FAppStyle::Get(), "ButtonText")
-								.Font(FAppStyle::Get().GetFontStyle("NormalFontBold"))
+								.Font(FAppStyle::Get().GetFontStyle("NormalFontBold"))	
 							]
 						]
-
-						// Filter List
-						+ SVerticalBox::Slot()
-						.Padding(0.0f)
-						.FillHeight(1.0f)
-						[
-							LegacyContentSourceWidgets->FilterListPtr.ToSharedRef()
-						]
 					]
-				]
-
-				+ SSplitter::Slot()
-				.Value(0.88f)
-				[
-					SNew(SVerticalBox)
 
+					// Filter List
 					+ SVerticalBox::Slot()
+					.Padding(0.0f)
 					.FillHeight(1.0f)
-					.Padding(0, 0)
 					[
-						LegacyContentSourceWidgets->AssetViewPtr.ToSharedRef()
+						SNew(SBorder)
+						.BorderImage(FContentBrowserStyle::Get().GetBrush("ContentBrowser.VerticalFilterViewBodyBrush"))
+						.Padding(0.0f)
+						[
+							LegacyContentSourceWidgets->FilterListPtr.ToSharedRef()	
+						]
 					]
 				]
-			);
+			];
 		}
 	}
 	else
@@ -862,32 +893,6 @@ void SContentBrowser::OnFilterBarLayoutChanging(EFilterBarLayout NewLayout)
 				[
 					SNew(SVerticalBox)
 
-					+ SVerticalBox::Slot()
-					.Padding(6, 4, 0, 0)
-					.AutoHeight()
-					[
-						SNew(SHorizontalBox)
-
-						+ SHorizontalBox::Slot()
-						.AutoWidth()
-						.Padding(0, 0, 2, 0)
-						[
-							FilterComboButton.ToSharedRef()
-						]
-
-						+ SHorizontalBox::Slot()
-						.AutoWidth()
-						.MaxWidth(SearchBoxMaxWidth)
-						[
-							SNew(SBox)
-							.VAlign(VAlign_Center)
-							.WidthOverride(SearchBoxDesiredWidth)
-							[
-								LegacyContentSourceWidgets->SearchBoxPtr.ToSharedRef()
-							]
-						]
-					]
-
 					+ SVerticalBox::Slot()
 					.FillHeight(1.0f)
 					.Padding(0, 0)
@@ -969,7 +974,17 @@ TSharedRef<SWidget> SContentBrowser::CreateAssetView(const FContentBrowserConfig
 
 	LegacyContentSourceWidgets->AssetViewBorder = SNew(SBorder)
 			.BorderImage(FAppStyle::Get().GetBrush("Brushes.Panel"))
-			.Padding(AssetViewPadding);
+			.Padding(AssetViewPadding)
+			[
+				SNew(SVerticalBox)
+
+				+ SVerticalBox::Slot()
+				.FillHeight(1.0f)
+				.Padding(0, 0)
+				[
+					LegacyContentSourceWidgets->AssetViewPtr.ToSharedRef()
+				]
+			];
 
 	return LegacyContentSourceWidgets->AssetViewBorder.ToSharedRef();
 }
@@ -3533,8 +3548,8 @@ TSharedRef<SWidget> SContentBrowser::CreateLegacyAssetViewWidgets()
 									+ SHorizontalBox::Slot()
 									.AutoWidth()
 									.HAlign(HAlign_Right)
-									.VAlign(VAlign_Center)
-									.Padding(5.0f, 0.0f, 0.0f, 0.0f)
+									.VAlign(VAlign_Top)
+									.Padding(5.0f, 4.0f, 0.0f, 0.0f)
 									[
 										ViewOptions
 									]
diff --git a/Engine/Source/Editor/ContentBrowser/Private/SContentBrowser.h b/Engine/Source/Editor/ContentBrowser/Private/SContentBrowser.h
--- a/Engine/Source/Editor/ContentBrowser/Private/SContentBrowser.h
+++ b/Engine/Source/Editor/ContentBrowser/Private/SContentBrowser.h
@@ -692,7 +692,7 @@ private:
 	
 	/** Executed when the legacy content source is enabled. */
 	void OnLegacyContentSourceDisabled();
-	
+
 	/** Create the content browser widgets (SPathView, SAssetView, search and filter bars etc) */
 	TSharedRef<SWidget> CreateLegacyAssetViewWidgets();
 	
@@ -750,6 +750,9 @@ private:
 
 		/** The number of fixed slots in the splitter, slots after this number are collection view slots */
 		int32 SourceTreeSplitterNumFixedSlots = 0;
+
+		/** The containing widget for the vertical filter view */
+		TSharedPtr<SBox> VerticalFilterViewContainer;
 	};
 
 	TSharedPtr<FLegacyContentSourceWidgets> LegacyContentSourceWidgets;
