[ContentBrowser] cis fix

diff --git a/Engine/Source/Editor/ContentBrowser/Private/ContentBrowserStyle.cpp b/Engine/Source/Editor/ContentBrowser/Private/ContentBrowserStyle.cpp
--- a/Engine/Source/Editor/ContentBrowser/Private/ContentBrowserStyle.cpp
+++ b/Engine/Source/Editor/ContentBrowser/Private/ContentBrowserStyle.cpp
@@ -74,14 +74,19 @@ namespace UE::ContentBrowser::Private
 
 		Set("ContentBrowser.DefaultFolderColor", FStyleColors::AccentFolder);
 
-		FSlateBrush* AssetTreeHeaderBrush = new FSlateColorBrush(FStyleColors::Panel);
-		FSlateBrush* AssetTreeBodyBrush = new FSlateColorBrush(FStyleColors::Recessed);
+		FSlateBrush* AssetTreeHeaderBrush = nullptr; 
+		FSlateBrush* AssetTreeBodyBrush = nullptr; 
 
-		if (!Editor::ContentBrowser::IsNewStyleEnabled())
+		if (Editor::ContentBrowser::IsNewStyleEnabled())
 		{
 			AssetTreeHeaderBrush = new FSlateColorBrush(FStyleColors::Panel);
 			AssetTreeBodyBrush = new FSlateColorBrush(FStyleColors::Recessed);
 		}
+		else
+		{
+			AssetTreeHeaderBrush = new FSlateColorBrush(FStyleColors::Header);
+			AssetTreeBodyBrush = new FSlateColorBrush(FStyleColors::Recessed);
+		}
 
 		Set("ContentBrowser.AssetTreeHeaderBrush", AssetTreeHeaderBrush);
 		Set("ContentBrowser.AssetTreeBodyBrush", AssetTreeBodyBrush);
