[ContentBrowser] Content Drawer blocking details panel drop target

diff --git a/Engine/Source/Editor/ContentBrowser/Private/SContentBrowser.cpp b/Engine/Source/Editor/ContentBrowser/Private/SContentBrowser.cpp
--- a/Engine/Source/Editor/ContentBrowser/Private/SContentBrowser.cpp
+++ b/Engine/Source/Editor/ContentBrowser/Private/SContentBrowser.cpp
@@ -396,6 +396,17 @@ void SContentBrowser::Construct( const FArguments& InArgs, const FName& InInstan
 	// Update the breadcrumb trail path
 	OnContentBrowserSettingsChanged(NAME_None);
 }
+
+void SContentBrowser::OnDragLeave(const FDragDropEvent& DragDropEvent)
+{
+	SCompoundWidget::OnDragLeave(DragDropEvent);
+	// Always dismiss the ContentDrawer if the drag leave the ContentBrowser
+	if (bIsDrawer)
+	{
+		GEditor->GetEditorSubsystem<UStatusBarSubsystem>()->DismissContentBrowserDrawer();
+	}
+}
+
 END_SLATE_FUNCTION_BUILD_OPTIMIZATION
 
 void SContentBrowser::BindCommands()
diff --git a/Engine/Source/Editor/ContentBrowser/Private/SContentBrowser.h b/Engine/Source/Editor/ContentBrowser/Private/SContentBrowser.h
--- a/Engine/Source/Editor/ContentBrowser/Private/SContentBrowser.h
+++ b/Engine/Source/Editor/ContentBrowser/Private/SContentBrowser.h
@@ -120,6 +120,8 @@ public:
 	/** Constructs this widget with InArgs */
 	void Construct( const FArguments& InArgs, const FName& InstanceName, const FContentBrowserConfig* Config );
 
+	virtual void OnDragLeave(const FDragDropEvent& DragDropEvent) override;
+
 	/** Sets up an inline-name for the creation of a new asset using the specified path and the specified class and/or factory */
 	void CreateNewAsset(const FString& DefaultAssetName, const FString& PackagePath, UClass* AssetClass, UFactory* Factory);
 
