[ContentBrowser] New assets are destroyed in Tiny and Small thumbnail sizes

diff --git a/Engine/Source/Editor/ContentBrowser/Private/AssetViewWidgets.cpp b/Engine/Source/Editor/ContentBrowser/Private/AssetViewWidgets.cpp
--- a/Engine/Source/Editor/ContentBrowser/Private/AssetViewWidgets.cpp
+++ b/Engine/Source/Editor/ContentBrowser/Private/AssetViewWidgets.cpp
@@ -2320,6 +2320,7 @@ void SAssetTileItem::Construct( const FArguments& InArgs )
 		.IsSelected(InArgs._IsSelected)
 		);
 
+	bCreatedInTinySize = CurrentThumbnailSize.Get() == EThumbnailSize::Tiny;
 	bShowType = InArgs._ShowType;
 	AssetThumbnail = InArgs._AssetThumbnail;
 	ItemWidth = InArgs._ItemWidth;
@@ -2712,30 +2713,33 @@ const FSlateBrush* SAssetTileItem::GetBorderImage() const
 
 void SAssetTileItem::OnRequestingRename()
 {
-	if (InlineRenameWidget.IsValid())
+	if (InlineRenameWidget.IsValid() && CurrentThumbnailSize.Get() != EThumbnailSize::Tiny)
 	{
-		if (CurrentThumbnailSize.Get() != EThumbnailSize::Tiny)
-		{
-			InlineRenameWidget->EnterEditingMode();
-			return;
-		}
+		InlineRenameWidget->EnterEditingMode();
+		return;
+	}
 
-		const FGeometry& WidgetGeometry = GetCachedGeometry();
-		FVector2f Position = WidgetGeometry.GetAbsolutePosition();
+	const FGeometry& WidgetGeometry = GetCachedGeometry();
+	FVector2f Position = WidgetGeometry.GetAbsolutePosition();
 
-		// Lower the Y Position of the Widget by Half the Height of the TileItem minus half of the RenameWindow to keep it centered.
-		constexpr float RenameWindowHalfSize = 15.f;
-		Position.Y += (WidgetGeometry.GetAbsoluteSize().Y * 0.5f) - RenameWindowHalfSize;
+	// Lower the Y Position of the Widget by Half the Height of the TileItem minus half of the RenameWindow to keep it centered.
+	const float AppScale = FSlateApplication::Get().GetApplicationScale();
+	constexpr float RenameWindowHalfSize = 15.f;
+	Position.Y += (WidgetGeometry.GetAbsoluteSize().Y * 0.5f) - (RenameWindowHalfSize * AppScale);
 
-		// DPI scale and AppScale is handled by the SRenameWindow
-		FSlateApplication::Get().AddWindow(
-			SNew(SRenameWindow)
-			.ScreenPosition(Position)
-			.OnBeginTextEdit(this, &SAssetTileItem::HandleBeginNameChange)
-			.OnVerifyTextChanged(this, &SAssetTileItem::HandleVerifyNameChanged)
-			.InitialText(InlineRenameWidget->GetText())
-			.OnTextCommitted(this, &SAssetTileItem::HandleNameCommitted));
-	}
+	// DPI scale and AppScale is handled by the SRenameWindow
+	FSlateApplication::Get().AddWindow(
+		SNew(SRenameWindow)
+		.ScreenPosition(Position)
+		.ForceRename(bCreatedInTinySize)
+		.OnBeginTextEdit(this, &SAssetTileItem::HandleBeginNameChange)
+		.OnVerifyTextChanged(this, &SAssetTileItem::HandleVerifyNameChanged)
+		.InitialText(InlineRenameWidget->GetText())
+		.OnTextCommitted(this, &SAssetTileItem::HandleNameCommitted));
+
+	// After the first time this is called, it shouldn't force the rename anymore.
+	// This is used on AssetCreation rename to avoid destroying the asset immediately if no name is given.
+	bCreatedInTinySize = false;
 }
 
 float SAssetTileItem::GetExtraStateIconWidth() const
diff --git a/Engine/Source/Editor/ContentBrowser/Private/SRenameWindow.cpp b/Engine/Source/Editor/ContentBrowser/Private/SRenameWindow.cpp
--- a/Engine/Source/Editor/ContentBrowser/Private/SRenameWindow.cpp
+++ b/Engine/Source/Editor/ContentBrowser/Private/SRenameWindow.cpp
@@ -32,12 +32,12 @@ class SRenameEditableTextBox : public SEditableTextBox
 			.OnTextCommitted(this, &SRenameEditableTextBox::OnRenameTextCommitted));
 	}
 
-	void EndRename()
+	void EndRename(bool bForceRename = false)
 	{
 		if (!bOnTextCommittedCalled)
 		{
 			bOnTextCommittedCalled = true;
-			OnRenameCommittedDelegate.ExecuteIfBound(GetText(), ETextCommit::OnCleared);
+			OnRenameCommittedDelegate.ExecuteIfBound(GetText(), bForceRename ? ETextCommit::OnEnter : ETextCommit::OnCleared);
 		}
 	}
 
@@ -77,6 +77,8 @@ private:
 
 void SRenameWindow::Construct(const FArguments& InArgs)
 {
+	bForceRename = InArgs._ForceRename;
+
 	float DPIScale = 1.f;
 	{
 		IMainFrameModule& MainFrameModule = FModuleManager::LoadModuleChecked<IMainFrameModule>(TEXT("MainFrame"));
@@ -132,7 +134,7 @@ void SRenameWindow::DeactivateWindow()
 {
 	if (RenameEditableTextBox.IsValid())
 	{
-		RenameEditableTextBox->EndRename();
+		RenameEditableTextBox->EndRename(bForceRename);
 	}
 	RequestDestroyWindow();
 }
diff --git a/Engine/Source/Editor/ContentBrowser/Private/SRenameWindow.h b/Engine/Source/Editor/ContentBrowser/Private/SRenameWindow.h
--- a/Engine/Source/Editor/ContentBrowser/Private/SRenameWindow.h
+++ b/Engine/Source/Editor/ContentBrowser/Private/SRenameWindow.h
@@ -14,6 +14,11 @@ class SRenameWindow : public SWindow
 {
 	SLATE_BEGIN_ARGS(SRenameWindow)
 		: _ScreenPosition(FVector2D(0.f, 0.f))
+		, _OnTextCommitted()
+		, _OnVerifyTextChanged()
+		, _OnBeginTextEdit()
+		, _InitialText()
+		, _ForceRename(false)
 	{}
 
 	SLATE_ARGUMENT(UE::Slate::FDeprecateVector2DParameter, ScreenPosition)
@@ -26,6 +31,8 @@ class SRenameWindow : public SWindow
 
 	SLATE_ARGUMENT(FText, InitialText)
 
+	SLATE_ARGUMENT(bool, ForceRename)
+
 	SLATE_END_ARGS()
 
 	void Construct(const FArguments& InArgs);
@@ -37,4 +44,5 @@ private:
 
 private:
 	TSharedPtr<SRenameEditableTextBox> RenameEditableTextBox;
+	bool bForceRename;
 };
diff --git a/Engine/Source/Editor/ContentBrowser/Public/AssetViewWidgets.h b/Engine/Source/Editor/ContentBrowser/Public/AssetViewWidgets.h
--- a/Engine/Source/Editor/ContentBrowser/Public/AssetViewWidgets.h
+++ b/Engine/Source/Editor/ContentBrowser/Public/AssetViewWidgets.h
@@ -820,6 +820,9 @@ private:
 
 	/** Current thumbnail size when this widget was generated */
 	TAttribute<EThumbnailSize> CurrentThumbnailSize;
+
+	/** Set to true if the current size during the tile item creation is Tiny */
+	bool bCreatedInTinySize = false;
 };
 
 /** An item in the asset column view */
