[ContentBrowser] Fix crash when exiting Editor: was caused by double delete of FSlateColorBrush
- Creating distinct brushes instead of reusing the same

#rnx

diff --git a/Engine/Source/Editor/ContentBrowser/Private/ContentBrowserStyle.cpp b/Engine/Source/Editor/ContentBrowser/Private/ContentBrowserStyle.cpp
--- a/Engine/Source/Editor/ContentBrowser/Private/ContentBrowserStyle.cpp
+++ b/Engine/Source/Editor/ContentBrowser/Private/ContentBrowserStyle.cpp
@@ -75,16 +75,23 @@ namespace UE::ContentBrowser::Private
 		Set("ContentBrowser.DefaultFolderColor", FStyleColors::AccentFolder);
 
 		FSlateBrush* AssetTreeHeaderBrush = nullptr; 
+		FSlateBrush* VerticalFilterViewHeaderBrush = nullptr;
+		FSlateBrush* VerticalFilterViewBodyBrush = nullptr;
+
 		FSlateBrush* AssetTreeBodyBrush = nullptr; 
 
 		if (Editor::ContentBrowser::IsNewStyleEnabled())
 		{
 			AssetTreeHeaderBrush = new FSlateColorBrush(FStyleColors::Panel);
+			VerticalFilterViewHeaderBrush = new FSlateColorBrush(FStyleColors::Panel);
+			VerticalFilterViewBodyBrush = new FSlateColorBrush(FStyleColors::Panel);
 			AssetTreeBodyBrush = new FSlateColorBrush(FStyleColors::Recessed);
 		}
 		else
 		{
 			AssetTreeHeaderBrush = new FSlateColorBrush(FStyleColors::Header);
+			VerticalFilterViewHeaderBrush = new FSlateColorBrush(FStyleColors::Header);
+			VerticalFilterViewBodyBrush = new FSlateColorBrush(FStyleColors::Header);
 			AssetTreeBodyBrush = new FSlateColorBrush(FStyleColors::Recessed);
 		}
 
@@ -429,11 +436,11 @@ namespace UE::ContentBrowser::Private
 			static constexpr float FilterViewHorizontalPadding = 10.0f;
 			static constexpr float FilterViewVerticalHeaderPadding = FilterViewHorizontalPadding - 6.0f;
 
-			Set("ContentBrowser.VerticalFilterViewHeaderBrush", AssetTreeHeaderBrush);
+			Set("ContentBrowser.VerticalFilterViewHeaderBrush", VerticalFilterViewHeaderBrush);
 			Set("ContentBrowser.VerticalFilterViewHeaderPadding", FMargin(FilterViewHorizontalPadding, FilterViewVerticalHeaderPadding));
 			Set("ContentBrowser.VerticalFilterViewHeaderTextHeight", CoreStyleConstants::Icon24x24.Y); // Based on button height when dictated by an icon, matches AssetTree header
 
-			Set("ContentBrowser.VerticalFilterViewBodyBrush", AssetTreeHeaderBrush); // Body is the same as the header brush/color
+			Set("ContentBrowser.VerticalFilterViewBodyBrush", VerticalFilterViewBodyBrush);
 		}
 
 		FSlateStyleRegistry::RegisterSlateStyle(*this);
