[ContentBrowser] Styling tweaks

diff --git a/Engine/Source/Editor/ContentBrowser/Private/ContentBrowserStyle.cpp b/Engine/Source/Editor/ContentBrowser/Private/ContentBrowserStyle.cpp
--- a/Engine/Source/Editor/ContentBrowser/Private/ContentBrowserStyle.cpp
+++ b/Engine/Source/Editor/ContentBrowser/Private/ContentBrowserStyle.cpp
@@ -372,7 +372,17 @@ namespace UE::ContentBrowser::Private
 
 		if (UE::Editor::ContentBrowser::IsNewStyleEnabled())
 		{
-			ContentBrowserToolBarStyle.SetButtonStyle(FAppStyle::Get().GetWidgetStyle<FButtonStyle>("Button"));
+			FButtonStyle ButtonStyle = FAppStyle::Get().GetWidgetStyle<FButtonStyle>("Button");
+
+			constexpr int32 HorizontalButtonPadding = 2;
+
+			FMargin ButtonPadding = ContentBrowserToolBarStyle.ButtonPadding;
+			ButtonPadding.Left = HorizontalButtonPadding;
+			ButtonPadding.Right = HorizontalButtonPadding;
+
+			ContentBrowserToolBarStyle.SetButtonPadding(ButtonPadding);
+
+			ContentBrowserToolBarStyle.SetButtonStyle(ButtonStyle);
 		}
 
 		// Separate style for buttons prior to the search box to get around SToolBarComboButtonBlock referencing the wrong style
diff --git a/Engine/Source/Editor/ContentBrowser/Private/SAssetView.cpp b/Engine/Source/Editor/ContentBrowser/Private/SAssetView.cpp
--- a/Engine/Source/Editor/ContentBrowser/Private/SAssetView.cpp
+++ b/Engine/Source/Editor/ContentBrowser/Private/SAssetView.cpp
@@ -1369,8 +1369,7 @@ void SAssetView::Construct( const FArguments& InArgs )
 			.VAlign(VAlign_Fill)
 			[
 				SAssignNew(ViewContainer, SBox)
-				.Padding(6.0f)
-
+				.Padding(UE::Editor::ContentBrowser::IsNewStyleEnabled() ? FMargin(6.0f, 0.0f) : 6.0f)
 			]
 
 			+ SOverlay::Slot()
diff --git a/Engine/Source/Editor/ContentBrowser/Private/SContentBrowser.cpp b/Engine/Source/Editor/ContentBrowser/Private/SContentBrowser.cpp
--- a/Engine/Source/Editor/ContentBrowser/Private/SContentBrowser.cpp
+++ b/Engine/Source/Editor/ContentBrowser/Private/SContentBrowser.cpp
@@ -1179,6 +1179,7 @@ TSharedRef<SWidget> SContentBrowser::CreatePathView(const FContentBrowserConfig*
 			SourcesSearch,
 			LegacyContentSourceWidgets->PathViewPtr.ToSharedRef())
 			.Label(FText::FromString(FApp::GetProjectName()))
+			.ExpandedByDefault(true) // The path area, unlike the favorites and collection areas, are expanded by default (unless overridden)
 			.OnExpansionChanged(this, &SContentBrowser::SetPathViewExpanded);
 
 	return ViewWidget;
@@ -3443,6 +3444,7 @@ TSharedRef<SWidget> SContentBrowser::CreateLegacyAssetViewWidgets()
 	if (UE::Editor::ContentBrowser::IsNewStyleEnabled())
 	{
 		constexpr float ToolBarVerticalPadding = 4.0f; // ToolBar total height should be 36.0f
+		constexpr float ToolBarButtonHeight = 24.0f; // Used for buttons that should appear to be part of the ToolBar, but aren't
 		constexpr float PanelInsetPadding = 2.0f;
 		constexpr float SourceTreeSectionPadding = 2.0f;
 
@@ -3586,9 +3588,13 @@ TSharedRef<SWidget> SContentBrowser::CreateLegacyAssetViewWidgets()
 									.AutoWidth()
 									.HAlign(HAlign_Right)
 									.VAlign(VAlign_Top)
-									.Padding(5.0f, 4.0f, 0.0f, 0.0f)
+									.Padding(5.0f, 2.0f, 0.0f, 0.0f)
 									[
-										ViewOptions
+										SNew(SBox)
+										.HeightOverride(ToolBarButtonHeight)
+										[
+											ViewOptions
+										]
 									]
 								]
 							]
@@ -3636,8 +3642,13 @@ TSharedRef<SWidget> SContentBrowser::CreateLegacyAssetViewWidgets()
 			TArray<FInt16Range>{ FInt16Range(0, 550), FInt16Range(550, 10000) },
 			[this]()
 			{
-				const float Size = this->GetCachedGeometry().GetLocalSize().X;
-				return Size;
+				if (this->LegacyContentSourceWidgets && this->LegacyContentSourceWidgets->AssetViewBorder)
+				{
+					const float Size = this->LegacyContentSourceWidgets->AssetViewBorder->GetCachedGeometry().GetLocalSize().X;
+					return Size;
+				}
+
+				return 0.0f;
 			}
 		);
 
@@ -3823,7 +3834,7 @@ TSharedRef<SWidget> SContentBrowser::CreateLegacyAssetViewWidgets()
 		LegacyContentSourceWidgets->SourceTreeSplitterNumFixedSlots = LegacyContentSourceWidgets->PathFavoriteSplitterPtr->GetChildren()->NumSlot();
 	}
 
-LegacyContentSourceWidgets->AssetContextMenu = MakeShared<FAssetContextMenu>(LegacyContentSourceWidgets->AssetViewPtr);
+	LegacyContentSourceWidgets->AssetContextMenu = MakeShared<FAssetContextMenu>(LegacyContentSourceWidgets->AssetViewPtr);
 	LegacyContentSourceWidgets->AssetContextMenu->BindCommands(Commands);
 	LegacyContentSourceWidgets->AssetContextMenu->SetOnShowInPathsViewRequested( FAssetContextMenu::FOnShowInPathsViewRequested::CreateSP(this, &SContentBrowser::OnShowInPathsViewRequested) );
 	LegacyContentSourceWidgets->AssetContextMenu->SetOnRenameRequested( FAssetContextMenu::FOnRenameRequested::CreateSP(this, &SContentBrowser::OnRenameRequested) );
@@ -3864,7 +3875,7 @@ LegacyContentSourceWidgets->AssetContextMenu = MakeShared<FAssetContextMenu>(Leg
 		// Make sure the sources view is initially visible if we were asked to show it
 		SetSourcesViewExpanded(InitConfig.bExpandSourcesView && InitConfig.bUseSourcesView);
 	}
-	//else
+	// else
 	{
 		// in case we do not have a config, see what the global default settings are for the Sources Panel
 		bool bSourcesExpanded = true;
diff --git a/Engine/Source/Editor/ContentBrowser/Private/SContentBrowser.h b/Engine/Source/Editor/ContentBrowser/Private/SContentBrowser.h
--- a/Engine/Source/Editor/ContentBrowser/Private/SContentBrowser.h
+++ b/Engine/Source/Editor/ContentBrowser/Private/SContentBrowser.h
@@ -808,10 +808,10 @@ private:
 	bool bPrivateContentFilterEnabled = false;
 
 	/** Starting width of the PathView Box */
-	float PathViewBoxWidth = 150.f;
+	float PathViewBoxWidth = 230.f;
 
 	/** Starting width of the Filter Box */
-	float FilterBoxWidth = 100.f;
+	float FilterBoxWidth = 180.f;
 
 	/** Cached result of CanWriteToPath to avoid recalculating it every frame */
 	mutable bool bCachedCanWriteToCurrentPath = false;
diff --git a/Engine/Source/Editor/ContentBrowser/Private/Widgets/SContentBrowserSourceTree.cpp b/Engine/Source/Editor/ContentBrowser/Private/Widgets/SContentBrowserSourceTree.cpp
--- a/Engine/Source/Editor/ContentBrowser/Private/Widgets/SContentBrowserSourceTree.cpp
+++ b/Engine/Source/Editor/ContentBrowser/Private/Widgets/SContentBrowserSourceTree.cpp
@@ -292,6 +292,7 @@ namespace UE::Editor::ContentBrowser::Private
 
 		Id = InId;
 		bIsEmpty = InArgs._IsEmpty;
+		bExpandedByDefault = InArgs._ExpandedByDefault;
 		OnExpansionChanged = InArgs._OnExpansionChanged;
 		BodyScrollableWidget = InBody;
 
@@ -477,7 +478,7 @@ namespace UE::Editor::ContentBrowser::Private
 	{
 		const FString& IdString = Id.ToString();
 
-		bool bAreaExpanded = false;
+		bool bAreaExpanded = bExpandedByDefault;
 		GConfig->GetBool(*InIniSection, *(InSettingsString + TEXT(".") + IdString + TEXT("AreaExpanded")), bAreaExpanded, InIniFilename);
 		SetExpanded(bAreaExpanded);
 
diff --git a/Engine/Source/Editor/ContentBrowser/Private/Widgets/SContentBrowserSourceTree.h b/Engine/Source/Editor/ContentBrowser/Private/Widgets/SContentBrowserSourceTree.h
--- a/Engine/Source/Editor/ContentBrowser/Private/Widgets/SContentBrowserSourceTree.h
+++ b/Engine/Source/Editor/ContentBrowser/Private/Widgets/SContentBrowserSourceTree.h
@@ -148,6 +148,9 @@ namespace UE::Editor::ContentBrowser::Private
 			/** (Optional) label to display in the body area when the body is empty, shown as RichText. */
 			SLATE_ATTRIBUTE(FText, EmptyBodyLabel)
 
+			/** (Optional) default state for area expansion. */
+			SLATE_ARGUMENT(bool, ExpandedByDefault)
+
 			/** Called when the area is expanded or collapsed */
 			SLATE_EVENT(FOnBooleanValueChanged, OnExpansionChanged)
 		SLATE_END_ARGS()
@@ -197,5 +200,7 @@ namespace UE::Editor::ContentBrowser::Private
 		TSharedPtr<SExpandableArea> ExpandableArea;
 
 		FOnBooleanValueChanged OnExpansionChanged;
+
+		bool bExpandedByDefault = false;
 	};
 }
