[ContentBrowser] Hide automatic view switch between Grid and Tile
Hiding the functionality with a CVar: ContentBrowser.EnableGridTileSwitch

diff --git a/Engine/Source/Editor/ContentBrowser/Private/SAssetView.cpp b/Engine/Source/Editor/ContentBrowser/Private/SAssetView.cpp
--- a/Engine/Source/Editor/ContentBrowser/Private/SAssetView.cpp
+++ b/Engine/Source/Editor/ContentBrowser/Private/SAssetView.cpp
@@ -91,6 +91,14 @@
 	#define checkAssetList(cond)
 #endif
 
+static bool bEnableGridTileSwitch = false;
+static FAutoConsoleVariableRef CVarEnableGridTileSwitch(
+	TEXT("ContentBrowser.EnableGridTileSwitch"),
+	bEnableGridTileSwitch,
+	TEXT("If true Grid and List view will switch between each other when reaching certain size.\n"
+	"List > Huge -> Grid.\n"
+	"Grid < Tiny -> List."));
+
 namespace UE::AssetView
 {
 	/** Time delay between recently added items being added to the filtered asset items list */
@@ -2713,7 +2721,7 @@ FReply SAssetView::OnMouseWheel( const FGeometry& MyGeometry, const FPointerEven
 			}
 
 			// Switch the view automatically
-			if (bWasSizeSupposedToChange)
+			if (bWasSizeSupposedToChange && bEnableGridTileSwitch)
 			{
 				if (CurrentViewType == EAssetViewType::List && IsMoreThanMaxSize)
 				{
