[ContentBrowser] Filters context menu tweaks - doesn't appear when no filters present, active in all negative space to the right of the filter bar when they are present

diff --git a/Engine/Source/Developer/ToolMenus/Private/ToolMenus.cpp b/Engine/Source/Developer/ToolMenus/Private/ToolMenus.cpp
--- a/Engine/Source/Developer/ToolMenus/Private/ToolMenus.cpp
+++ b/Engine/Source/Developer/ToolMenus/Private/ToolMenus.cpp
@@ -2166,13 +2166,13 @@ void UToolMenus::PopulateToolBarBuilder(FToolBarBuilder& ToolBarBuilder, UToolMe
 			// Add a spacer before the middle and last alignment groups, and only if we've already added a section
 			if (bIsMiddleOrLast && SectionIndex > 0)
 			{
-				static constexpr float AlmostZero = 0.1f; // Using 0.0f results in different behavior, so we just use a small, sub-1 pixel value.
+				static constexpr float AlmostZero = UE_KINDA_SMALL_NUMBER; // Using 0.0f results in different behavior, so we just use a small, sub-1 pixel value (which is also used as a proportion of the overall layout space).
 
 				FMenuEntryStyleParams StyleParams;
 
 				StyleParams.HorizontalAlignment = HAlign_Right;
 				StyleParams.SizeRule = FSizeParam::ESizeRule::SizeRule_StretchContent;
-				StyleParams.FillSize = 1.0f;
+				StyleParams.FillSize = AlmostZero;
 				StyleParams.FillSizeMin = AlmostZero; // Allow the spacer to shrink to nothing if another widget needs the space.
 				StyleParams.MinimumSize = AlmostZero;
 				StyleParams.DesiredWidthOverride = StyleParams.DesiredHeightOverride = AlmostZero;
diff --git a/Engine/Source/Editor/ContentBrowser/Private/Menus/ContentBrowserMenus.cpp b/Engine/Source/Editor/ContentBrowser/Private/Menus/ContentBrowserMenus.cpp
--- a/Engine/Source/Editor/ContentBrowser/Private/Menus/ContentBrowserMenus.cpp
+++ b/Engine/Source/Editor/ContentBrowser/Private/Menus/ContentBrowserMenus.cpp
@@ -12,6 +12,8 @@
 #include "ToolMenu.h"
 #include "ToolMenus.h"
 #include "UObject/NameTypes.h"
+#include "Widgets/Colors/SColorBlock.h"
+#include "Widgets/Layout/SSpacer.h"
 
 #define LOCTEXT_NAMESPACE "ContentBrowser"
 
@@ -485,7 +487,7 @@ FDelayedAutoRegisterHelper SContentBrowser::ToolBarMenuRegistration(
 										SNew(SBox)
 										.Padding(FMargin(-3))
 										[
-											ContentBrowser->LegacyContentSourceWidgets->FilterListPtr.ToSharedRef()
+											ContentBrowser->LegacyContentSourceWidgets->FilterListPtr.ToSharedRef()											
 										],
 										FText::GetEmpty()));
 
@@ -496,7 +498,15 @@ FDelayedAutoRegisterHelper SContentBrowser::ToolBarMenuRegistration(
 										{
 											if (WeakContentBrowser.IsValid())
 											{
-												return WeakContentBrowser.Pin()->GetFilterLayout() == EFilterBarLayout::Horizontal ? EVisibility::Visible : EVisibility::Collapsed;
+												const TSharedRef<SContentBrowser> StrongContentBrowser = WeakContentBrowser.Pin().ToSharedRef();
+												const bool bIsHorizontal = StrongContentBrowser->GetFilterLayout() == EFilterBarLayout::Horizontal;
+												const bool bHasAnyFilters = StrongContentBrowser->LegacyContentSourceWidgets
+													&& StrongContentBrowser->LegacyContentSourceWidgets->FilterListPtr.IsValid()
+													&& StrongContentBrowser->LegacyContentSourceWidgets->FilterListPtr->HasAnyFilters();
+												
+												return bIsHorizontal && bHasAnyFilters
+													? EVisibility::Visible
+													: EVisibility::Collapsed;
 											}
 
 											return EVisibility::Collapsed;
@@ -533,12 +543,12 @@ FDelayedAutoRegisterHelper SContentBrowser::ToolBarMenuRegistration(
 	                            .OnClicked(ContentBrowser, &SContentBrowser::ToggleLockClicked)
 	                            .AddMetaData<FTagMetaData>(FTagMetaData(TEXT("ContentBrowserLock")));
 
-	                        InSection.AddEntry(
-	                            FToolMenuEntry::InitWidget(
-	                                "Lock",
-	                                LockToggleButton,
-	                                FText::GetEmpty()
-	                            ));
+							FToolMenuEntry& LockEntry = InSection.AddEntry(
+								FToolMenuEntry::InitWidget(
+									"Lock",
+									LockToggleButton,
+									FText::GetEmpty()
+								));
 	                    }
 	                }
 	            }));
