[ContentBrowser]  New columns in List Mode don't appear until folder is refreshed

diff --git a/Engine/Source/Editor/ContentBrowser/Private/SAssetView.cpp b/Engine/Source/Editor/ContentBrowser/Private/SAssetView.cpp
--- a/Engine/Source/Editor/ContentBrowser/Private/SAssetView.cpp
+++ b/Engine/Source/Editor/ContentBrowser/Private/SAssetView.cpp
@@ -1793,6 +1793,7 @@ FContentBrowserItem SAssetView::EndCreateDeferredItem(const TSharedPtr<FAssetVie
 		// Remove the temporary item before we do any work to ensure the new item creation is not prevented
 		Items->RemoveItem(InItem);
 		FilteredAssetItems.Remove(InItem);
+		RequestQuickFrontendListRefresh();
 		RefreshList();
 
 		// If not finalizing then we just discard the temporary
