Test automation fix for
Verify_Filtering_By_Asset_Type_Shows_Only_That_Type

diff --git a/Engine/Source/Editor/ContentBrowser/Private/AssetViewWidgets.cpp b/Engine/Source/Editor/ContentBrowser/Private/AssetViewWidgets.cpp
--- a/Engine/Source/Editor/ContentBrowser/Private/AssetViewWidgets.cpp
+++ b/Engine/Source/Editor/ContentBrowser/Private/AssetViewWidgets.cpp
@@ -2564,6 +2564,7 @@ void SAssetTileItem::Construct( const FArguments& InArgs )
 										.OverflowPolicy(ETextOverflowPolicy::Ellipsis)
 										.Text(this, &SAssetTileItem::GetAssetClassText)
 										.ColorAndOpacity(this, &SAssetTileItem::GetAssetClassLabelTextColor)
+										.AddMetaData(FDriverMetaData::Id(TEXT("AssetTileClassTextBlock")))
 									]
 								]
 							]
