Merge content browser folders deterministically.  Fixes an issue where the folder display name may change each time the items refresh.

diff --git a/Engine/Source/Editor/ContentBrowser/Private/SAssetView.cpp b/Engine/Source/Editor/ContentBrowser/Private/SAssetView.cpp
--- a/Engine/Source/Editor/ContentBrowser/Private/SAssetView.cpp
+++ b/Engine/Source/Editor/ContentBrowser/Private/SAssetView.cpp
@@ -3331,6 +3331,8 @@ void FAssetViewItemCollection::RefreshItemsFromBackend(const FAssetViewContentSo
 					continue;
 				}
 
+				TArray<uint32, TInlineAllocator<4>> ToMerge;
+
 				FContentBrowserItemKey MergeWithKey(Items[StartIndex]->GetItem());
 				uint32 Other = Lookup.Next(StartIndex);
 				while (Lookup.IsValid(Other))
@@ -3342,15 +3344,46 @@ void FAssetViewItemCollection::RefreshItemsFromBackend(const FAssetViewContentSo
 						uint32 ToRemove = Other;
 						Other = Lookup.Next(Other);
 						Lookup.Remove(Bucket, ToRemove);
-						TSharedPtr<FAssetViewItem> RemovedItem = MarkItemRemoved(ToRemove);
-						Items[StartIndex]->AppendItemData(RemovedItem->GetItem()); 
-						Items[ToRemove].Reset();
+						ToMerge.Add(ToRemove);
 					}
 					else
 					{
 						Other = Lookup.Next(Other);
 					}
 				}
+
+				// For determinism, make sure we merge items preserving their original order.
+				// Different items may have different display data, and Lookup.Add_Concurrent
+				// could cause us to select different items to display each time we refresh.
+				if (!ToMerge.IsEmpty())
+				{
+					ToMerge.Add(StartIndex);
+					ToMerge.Sort();
+
+					const uint32 MergeToIndex = ToMerge[0];
+
+					if (StartIndex != MergeToIndex)
+					{
+						checkf(Items[MergeToIndex]->GetItem().GetInternalItems().Num() == 1, TEXT("New items should only have a single internal item before merging."));
+
+						// Update the Item's index to StartIndex.
+						Items[MergeToIndex]->ResetItemData(MergeToIndex, StartIndex, *Items[MergeToIndex]->GetItem().GetPrimaryInternalItem());
+					}
+
+					for (uint32 ToRemove : MakeArrayView(ToMerge).RightChop(1))
+					{
+						TSharedPtr<FAssetViewItem> RemovedItem = MarkItemRemoved(ToRemove);
+						Items[MergeToIndex]->AppendItemData(RemovedItem->GetItem());
+						Items[ToRemove].Reset();
+					}
+
+					// Make sure StartIndex still points to the merged item.
+					if (StartIndex != MergeToIndex)
+					{
+						Swap(Items[StartIndex], Items[MergeToIndex]);
+						Swap(FilterState[StartIndex], FilterState[MergeToIndex]);
+					}
+				}
 			}
 
 		}, EParallelForFlags::Unbalanced | (UE::AssetView::AllowParallelism ? EParallelForFlags::None : EParallelForFlags::ForceSingleThread));
