[ContentBrowser] "Browse To" checks the given content browser, rather than always checking the primary one.

diff --git a/Engine/Source/Editor/ContentBrowser/Private/ContentBrowserSingleton.cpp b/Engine/Source/Editor/ContentBrowser/Private/ContentBrowserSingleton.cpp
--- a/Engine/Source/Editor/ContentBrowser/Private/ContentBrowserSingleton.cpp
+++ b/Engine/Source/Editor/ContentBrowser/Private/ContentBrowserSingleton.cpp
@@ -402,16 +402,16 @@ TSharedPtr<SContentBrowser> FContentBrowserSingleton::FindContentBrowserToSync(b
 	}
 
 	auto CanContentBrowserBeSynced =
-		[ this, bAllowLockedBrowsers ]( const TWeakPtr< SContentBrowser >& ContentBrowser ) -> bool
+		[ this, bAllowLockedBrowsers, PinnedDrawer = ContentBrowserDrawer.Pin() ]( const TSharedPtr< SContentBrowser >& ContentBrowser ) -> bool
 		{
 			// Content browser can be synced if:
 			// - It's valid.
 			// - It's not locked or we allow locked browsers (bAllowLockedBrowsers).
 			// - If it's the drawer, the active window needs to have a status bar so that it can show the drawer.
 
-			bool bCanBeSynced = ContentBrowser.IsValid() && ( bAllowLockedBrowsers || !PrimaryContentBrowser.Pin()->IsLocked() );
+			bool bCanBeSynced = ContentBrowser.IsValid() && ( bAllowLockedBrowsers || !ContentBrowser->IsLocked() );
 
-			if ( ContentBrowser == ContentBrowserDrawer )
+			if ( ContentBrowser == PinnedDrawer )
 			{
 				bCanBeSynced = bCanBeSynced && GEditor->GetEditorSubsystem<UStatusBarSubsystem>()->ActiveWindowHasStatusBar();
 			}
@@ -419,7 +419,7 @@ TSharedPtr<SContentBrowser> FContentBrowserSingleton::FindContentBrowserToSync(b
 			return bCanBeSynced;
 		};
 
-	if ( CanContentBrowserBeSynced( PrimaryContentBrowser ) )
+	if ( CanContentBrowserBeSynced( PrimaryContentBrowser.Pin() ) )
 	{
 		// If wanting to spawn a new browser window, don't set the BrowserToSync in order to summon a new browser
 		if (!bNewSpawnBrowser)
@@ -432,7 +432,8 @@ TSharedPtr<SContentBrowser> FContentBrowserSingleton::FindContentBrowserToSync(b
 	{
 
 		// If there is no primary or it is locked, sync the content browser drawer if it's available
-		if ( CanContentBrowserBeSynced( ContentBrowserDrawer ) )
+		if ( PrimaryContentBrowser != ContentBrowserDrawer
+			&& CanContentBrowserBeSynced( ContentBrowserDrawer.Pin() ) )
 		{
 			ContentBrowserToSync = ContentBrowserDrawer.Pin();
 		}
@@ -441,7 +442,7 @@ TSharedPtr<SContentBrowser> FContentBrowserSingleton::FindContentBrowserToSync(b
 			// if the drawer isn't available, find the first non-locked valid browser
 			for ( const TWeakPtr< SContentBrowser >& ContentBrowser : AllContentBrowsers )
 			{
-				if ( CanContentBrowserBeSynced( ContentBrowser ) )
+				if ( CanContentBrowserBeSynced( ContentBrowser.Pin() ) )
 				{
 					ContentBrowserToSync = ContentBrowser.Pin();
 					break;
