[ContentBrowser] SourceTree expansion restores (to default) when search cleared

diff --git a/Engine/Source/Editor/ContentBrowser/Private/SPathView.cpp b/Engine/Source/Editor/ContentBrowser/Private/SPathView.cpp
--- a/Engine/Source/Editor/ContentBrowser/Private/SPathView.cpp
+++ b/Engine/Source/Editor/ContentBrowser/Private/SPathView.cpp
@@ -1214,6 +1214,11 @@ void SPathView::Tick(const FGeometry& AllottedGeometry, const double InCurrentTi
 		bLastShowRedirectors = bNewShowRedirectors;
 		HandleSettingChanged("ShowRedirectors");
 	}
+
+	if (bLastExpandedPathsDirty)
+	{
+		UpdateLastExpandedPathsIfDirty();
+	}
 }
 
 FReply SPathView::OnKeyDown(const FGeometry& MyGeometry, const FKeyEvent& InKeyEvent)
@@ -2200,18 +2205,29 @@ void SPathView::FilterUpdated()
 	{
 		TreeData->ClearItemFilterState();
 		TreeViewPtr->ClearExpandedItems();
-		TArray<TSharedPtr<FTreeItem>> SelectedItems = TreeViewPtr->GetSelectedItems();
-		for (const TSharedPtr<FTreeItem>& SelectedItem : SelectedItems)
+
+		// First expand the default expanded paths
+		for (const FName PathToExpand : GetDefaultPathsToExpand())
 		{
-			for (TSharedPtr<FTreeItem> Parent = SelectedItem->GetParent(); Parent.IsValid();
-				 Parent = Parent->GetParent())
+			if (TSharedPtr<FTreeItem> FoundItem = TreeData->FindTreeItem(PathToExpand))
 			{
-				TreeViewPtr->SetItemExpansion(Parent, true);
+				RecursiveExpandParents(FoundItem);
+				TreeViewPtr->SetItemExpansion(FoundItem, true);
 			}
 		}
 
-		if (SelectedItems.Num())
+		TArray<TSharedPtr<FTreeItem>> SelectedItems = TreeViewPtr->GetSelectedItems();
+		if (SelectedItems.IsEmpty())
 		{
+			for (const TSharedPtr<FTreeItem>& SelectedItem : SelectedItems)
+			{
+				for (TSharedPtr<FTreeItem> Parent = SelectedItem->GetParent(); Parent.IsValid();
+					 Parent = Parent->GetParent())
+				{
+					TreeViewPtr->SetItemExpansion(Parent, true);
+				}
+			}
+
 			TreeViewPtr->RequestScrollIntoView(SelectedItems[0]);
 		}
 	}
