[ContentBrowser] SourceTree fills when expanded

diff --git a/Engine/Source/Editor/ContentBrowser/Private/Widgets/SContentBrowserSourceTree.cpp b/Engine/Source/Editor/ContentBrowser/Private/Widgets/SContentBrowserSourceTree.cpp
--- a/Engine/Source/Editor/ContentBrowser/Private/Widgets/SContentBrowserSourceTree.cpp
+++ b/Engine/Source/Editor/ContentBrowser/Private/Widgets/SContentBrowserSourceTree.cpp
@@ -5,12 +5,11 @@
 #include "ContentBrowserStyle.h"
 #include "Framework/Layout/IScrollableWidget.h"
 #include "Misc/ConfigCacheIni.h"
-#include "SContentBrowser.h"
 #include "SlateOptMacros.h"
 #include "SourcesSearch.h"
 #include "SSearchToggleButton.h"
-#include "Widgets/Colors/SColorBlock.h"
 #include "Widgets/Input/SSearchBox.h"
+#include "Widgets/Layout/SBox.h"
 #include "Widgets/Layout/SExpandableArea.h"
 #include "Widgets/Layout/SScrollBorder.h"
 #include "Widgets/Layout/SWidgetSwitcher.h"
@@ -23,7 +22,7 @@
 
 namespace UE::Editor::ContentBrowser::Private
 {
-	static constexpr float MinBodyHeight = 88.0f;
+	constexpr float MinBodyHeight = 88.0f;
 
 	BEGIN_SLATE_FUNCTION_BUILD_OPTIMIZATION
 
@@ -64,11 +63,17 @@ namespace UE::Editor::ContentBrowser::Private
 	SSplitter::ESizeRule SContentBrowserSourceTree::FSlot::GetExpandedSizeRule() const
 	{
 		// If the slot is empty, we size to content (to enforce a fixed size and disallow resizing)
-		return EntryWidget->IsExpanded() && IsVisible() && !EntryWidget->IsEmpty()
-			? ExpandedSizeRule
+		return EntryWidget->IsExpanded() && IsVisible()
+			? EntryWidget->IsEmpty()
+				? EmptyExpandedSizeRule : ExpandedSizeRule
 			: SSplitter::ESizeRule::SizeToContent;
 	}
 
+	SSplitter::ESizeRule SContentBrowserSourceTree::FSlot::GetEmptyExpandedSizeRule() const
+	{
+		return EmptyExpandedSizeRule;
+	}
+
 	float SContentBrowserSourceTree::FSlot::GetMinHeight() const
 	{
 		return IsVisible()
@@ -265,15 +270,15 @@ namespace UE::Editor::ContentBrowser::Private
 			.HAlign(HAlign_Fill)
 			.VAlign(VAlign_Fill)
 			.MinHeight(TotalHeaderHeightAttribute)
-			.FillContentHeight(1.0f)
+			.FillHeight(1.0f)
 			[
 				Splitter.ToSharedRef()
 			]
 
 			+ SVerticalBox::Slot()
 			.HAlign(HAlign_Fill)
-			.VAlign(VAlign_Fill)
-			.FillContentHeight(1.0f, 0.0f)
+			.VAlign(VAlign_Bottom)
+			.AutoHeight()
 			[
 				SNullWidget::NullWidget
 			]
@@ -296,25 +301,25 @@ namespace UE::Editor::ContentBrowser::Private
 		OnExpansionChanged = InArgs._OnExpansionChanged;
 		BodyScrollableWidget = InBody;
 
-		static const float HorizontalPadding =
+		const float HorizontalPadding =
 			UE::Editor::ContentBrowser::IsNewStyleEnabled()
 			? 10.0f
 			: 4.0f;
 
-		static const float VerticalHeaderPadding =
+		const float VerticalHeaderPadding =
 			UE::Editor::ContentBrowser::IsNewStyleEnabled()
 			? HorizontalPadding - 6.0f
 			: 0.0f;
 
-		static const float SearchButtonRightPadding =
+		const float SearchButtonRightPadding =
 			UE::Editor::ContentBrowser::IsNewStyleEnabled()
 			? 0.0f
 			: 4.0f;
 
-		static constexpr float EmptyBodyLabelPadding = 8.0f;
+		constexpr float EmptyBodyLabelPadding = 8.0f;
 
-		static const FMargin HeaderPadding = FMargin(HorizontalPadding, VerticalHeaderPadding);
-		static const FMargin ExpandableAreaPadding = FMargin(0.f, 1.f, 0.f , 0.f);
+		const FMargin HeaderPadding = FMargin(HorizontalPadding, VerticalHeaderPadding);
+		const FMargin ExpandableAreaPadding = FMargin(0.f, 1.f, 0.f , 0.f);
 
 		TSharedPtr<SWidget> SearchButtonWidget = SNullWidget::NullWidget;
 		if (InSearch.IsValid())
diff --git a/Engine/Source/Editor/ContentBrowser/Private/Widgets/SContentBrowserSourceTree.h b/Engine/Source/Editor/ContentBrowser/Private/Widgets/SContentBrowserSourceTree.h
--- a/Engine/Source/Editor/ContentBrowser/Private/Widgets/SContentBrowserSourceTree.h
+++ b/Engine/Source/Editor/ContentBrowser/Private/Widgets/SContentBrowserSourceTree.h
@@ -50,6 +50,7 @@ namespace UE::Editor::ContentBrowser::Private
 
 		public:
 			SSplitter::ESizeRule GetExpandedSizeRule() const;
+			SSplitter::ESizeRule GetEmptyExpandedSizeRule() const;
 			float GetMinHeight() const;
 			float GetHeaderHeight() const;
 			float GetInitialSlotSize() const;
@@ -62,6 +63,7 @@ namespace UE::Editor::ContentBrowser::Private
 
 		private:
 			SSplitter::ESizeRule ExpandedSizeRule = SSplitter::ESizeRule::FractionOfParent;
+			SSplitter::ESizeRule EmptyExpandedSizeRule = SSplitter::ESizeRule::FractionOfParent;
 			float HeaderHeight = 36.0f;
 			float InitialSlotSize = 0.0f;
 			float SlotSize = 0.0f; // The current slot height
