[ContentBrowser] SPathView cis fix

diff --git a/Engine/Source/Editor/ContentBrowser/Private/SPathView.cpp b/Engine/Source/Editor/ContentBrowser/Private/SPathView.cpp
--- a/Engine/Source/Editor/ContentBrowser/Private/SPathView.cpp
+++ b/Engine/Source/Editor/ContentBrowser/Private/SPathView.cpp
@@ -2217,7 +2217,7 @@ void SPathView::FilterUpdated()
 		}
 
 		TArray<TSharedPtr<FTreeItem>> SelectedItems = TreeViewPtr->GetSelectedItems();
-		if (SelectedItems.IsEmpty())
+		if (!SelectedItems.IsEmpty())
 		{
 			for (const TSharedPtr<FTreeItem>& SelectedItem : SelectedItems)
 			{
