[ContentBrowser] Moved navigation/path bar to the top toolbar

diff --git a/Engine/Source/Editor/ContentBrowser/Private/Menus/ContentBrowserMenus.cpp b/Engine/Source/Editor/ContentBrowser/Private/Menus/ContentBrowserMenus.cpp
--- a/Engine/Source/Editor/ContentBrowser/Private/Menus/ContentBrowserMenus.cpp
+++ b/Engine/Source/Editor/ContentBrowser/Private/Menus/ContentBrowserMenus.cpp
@@ -12,7 +12,6 @@
 #include "ToolMenu.h"
 #include "ToolMenus.h"
 #include "UObject/NameTypes.h"
-#include "Widgets/Colors/SColorBlock.h"
 #include "Widgets/Layout/SSpacer.h"
 
 #define LOCTEXT_NAMESPACE "ContentBrowser"
@@ -28,122 +27,7 @@ FDelayedAutoRegisterHelper SContentBrowser::NavigationBarMenuRegistration(
 		ToolBar->SetStyleSet(&UE::ContentBrowser::Private::FContentBrowserStyle::Get());
 		ToolBar->StyleName = "ContentBrowser.ToolBar";
 
-		ToolBar->FindOrAddSection("History")
-		.AddDynamicEntry(
-			"History",
-			FNewToolMenuSectionDelegate::CreateLambda([](FToolMenuSection& InSection)
-			{
-				if (!UE::Editor::ContentBrowser::IsNewStyleEnabled())
-				{
-					return;
-				}
-
-				if (const UContentBrowserToolbarMenuContext* Context = InSection.FindContext<UContentBrowserToolbarMenuContext>())
-				{
-					if (Context->ContentBrowser.IsValid())
-					{
-						const TSharedRef<SContentBrowser> ContentBrowser = Context->ContentBrowser.Pin().ToSharedRef();
-
-						const TSharedRef<SActionButton> HistoryBackButton =
-							SNew(SActionButton)
-							.ActionButtonType(EActionButtonType::Simple)
-							.ToolTipText(ContentBrowser, &SContentBrowser::GetHistoryBackTooltip)
-							.IsEnabled(ContentBrowser, &SContentBrowser::IsBackEnabled)
-							.OnClicked(ContentBrowser, &SContentBrowser::BackClicked)
-							.Icon(FAppStyle::Get().GetBrush("Icons.CircleArrowLeft"))
-							.AddMetaData<FTagMetaData>(FTagMetaData(TEXT("ContentBrowserHistoryBack")));
-
-						InSection.AddEntry(
-							FToolMenuEntry::InitWidget(
-								"HistoryBackButton",
-								HistoryBackButton,
-								FText::GetEmpty()
-							));
-
-						const TSharedRef<SActionButton> HistoryForwardButton =
-							SNew(SActionButton)
-							.ActionButtonType(EActionButtonType::Simple)
-							.ToolTipText(ContentBrowser, &SContentBrowser::GetHistoryForwardTooltip)
-							.IsEnabled(ContentBrowser, &SContentBrowser::IsForwardEnabled)
-							.OnClicked(ContentBrowser, &SContentBrowser::ForwardClicked)
-							.Icon(FAppStyle::Get().GetBrush("Icons.CircleArrowRight"))
-							.AddMetaData<FTagMetaData>(FTagMetaData(TEXT("ContentBrowserHistoryForward")));
-
-						InSection.AddEntry(
-							FToolMenuEntry::InitWidget(
-								"HistoryForwardButton",
-								HistoryForwardButton,
-								FText::GetEmpty()
-							));
-					}
-				}
-			}));
-
-		ToolBar->FindOrAddSection("Path", { }, FToolMenuInsert("History", EToolMenuInsertType::After))
-		.AddDynamicEntry(
-			"Path",
-			FNewToolMenuSectionDelegate::CreateLambda(
-			[](FToolMenuSection& InSection)
-			{
-				if (!UE::Editor::ContentBrowser::IsNewStyleEnabled())
-				{
-					return;
-				}
-
-				static constexpr float NavigationPathMinWidth = 180.0f;
-				static constexpr float NavigationPathMaxWidth = 640.0f;
-
-				if (const UContentBrowserToolbarMenuContext* Context = InSection.FindContext<UContentBrowserToolbarMenuContext>())
-				{
-					if (Context->ContentBrowser.IsValid())
-					{
-						const TSharedRef<SContentBrowser> ContentBrowser = Context->ContentBrowser.Pin().ToSharedRef();
-
-						// Should have been initialized in SContentBrowser::CreateAssetView
-						if (!ensure(ContentBrowser->LegacyContentSourceWidgets->NavigationBar.IsValid()))
-						{
-							return;
-						}
-
-						FToolMenuEntry& PathBoxEntry =
-							InSection.AddEntry(
-								FToolMenuEntry::InitWidget(
-									"PathBox",
-									ContentBrowser->LegacyContentSourceWidgets->NavigationBar.ToSharedRef(),
-									FText::GetEmpty()));
-
-						ContentBrowser->LegacyContentSourceWidgets->NavigationBarSizeSwitcher->SetSizeRange(FInt16Range(static_cast<int16>(FMath::FloorToInt(NavigationPathMaxWidth))));
-
-						PathBoxEntry.WidgetData.StyleParams.SizeRule = FSizeParam::SizeRule_StretchContent;
-						PathBoxEntry.WidgetData.StyleParams.MinimumSize = NavigationPathMinWidth;
-						PathBoxEntry.WidgetData.StyleParams.MaximumSize = TAttribute<float>::Create(
-							TAttribute<float>::FGetter::CreateSPLambda(
-								ContentBrowser,
-								[SizeSwitcher = ContentBrowser->LegacyContentSourceWidgets->NavigationBarSizeSwitcher.ToWeakPtr()]() -> float
-								{
-									if (SizeSwitcher.IsValid())
-									{
-										return SizeSwitcher.Pin()->GetDesiredSizeOverride();
-									}
-
-									return 0.0f;
-								}));
-
-						PathBoxEntry.WidgetData.StyleParams.DesiredWidthOverride = TAttribute<FOptionalSize>::Create(
-							TAttribute<FOptionalSize>::FGetter::CreateSPLambda(
-								ContentBrowser,
-								[SizeSwitcher = ContentBrowser->LegacyContentSourceWidgets->NavigationBarSizeSwitcher.ToWeakPtr()]() -> FOptionalSize
-								{
-									if (SizeSwitcher.IsValid())
-									{
-										return FOptionalSize(SizeSwitcher.Pin()->GetDesiredSizeOverride());
-									}
-
-									return FOptionalSize();
-								}));
-					}
-				}
-			}));
+		// @note: The navigation bar is registered here, but only used for extensions (@see: AssetViewMenus.cpp)
 	});
 
 FDelayedAutoRegisterHelper SContentBrowser::AddNewContextMenuRegistration(
@@ -347,8 +231,106 @@ FDelayedAutoRegisterHelper SContentBrowser::ToolBarMenuRegistration(
 
         if (UE::Editor::ContentBrowser::IsNewStyleEnabled())
         {
-			FToolMenuSection& FilterSection = ToolBar->FindOrAddSection("Filter", {}, FToolMenuInsert("Save", EToolMenuInsertType::After));
+			FToolMenuSection& HistorySection = ToolBar->FindOrAddSection("History", {}, FToolMenuInsert("Save", EToolMenuInsertType::After));
+        	HistorySection.ResizeParams.Wrapping.Allow = true;
+        	HistorySection.ResizeParams.Wrapping.Priority = 5;
+			HistorySection.AddDynamicEntry(
+				"History",
+				FNewToolMenuSectionDelegate::CreateLambda([](FToolMenuSection& InSection)
+				{
+					if (!UE::Editor::ContentBrowser::IsNewStyleEnabled())
+					{
+						return;
+					}
+
+					if (const UContentBrowserToolbarMenuContext* Context = InSection.FindContext<UContentBrowserToolbarMenuContext>())
+					{
+						if (Context->ContentBrowser.IsValid())
+						{
+							const TSharedRef<SContentBrowser> ContentBrowser = Context->ContentBrowser.Pin().ToSharedRef();
+
+							const TSharedRef<SActionButton> HistoryBackButton =
+								SNew(SActionButton)
+								.ActionButtonType(EActionButtonType::Simple)
+								.ToolTipText(ContentBrowser, &SContentBrowser::GetHistoryBackTooltip)
+								.IsEnabled(ContentBrowser, &SContentBrowser::IsBackEnabled)
+								.OnClicked(ContentBrowser, &SContentBrowser::BackClicked)
+								.Icon(FAppStyle::Get().GetBrush("Icons.CircleArrowLeft"))
+								.AddMetaData<FTagMetaData>(FTagMetaData(TEXT("ContentBrowserHistoryBack")));
+
+							InSection.AddEntry(
+								FToolMenuEntry::InitWidget(
+									"HistoryBackButton",
+									HistoryBackButton,
+									FText::GetEmpty()
+								));
+
+							const TSharedRef<SActionButton> HistoryForwardButton =
+								SNew(SActionButton)
+								.ActionButtonType(EActionButtonType::Simple)
+								.ToolTipText(ContentBrowser, &SContentBrowser::GetHistoryForwardTooltip)
+								.IsEnabled(ContentBrowser, &SContentBrowser::IsForwardEnabled)
+								.OnClicked(ContentBrowser, &SContentBrowser::ForwardClicked)
+								.Icon(FAppStyle::Get().GetBrush("Icons.CircleArrowRight"))
+								.AddMetaData<FTagMetaData>(FTagMetaData(TEXT("ContentBrowserHistoryForward")));
+
+							InSection.AddEntry(
+								FToolMenuEntry::InitWidget(
+									"HistoryForwardButton",
+									HistoryForwardButton,
+									FText::GetEmpty()
+								));
+						}
+					}
+				}));
+
+	        FToolMenuSection& PathSection = ToolBar->FindOrAddSection("Path", {}, FToolMenuInsert("History", EToolMenuInsertType::After));
+			PathSection.ResizeParams.Wrapping.Allow = true;
+        	PathSection.ResizeParams.Wrapping.Mode = UE::Slate::PrioritizedWrapBox::EWrapMode::Parent;
+			PathSection.ResizeParams.Wrapping.Priority = 5;
+
+	        PathSection.AddDynamicEntry(
+				"Path",
+				FNewToolMenuSectionDelegate::CreateLambda(
+				[](FToolMenuSection& InSection)
+				{
+					if (!UE::Editor::ContentBrowser::IsNewStyleEnabled())
+					{
+						return;
+					}
+
+					constexpr float NavigationPathMinWidth = 400.0f;
+
+					if (const UContentBrowserToolbarMenuContext* Context = InSection.FindContext<UContentBrowserToolbarMenuContext>())
+					{
+						if (Context->ContentBrowser.IsValid())
+						{
+							const TSharedRef<SContentBrowser> ContentBrowser = Context->ContentBrowser.Pin().ToSharedRef();
+
+							// Should have been initialized in SContentBrowser::CreateAssetView
+							if (!ensure(ContentBrowser->LegacyContentSourceWidgets->NavigationBar.IsValid()))
+							{
+								return;
+							}
+
+							FToolMenuEntry& PathBoxEntry =
+								InSection.AddEntry(
+									FToolMenuEntry::InitWidget(
+										"PathBox",
+										ContentBrowser->LegacyContentSourceWidgets->NavigationBar.ToSharedRef(),
+										FText::GetEmpty()));
+
+							PathBoxEntry.WidgetData.StyleParams.SizeRule = FSizeParam::SizeRule_StretchContent;
+							PathBoxEntry.WidgetData.StyleParams.FillSize = 1.0f;
+							PathBoxEntry.WidgetData.StyleParams.DesiredHeightOverride = 24.0f; // Match the height of the other widgets
+							PathBoxEntry.WidgetData.StyleParams.MinimumSize = NavigationPathMinWidth;
+						}
+					}
+				}));
+
+			FToolMenuSection& FilterSection = ToolBar->FindOrAddSection("Filter", {}, FToolMenuInsert("Path", EToolMenuInsertType::After));
         	FilterSection.ResizeParams.Wrapping.Allow = true;
+        	FilterSection.ResizeParams.Wrapping.ForceNewLine = true;
         	FilterSection.ResizeParams.Wrapping.Priority = 10;
 
         	FilterSection.AddDynamicEntry(
@@ -398,8 +380,8 @@ FDelayedAutoRegisterHelper SContentBrowser::ToolBarMenuRegistration(
                 FNewToolMenuSectionDelegate::CreateLambda(
                 [](FToolMenuSection& InSection)
                 {
-                    static constexpr float SearchBoxMinWidth = 180.0f;
-                    static constexpr float SearchBoxMaxWidth = 640.0f;
+                    constexpr float SearchBoxMinWidth = 180.0f;
+                    constexpr float SearchBoxMaxWidth = 640.0f;
 
                     if (const UContentBrowserToolbarMenuContext* Context = InSection.FindContext<UContentBrowserToolbarMenuContext>())
                     {
diff --git a/Engine/Source/Editor/ContentBrowser/Private/SContentBrowser.cpp b/Engine/Source/Editor/ContentBrowser/Private/SContentBrowser.cpp
--- a/Engine/Source/Editor/ContentBrowser/Private/SContentBrowser.cpp
+++ b/Engine/Source/Editor/ContentBrowser/Private/SContentBrowser.cpp
@@ -2394,16 +2394,6 @@ void SContentBrowser::OnSearchBoxCommitted(const FText& InSearchText, ETextCommi
 	SetSearchBoxText(InSearchText);
 }
 
-void SContentBrowser::Tick(const FGeometry& AllottedGeometry, const double InCurrentTime, const float InDeltaTime)
-{
-	if (UE::Editor::ContentBrowser::IsNewStyleEnabled() && ContentSourcesContainer->IsLegacyContentSourceActive())
-	{
-		LegacyContentSourceWidgets->NavigationToolBarParentSwitcher->Update();
-	}
-
-	SCompoundWidget::Tick(AllottedGeometry, InCurrentTime, InDeltaTime);
-}
-
 FReply SContentBrowser::OnSearchKeyDown(const FGeometry& Geometry, const FKeyEvent& InKeyEvent)
 {
 	FInputChord CheckChord(InKeyEvent.GetKey(), EModifierKey::FromBools(InKeyEvent.IsControlDown(), InKeyEvent.IsAltDown(), InKeyEvent.IsShiftDown(), InKeyEvent.IsCommandDown()));
@@ -3471,6 +3461,8 @@ TSharedRef<SWidget> SContentBrowser::CreateLegacyAssetViewWidgets()
 			nullptr,
 			FInt16Range(100)); // @note: this is overridden in menu registration
 
+		LegacyContentSourceWidgets->NavigationToolBarWidget = CreateNavigationToolBar(&InitConfig);
+
 		TSharedRef<SWidget> ToolBarWidget = CreateToolBar(&InitConfig);
 
 		LegacyContentSourceWidgets->SearchBoxSizeSwitcher->SetMaxSizeReferenceWidget(ToolBarWidget.ToSharedPtr());
@@ -3478,15 +3470,6 @@ TSharedRef<SWidget> SContentBrowser::CreateLegacyAssetViewWidgets()
 		TSharedPtr<SBox> AssetViewNavigationToolBarContainer = nullptr;
 		TSharedPtr<SBox> SourceTreeAndAssetViewNavigationToolBarContainer = nullptr;
 
-		LegacyContentSourceWidgets->NavigationBarSizeSwitcher = MakeShared<TWidgetDesiredSizeSwitcher<EAxis::X>>(
-			LegacyContentSourceWidgets->NavigationBar,
-			nullptr,
-			FInt16Range(100)); // @note: this is overridden in menu registration
-
-		LegacyContentSourceWidgets->NavigationToolBarWidget = CreateNavigationToolBar(&InitConfig);
-
-		LegacyContentSourceWidgets->NavigationBarSizeSwitcher->SetMaxSizeReferenceWidget(LegacyContentSourceWidgets->NavigationToolBarWidget);
-
 		FinalWidget = 
 		SNew(SVerticalBox)
 
@@ -3636,22 +3619,6 @@ TSharedRef<SWidget> SContentBrowser::CreateLegacyAssetViewWidgets()
 			]
 		];
 
-		LegacyContentSourceWidgets->NavigationToolBarParentSwitcher = MakeShared<UE::Editor::ContentBrowser::Private::TWidgetSizeParentSwitcher<Orient_Horizontal>>(
-			LegacyContentSourceWidgets->NavigationToolBarWidget,
-			TArray<TSharedPtr<SBox>>{ SourceTreeAndAssetViewNavigationToolBarContainer, AssetViewNavigationToolBarContainer },
-			TArray<FInt16Range>{ FInt16Range(0, 550), FInt16Range(550, 10000) },
-			[this]()
-			{
-				if (this->LegacyContentSourceWidgets && this->LegacyContentSourceWidgets->AssetViewBorder)
-				{
-					const float Size = this->LegacyContentSourceWidgets->AssetViewBorder->GetCachedGeometry().GetLocalSize().X;
-					return Size;
-				}
-
-				return 0.0f;
-			}
-		);
-
 		LegacyContentSourceWidgets->SourceTreeSplitterNumFixedSlots = LegacyContentSourceWidgets->SourceTreePtr->GetSplitter()->GetChildren()->NumSlot();
 	}
 	else
diff --git a/Engine/Source/Editor/ContentBrowser/Private/SContentBrowser.h b/Engine/Source/Editor/ContentBrowser/Private/SContentBrowser.h
--- a/Engine/Source/Editor/ContentBrowser/Private/SContentBrowser.h
+++ b/Engine/Source/Editor/ContentBrowser/Private/SContentBrowser.h
@@ -268,8 +268,6 @@ public:
 	/** Called by the editable text control when the user commits a text change */
 	void OnSearchBoxCommitted(const FText& InSearchText, ETextCommit::Type CommitInfo);
 
-	virtual void Tick(const FGeometry& AllottedGeometry, const double InCurrentTime, const float InDeltaTime) override;
-
 	/** Toggles the favorite status of an array of folders*/
 	void ToggleFolderFavorite(const TArray<FString>& FolderPaths);
 
@@ -743,12 +741,6 @@ private:
 
 		/** The splitter between the path, favorite and collection views */
 		TSharedPtr<SSplitter> PathFavoriteSplitterPtr;
-		
-		/** Handles the navigation bar location switching */
-		TSharedPtr<UE::Editor::ContentBrowser::Private::TWidgetSizeParentSwitcher<Orient_Horizontal>> NavigationToolBarParentSwitcher;
-
-		/** Handles the navigation bar sizing within the navigation toolbar */
-		TSharedPtr<UE::Editor::ContentBrowser::Private::TWidgetDesiredSizeSwitcher<EAxis::X>> NavigationBarSizeSwitcher;
 
 		/** The number of fixed slots in the splitter, slots after this number are collection view slots */
 		int32 SourceTreeSplitterNumFixedSlots = 0;
diff --git a/Engine/Source/Runtime/Slate/Private/Framework/MultiBox/MultiBox.cpp b/Engine/Source/Runtime/Slate/Private/Framework/MultiBox/MultiBox.cpp
--- a/Engine/Source/Runtime/Slate/Private/Framework/MultiBox/MultiBox.cpp
+++ b/Engine/Source/Runtime/Slate/Private/Framework/MultiBox/MultiBox.cpp
@@ -917,6 +917,7 @@ void SMultiBoxWidget::AddBlockWidget(const FMultiBlock& Block, TSharedPtr<SHoriz
 				NewSlot.WrapPriority(ResizeParams.Wrapping.Priority.IsSet() ? ResizeParams.Wrapping.Priority : FMenuEntryResizeParams::FWrappingParams::DefaultPriority);
 				NewSlot.WrapMode(ResizeParams.Wrapping.Mode.IsSet() ? ResizeParams.Wrapping.Mode : FMenuEntryResizeParams::FWrappingParams::DefaultMode);
 				NewSlot.VerticalOverflowBehavior(ResizeParams.Wrapping.VerticalOverflowBehavior.Get(FMenuEntryResizeParams::FWrappingParams::DefaultVerticalOverflowBehavior));
+				NewSlot.ForceNewLine(ResizeParams.Wrapping.ForceNewLine.Get(FMenuEntryResizeParams::FWrappingParams::DefaultForceNewLine));
 				NewSlot.ExcludeIfFirstOrLast(Block.IsSeparator());
 
 				NewSlot.Padding(0.f)
diff --git a/Engine/Source/Runtime/Slate/Private/Framework/MultiBox/MultiBoxBuilder.cpp b/Engine/Source/Runtime/Slate/Private/Framework/MultiBox/MultiBoxBuilder.cpp
--- a/Engine/Source/Runtime/Slate/Private/Framework/MultiBox/MultiBoxBuilder.cpp
+++ b/Engine/Source/Runtime/Slate/Private/Framework/MultiBox/MultiBoxBuilder.cpp
@@ -68,6 +68,7 @@ namespace UE::Slate::Private
 				ApplyUnsetAttributeOverride(BlockResizeParams.Wrapping.Mode, SectionResizeParams.Wrapping.Mode);
 				ApplyUnsetAttributeOverride(BlockResizeParams.Wrapping.Priority, SectionResizeParams.Wrapping.Priority);
 				ApplyUnsetOptionalOverride(BlockResizeParams.Wrapping.VerticalOverflowBehavior, SectionResizeParams.Wrapping.VerticalOverflowBehavior);
+				ApplyUnsetOptionalOverride(BlockResizeParams.Wrapping.ForceNewLine, SectionResizeParams.Wrapping.ForceNewLine);
 
 				BlockResizeParamsAttribute = BlockResizeParams;
 			}
@@ -83,6 +84,7 @@ namespace UE::Slate::Private
 							ApplyUnsetAttributeOverride(ResizeParams.Wrapping.Mode, SectionResizeParams.Wrapping.Mode);
 							ApplyUnsetAttributeOverride(ResizeParams.Wrapping.Priority, SectionResizeParams.Wrapping.Priority);
 							ApplyUnsetOptionalOverride(ResizeParams.Wrapping.VerticalOverflowBehavior, SectionResizeParams.Wrapping.VerticalOverflowBehavior);
+							ApplyUnsetOptionalOverride(ResizeParams.Wrapping.ForceNewLine, SectionResizeParams.Wrapping.ForceNewLine);
 
 							return ResizeParams;
 						})));
diff --git a/Engine/Source/Runtime/Slate/Private/Framework/MultiBox/SPrioritizedWrapBox.cpp b/Engine/Source/Runtime/Slate/Private/Framework/MultiBox/SPrioritizedWrapBox.cpp
--- a/Engine/Source/Runtime/Slate/Private/Framework/MultiBox/SPrioritizedWrapBox.cpp
+++ b/Engine/Source/Runtime/Slate/Private/Framework/MultiBox/SPrioritizedWrapBox.cpp
@@ -36,6 +36,7 @@ namespace UE::Slate::PrioritizedWrapBox
 				VerticalExpansionThreshold = InSlot.GetVerticalExpansionThreshold();
 				WrapPriority = InSlot.GetWrapPriority();
 				bExcludeIfFirstOrLast = InSlot.GetExcludeIfFirstOrLast();
+				bForceNewLine = InSlot.GetForceNewLine();
 			}
 
 			/** Applies the given slot's values to this proxy. Will return true if any values have changed from those stored. */
@@ -55,6 +56,9 @@ namespace UE::Slate::PrioritizedWrapBox
 				bAnyValueChanged = bAnyValueChanged || (WrapPriority != InSlot.GetWrapPriority());
 				WrapPriority = InSlot.GetWrapPriority();
 
+				bAnyValueChanged = bAnyValueChanged || (bForceNewLine != InSlot.GetForceNewLine());
+				bForceNewLine = InSlot.GetForceNewLine();
+
 				bAnyValueChanged = bAnyValueChanged || (bExcludeIfFirstOrLast != InSlot.GetExcludeIfFirstOrLast());
 				bExcludeIfFirstOrLast = InSlot.GetExcludeIfFirstOrLast();
 
@@ -65,6 +69,7 @@ namespace UE::Slate::PrioritizedWrapBox
 			EVerticalOverflowBehavior VerticalOverflowBehavior = EVerticalOverflowBehavior::Default;
 			TOptional<float> VerticalExpansionThreshold;
 			int32 WrapPriority = INDEX_NONE;
+			bool bForceNewLine = false;
 			bool bExcludeIfFirstOrLast = false;
 		};
 
@@ -76,9 +81,12 @@ namespace UE::Slate::PrioritizedWrapBox
 
 			int32 WrapPriority = INDEX_NONE;
 
-			/** Whether the block can wrap at all. */
+			/** Whether the block can wrap. Note that if ForceNewLine is true, it will always move/"wrap" to a new line, */
 			bool bCanWrap = true;
 
+			/** Whether to forcibly place this block on a new line. Others can appear to the right of this block, but none to the left (if true). */
+			bool bForceNewLine = false;
+
 			/** Based on it's desired size, it was determined this doesn't vertically expand. */
 			bool bHasVerticalExpansion = false;
 
@@ -108,6 +116,7 @@ namespace UE::Slate::PrioritizedWrapBox
 				SequentialIndex = INDEX_NONE;
 				WrapPriority = INDEX_NONE;
 				bCanWrap = true;
+				bForceNewLine = false;
 				bHasVerticalExpansion = false;
 				bCanVerticallyExpand = false;
 				DesiredSize = FVector2f::Zero();
@@ -282,6 +291,59 @@ namespace UE::Slate::PrioritizedWrapBox
 
 					MakeBlockLine(CurrentLine, CurrentLineBlocks);
 
+					int32 NumSortablesUncheckedForForceNewLine = SortedBlockView.Num();
+					int32 LastSortableCheckedIndex = 0;
+
+					// Find and move the first "ForceNewLine'able" to the next line, and move all blocks with the same or higher priority.
+					while (NumSortablesUncheckedForForceNewLine > 0
+						&& LineHasMultipleBlocksAndPriorities(CurrentLine)
+						&& HasMoreRemovableBlocks())
+					{
+						const FBlock& BlockToConsider = SortedBlockView[LastSortableCheckedIndex++];
+
+						const bool bCanRemoveBlock = BlockToConsider.bForceNewLine && BlockToConsider.WrapPriority != CurrentLine.Blocks[0].WrapPriority;
+
+						--NumSortablesUncheckedForForceNewLine;
+
+						if (bCanRemoveBlock)
+						{
+							FVector2f::FReal EffectiveLengthToRemove = BlockToConsider.GetEffectiveLength();
+
+							FBlock BlockCopy = BlockToConsider;
+							if (RemoveBlockFromLine(CurrentLine, BlockToConsider))
+							{
+								int32 LastForcedPriorityRemoved = BlockCopy.WrapPriority;
+
+								// Don't include 0-length entries, they'll dictate the line height without having any visual content
+								if (!FMath::IsNearlyZero(EffectiveLengthToRemove))
+								{
+									NextLineBlocks.Emplace(BlockCopy);
+								}
+
+								// Remove block, and all after it from further consideration for this line
+								TConstArrayView<FBlock> RemainingBlockView = SortedBlockView.RightChop(LastSortableCheckedIndex);
+
+								// The block was successfully removed. Now reverse iterate over the remaining blocks and remove those with the same or higher wrapping priority.
+								// Note that we don't need to check the priority, it will naturally be the same or higher in the sorted view.
+								for (int32 BlockIndex = RemainingBlockView.Num() - 1; BlockIndex >= 0; --BlockIndex)
+								{
+									const FBlock& EndBlockToConsider = RemainingBlockView[BlockIndex];
+
+									EffectiveLengthToRemove = EndBlockToConsider.GetEffectiveLength();
+									BlockCopy = EndBlockToConsider;
+									if (RemoveBlockFromLine(CurrentLine, EndBlockToConsider))
+									{
+										// Don't include 0-length entries, they'll dictate the line height without having any visual content
+										if (!FMath::IsNearlyZero(EffectiveLengthToRemove))
+										{
+											NextLineBlocks.Emplace(BlockCopy);
+										}
+									}
+								}
+							}
+						}
+					}
+
 					int32 NumSortablesUncheckedForExpansion = SortedBlockView.Num();
 
 					// Remove vertical expandables, which rely on the full line length (not min + full combination)
@@ -290,17 +352,17 @@ namespace UE::Slate::PrioritizedWrapBox
 						&& LineHasMultipleBlocksAndPriorities(CurrentLine)
 						&& HasMoreRemovableBlocks())
 					{
-						const FBlock& BlockToRemove = SortedBlockView.Last();
-						const bool bCanRemoveBlock = BlockToRemove.bCanWrap && BlockToRemove.bCanVerticallyExpand && BlockToRemove.WrapPriority != CurrentLine.Blocks[0].WrapPriority;
+						const FBlock& BlockToConsider = SortedBlockView.Last();
+						const bool bCanRemoveBlock = BlockToConsider.bCanWrap && BlockToConsider.bCanVerticallyExpand && BlockToConsider.WrapPriority != CurrentLine.Blocks[0].WrapPriority;
 
 						--NumSortablesUncheckedForExpansion;
 
 						if (bCanRemoveBlock)
 						{
 							bool bShouldRemoveBlock = true;
-							const FVector2f::FReal EffectiveLengthToRemove = BlockToRemove.GetEffectiveLength();
+							const FVector2f::FReal EffectiveLengthToRemove = BlockToConsider.GetEffectiveLength();
 
-							if (BlockToRemove.MinLength != BlockToRemove.ReserveLength)
+							if (BlockToConsider.MinLength != BlockToConsider.ReserveLength)
 							{
 								const float LineLengthAfterRemovalOfEffective = CurrentLine.MaxLength - EffectiveLengthToRemove;
 								if (LineLengthAfterRemovalOfEffective + EffectiveLengthToRemove < ParentLineLength)
@@ -312,8 +374,8 @@ namespace UE::Slate::PrioritizedWrapBox
 
 							if (bShouldRemoveBlock)
 							{
-								FBlock BlockCopy = BlockToRemove;
-								if (RemoveBlockFromLine(CurrentLine, BlockToRemove))
+								FBlock BlockCopy = BlockToConsider;
+								if (RemoveBlockFromLine(CurrentLine, BlockToConsider))
 								{
 									// Don't include 0-length entries, they'll dictate the line height without having any visual content
 									if (!FMath::IsNearlyZero(EffectiveLengthToRemove))
@@ -625,6 +687,7 @@ namespace UE::Slate::PrioritizedWrapBox
 			CurrentBlock->ResetToDefault();
 			CurrentBlock->SequentialIndex = NumBlocks;
 			CurrentBlock->WrapPriority = ChildProxies[0].WrapPriority;
+			CurrentBlock->bForceNewLine = ChildProxies[0].bForceNewLine;
 
 			NumBlocks++;
 
@@ -639,6 +702,7 @@ namespace UE::Slate::PrioritizedWrapBox
 					CurrentBlock->ResetToDefault();
 					CurrentBlock->SequentialIndex = NumBlocks;
 					CurrentBlock->WrapPriority = Child.GetWrapPriority();
+					CurrentBlock->bForceNewLine = Child.GetForceNewLine();
 
 					NumBlocks++;
 				}
@@ -709,6 +773,7 @@ namespace UE::Slate::PrioritizedWrapBox
 				CurrentBlock->ReserveLength += OverflowSizeX;
 				CurrentBlock->MinLength += MinDesiredSizeX;
 				CurrentBlock->bCanWrap = Child.GetAllowWrapping();
+				CurrentBlock->bForceNewLine = CurrentBlock->bForceNewLine || Child.GetForceNewLine(); // If -any- child is forced, the block is too
 				CurrentBlock->Area += Area;
 
 				FSlotProxy& ChildProxy = ChildProxies[ChildIndex];
@@ -723,19 +788,23 @@ namespace UE::Slate::PrioritizedWrapBox
 			// Needs re-sorting
 			SortedBlocks.StableSort([](const FBlock& A, const FBlock& B)
 			{
+				// A block that's forced to a new line, but can't wrap (to another), is considered to be wrappable
+				const bool bACanEverWrap = A.bCanWrap || A.bForceNewLine;
+				const bool bBCanEverWrap = B.bCanWrap || B.bForceNewLine;
+				
 				// Non-wrappable children are always at the start
-				if (!A.bCanWrap && !B.bCanWrap)
+				if (!bACanEverWrap && !bBCanEverWrap)
 				{
 					return A.SequentialIndex < B.SequentialIndex;
 				}
 
 				// Non-wrappable children come before wrappable ones
-				if (!A.bCanWrap) // B is wrappable
+				if (!bACanEverWrap) // B is wrappable
 				{
 					return true;
 				}
 
-				if (!B.bCanWrap) // A is wrappable
+				if (!bBCanEverWrap) // A is wrappable
 				{
 					return false;
 				}
@@ -754,7 +823,7 @@ namespace UE::Slate::PrioritizedWrapBox
 			for (int32 BlockIndex = 0; BlockIndex < SortedBlocks.Num(); ++BlockIndex)
 			{
 				const FBlock& SortedBlock = SortedBlocks[BlockIndex];
-				if (SortedBlock.bCanWrap)
+				if (SortedBlock.bCanWrap || SortedBlock.bForceNewLine)
 				{
 					FirstWrappableBlockIndex = BlockIndex;
 					break;
@@ -841,6 +910,7 @@ void SPrioritizedWrapBox::FSlot::Construct(const FChildren& SlotOwner, FSlotArgu
 
 	VerticalOverflowBehavior = InArgs._VerticalOverflowBehavior;
 	VerticalExpansionThreshold = InArgs._VerticalExpansionThreshold;
+	bForceNewLine = InArgs._ForceNewLine;
 	bExcludeIfFirstOrLast = InArgs._ExcludeIfFirstOrLast;
 }
 
diff --git a/Engine/Source/Runtime/Slate/Private/Framework/MultiBox/SPrioritizedWrapBox.h b/Engine/Source/Runtime/Slate/Private/Framework/MultiBox/SPrioritizedWrapBox.h
--- a/Engine/Source/Runtime/Slate/Private/Framework/MultiBox/SPrioritizedWrapBox.h
+++ b/Engine/Source/Runtime/Slate/Private/Framework/MultiBox/SPrioritizedWrapBox.h
@@ -40,6 +40,7 @@ public:
 			SLATE_ATTRIBUTE(UE::Slate::PrioritizedWrapBox::EWrapMode, WrapMode)
 			SLATE_ARGUMENT_DEFAULT(UE::Slate::PrioritizedWrapBox::EVerticalOverflowBehavior, VerticalOverflowBehavior) = UE::Slate::PrioritizedWrapBox::EVerticalOverflowBehavior::Default;
 			SLATE_ARGUMENT(TOptional<float>, VerticalExpansionThreshold) // Vertical expansion will occur when the widget's desired size is at or beyond this value
+			SLATE_ARGUMENT_DEFAULT(bool, ForceNewLine) = false; // If true, this will forcibly place this slot on a new line, regardless of the wrapping behavior. This affects all slots after this one with the same or higher wrapping priorities.
 			SLATE_ARGUMENT_DEFAULT(bool, ExcludeIfFirstOrLast) = false;
 		SLATE_SLOT_END_ARGS()
 
@@ -68,6 +69,12 @@ public:
 		/** If the vertical overflow behavior requires expansion, this optionally specifies the height at which the slot should vertically expand. */
 		const TOptional<float>& GetVerticalExpansionThreshold() const { return VerticalExpansionThreshold; }
 
+		/**
+		 * If true, this will forcibly place this slot on a new line, regardless of the wrapping behavior. This affects all slots after this one with the same or higher wrapping priorities.
+		 * When AllowWrapping is false, this will only ever be placed on the line it's forcibly wrapped to.
+		 */
+		const bool GetForceNewLine() const { return bForceNewLine; }
+
 		/** If, after wrapping, this slot is the first or last in the line, it will be excluded from the result if this is true. */
 		bool GetExcludeIfFirstOrLast() const { return bExcludeIfFirstOrLast; }
 
@@ -85,6 +92,7 @@ public:
 
 		UE::Slate::PrioritizedWrapBox::EVerticalOverflowBehavior VerticalOverflowBehavior = UE::Slate::PrioritizedWrapBox::EVerticalOverflowBehavior::Default;
 		TOptional<float> VerticalExpansionThreshold;
+		bool bForceNewLine = false;;
 
 		bool bExcludeIfFirstOrLast = false;
 	};
diff --git a/Engine/Source/Runtime/Slate/Public/Framework/MultiBox/MultiBoxDefs.h b/Engine/Source/Runtime/Slate/Public/Framework/MultiBox/MultiBoxDefs.h
--- a/Engine/Source/Runtime/Slate/Public/Framework/MultiBox/MultiBoxDefs.h
+++ b/Engine/Source/Runtime/Slate/Public/Framework/MultiBox/MultiBoxDefs.h
@@ -288,6 +288,7 @@ struct FMenuEntryResizeParams
 		static constexpr int32 DefaultPriority = 0;
 		static constexpr UE::Slate::PrioritizedWrapBox::EWrapMode DefaultMode = UE::Slate::PrioritizedWrapBox::EWrapMode::Parent;
 		static constexpr UE::Slate::PrioritizedWrapBox::EVerticalOverflowBehavior DefaultVerticalOverflowBehavior = UE::Slate::PrioritizedWrapBox::EVerticalOverflowBehavior::Default;
+		static constexpr bool DefaultForceNewLine = false;
 
 		/** If true (default), allow this entry to be wrapped to the next line during resizing of toolbars. If false, this entry will never be wrapped. */
 		TAttribute<bool> Allow;
@@ -300,5 +301,8 @@ struct FMenuEntryResizeParams
 
 		/** Override to specify the vertical overflow behavior. By default, this is "ClipOrExpandProportional". */
 		TOptional<UE::Slate::PrioritizedWrapBox::EVerticalOverflowBehavior> VerticalOverflowBehavior;
+
+		/** If true (default is false), the entry should always be placed on a new line. Other entries can appear to it's right, but never to it's left. */
+		TOptional<bool> ForceNewLine;
 	} Wrapping;
 };
