[ContentBrowser] Fixes based on feedbacks
Full tooltip on by default
Add fake column view (Only difference will be columns), Revert AssetPicker setting
Tiny use FileIcon
Add command for the column view
class name and status for bp class
Separator after status and separator before collection
re-add Reset Columns
Re-Add export to csv
Separeted Column hidden logic List/Column
Add Modified badge Tiny List/Column

diff --git a/Engine/Config/BaseEditorSettings.ini b/Engine/Config/BaseEditorSettings.ini
--- a/Engine/Config/BaseEditorSettings.ini
+++ b/Engine/Config/BaseEditorSettings.ini
@@ -40,6 +40,8 @@ IncludeClassNames=True
 IncludeAssetPaths=True
 ; Whether to include collection names when searching by text
 IncludeCollectionNames=True
+; Whether to display the full tooltip on Asset and Folder or the small version with the Ctrl + Alt
+AlwaysExpandTooltips=True
 
 [/Script/UnrealEd.PropertyColorSettings]
 +CustomProperties=(Name="AffectsNavmesh",Text=NSLOCTEXT("ActorColoration", "AffectsNavmesh", "Affects NavMesh"), TextToolTip=NSLOCTEXT("ActorColoration", "AffectsNavmesh_ToolTip", "Colorize actor in Blue if it is affecting the navmesh."), PropertyChain="RootComponent.bCanEverAffectNavigation",PropertyValue="True",PropertyColor=(R=0,G=0,B=255))
diff --git a/Engine/Source/Editor/AssetDefinition/Public/AssetDefinitionAssetInfo.h b/Engine/Source/Editor/AssetDefinition/Public/AssetDefinitionAssetInfo.h
--- a/Engine/Source/Editor/AssetDefinition/Public/AssetDefinitionAssetInfo.h
+++ b/Engine/Source/Editor/AssetDefinition/Public/AssetDefinitionAssetInfo.h
@@ -73,6 +73,7 @@ struct FAssetDisplayInfo
 {
 public:
 	TAttribute<const FSlateBrush*> StatusIcon;
+	TAttribute<const FSlateBrush*> StatusIconOverlay;
 	TAttribute<FText> StatusTitle;
 	TAttribute<FText> StatusDescription;
 	TAttribute<EVisibility> IsVisible;
diff --git a/Engine/Source/Editor/ContentBrowser/Private/AssetSystemContentBrowserInfoProvider.cpp b/Engine/Source/Editor/ContentBrowser/Private/AssetSystemContentBrowserInfoProvider.cpp
--- a/Engine/Source/Editor/ContentBrowser/Private/AssetSystemContentBrowserInfoProvider.cpp
+++ b/Engine/Source/Editor/ContentBrowser/Private/AssetSystemContentBrowserInfoProvider.cpp
@@ -230,11 +230,22 @@ void FAssetSystemContentBrowserInfoProvider::PopulateAssetInfo(TArray<FAssetDisp
 						return result < 0 || (result == 0 && A.Type < B.Type);
 					});
 
+				bool bAddedCollectionHeader = false;
 				for (const FCollectionNameType& CollectionContainingObject : CollectionsContainingObject)
 				{
 					FCollectionStatusInfo CollectionStatusInfo;
 					if (CollectionContainer->GetCollectionStatusInfo(CollectionContainingObject.Name, CollectionContainingObject.Type, CollectionStatusInfo))
 					{
+						if (!bAddedCollectionHeader)
+						{
+							// StatusTitle currently used to add a separator for status, need to be changed in future version to allow more configurability
+							bAddedCollectionHeader = true;
+							FAssetDisplayInfo CollectionHeader = FAssetDisplayInfo();
+							CollectionHeader.StatusTitle = LOCTEXT("CollectionHeaderTitle", "Collection(s)");
+							CollectionHeader.StatusDescription = LOCTEXT("CollectionHeaderDescription", "");
+							OutAssetDisplayInfo.Add(CollectionHeader);
+						}
+
 						FAssetDisplayInfo CollectionInfo = FAssetDisplayInfo();
 						CollectionInfo.StatusTitle = FText::FromName(CollectionContainingObject.Name);
 						CollectionInfo.StatusDescription = FText::AsNumber(CollectionStatusInfo.NumObjects);
diff --git a/Engine/Source/Editor/ContentBrowser/Private/AssetView/AssetViewConfig.h b/Engine/Source/Editor/ContentBrowser/Private/AssetView/AssetViewConfig.h
--- a/Engine/Source/Editor/ContentBrowser/Private/AssetView/AssetViewConfig.h
+++ b/Engine/Source/Editor/ContentBrowser/Private/AssetView/AssetViewConfig.h
@@ -29,6 +29,9 @@ struct FAssetViewInstanceConfig
 
 	UPROPERTY()
 	TArray<FName> HiddenColumns;
+
+	UPROPERTY()
+	TArray<FName> ListHiddenColumns;
 };
 
 UCLASS(EditorConfig="AssetView")
diff --git a/Engine/Source/Editor/ContentBrowser/Private/AssetViewWidgets.cpp b/Engine/Source/Editor/ContentBrowser/Private/AssetViewWidgets.cpp
--- a/Engine/Source/Editor/ContentBrowser/Private/AssetViewWidgets.cpp
+++ b/Engine/Source/Editor/ContentBrowser/Private/AssetViewWidgets.cpp
@@ -2022,15 +2022,28 @@ int32 SAssetListItem::GetWidgetSwitcherIndex() const
 	constexpr int32 RealThumbnailIndex = 1;
 	constexpr int32 IconThumbnailIndex = 0;
 
-	if (IsFolder())
+	if (CurrentThumbnailSize.IsSet())
 	{
-		return CurrentThumbnailSize.IsSet() && CurrentThumbnailSize.Get() >= EThumbnailSize::Medium ? RealThumbnailIndex : IconThumbnailIndex;
+		if (IsFolder())
+		{
+			return CurrentThumbnailSize.Get() >= EThumbnailSize::Medium ? RealThumbnailIndex : IconThumbnailIndex;
+		}
+
+		if (CurrentThumbnailSize.Get() == EThumbnailSize::Tiny)
+		{
+			return IconThumbnailIndex;
+		}
 	}
 
-	// Always show thumbnails for assets
+	// By default show the real thumbnail
 	return RealThumbnailIndex;
 }
 
+EVisibility SAssetListItem::GetDirtyBadgeVisibility() const
+{
+	return CurrentThumbnailSize.IsSet() && CurrentThumbnailSize.Get() == EThumbnailSize::Tiny ? EVisibility::Visible : EVisibility::Collapsed;
+}
+
 TSharedRef<SWidget> SAssetListItem::GenerateWidgetForColumn(const FName& ColumnName, FIsSelected InIsSelected)
 {
 	if (UE::Editor::ContentBrowser::IsNewStyleEnabled())
@@ -2204,6 +2217,9 @@ TSharedRef<SWidget> SAssetListItem::GenerateWidgetForColumnNewStyle(const FName&
 			}
 		}
 
+		// Make dirty icon overlay a reasonable size in relation to the icon size (note: it is assumed this icon is square)
+		const float IconOverlaySize = IconBrush->ImageSize.X * 0.6f;
+
 		Content = SNew(SBorder)
 		.BorderImage(this, &SAssetViewItem::GetBorderImage)
 		.Padding(DefaultColumnPadding)
@@ -2227,9 +2243,29 @@ TSharedRef<SWidget> SAssetListItem::GenerateWidgetForColumnNewStyle(const FName&
 					.VAlign(VAlign_Fill)
 					.HAlign(HAlign_Fill)
 					[
-						SNew(SImage)
-						.Image(IconBrush)
-						.ColorAndOpacity(this, &SAssetViewItem::GetAssetColor)
+						SNew(SOverlay)
+
+						+ SOverlay::Slot()
+						[
+							SNew(SImage)
+							.Image(IconBrush)
+							.ColorAndOpacity(this, &SAssetViewItem::GetAssetColor)
+						]
+
+						// Dirty state
+						+SOverlay::Slot()
+						.HAlign(HAlign_Left)
+						.VAlign(VAlign_Bottom)
+						[
+							SNew(SBox)
+							.WidthOverride(IconOverlaySize)
+							.HeightOverride(IconOverlaySize)
+							.Visibility(this, &SAssetListItem::GetDirtyBadgeVisibility)
+							[
+								SNew(SImage)
+								.Image(this, &SAssetListItem::GetDirtyImage)
+							]
+						]
 					]
 
 					+ SWidgetSwitcher::Slot()
diff --git a/Engine/Source/Editor/ContentBrowser/Private/ContentBrowserCommands.cpp b/Engine/Source/Editor/ContentBrowser/Private/ContentBrowserCommands.cpp
--- a/Engine/Source/Editor/ContentBrowser/Private/ContentBrowserCommands.cpp
+++ b/Engine/Source/Editor/ContentBrowser/Private/ContentBrowserCommands.cpp
@@ -29,6 +29,7 @@ void FContentBrowserCommands::RegisterCommands()
 	UI_COMMAND(AssetViewCopyPackageName, "Copy Selected Package Name", "Copy the selected package name", EUserInterfaceActionType::Button, FInputChord(EModifierKey::Control | EModifierKey::Alt, EKeys::C));
 	UI_COMMAND(GridViewShortcut, "Grid", "View assets as tiles in a grid.", EUserInterfaceActionType::RadioButton, FInputChord());
 	UI_COMMAND(ListViewShortcut, "List", "View assets in a list with thumbnails.", EUserInterfaceActionType::RadioButton, FInputChord());
+	UI_COMMAND(ColumnViewShortcut, "Column", "View assets in a list with columns of details.", EUserInterfaceActionType::RadioButton, FInputChord());
 }
 
 
diff --git a/Engine/Source/Editor/ContentBrowser/Private/ContentBrowserCommands.h b/Engine/Source/Editor/ContentBrowser/Private/ContentBrowserCommands.h
--- a/Engine/Source/Editor/ContentBrowser/Private/ContentBrowserCommands.h
+++ b/Engine/Source/Editor/ContentBrowser/Private/ContentBrowserCommands.h
@@ -44,4 +44,5 @@ public:
 	TSharedPtr<FUICommandInfo> AssetViewCopyPackageName;
 	TSharedPtr<FUICommandInfo> GridViewShortcut;
 	TSharedPtr<FUICommandInfo> ListViewShortcut;
+	TSharedPtr<FUICommandInfo> ColumnViewShortcut;
 };
diff --git a/Engine/Source/Editor/ContentBrowser/Private/Menus/AssetViewMenus.cpp b/Engine/Source/Editor/ContentBrowser/Private/Menus/AssetViewMenus.cpp
--- a/Engine/Source/Editor/ContentBrowser/Private/Menus/AssetViewMenus.cpp
+++ b/Engine/Source/Editor/ContentBrowser/Private/Menus/AssetViewMenus.cpp
@@ -114,6 +114,7 @@ FDelayedAutoRegisterHelper SAssetView::AssetViewOptionsMenuRegistration(
 						{
 							ViewTypeSection.AddMenuEntryWithCommandList(FContentBrowserCommands::Get().GridViewShortcut, AssetView->Commands);
 							ViewTypeSection.AddMenuEntryWithCommandList(FContentBrowserCommands::Get().ListViewShortcut, AssetView->Commands);
+							ViewTypeSection.AddMenuEntryWithCommandList(FContentBrowserCommands::Get().ColumnViewShortcut, AssetView->Commands);
 						}
 						else
 						{
@@ -509,32 +510,29 @@ FDelayedAutoRegisterHelper SAssetView::AssetViewOptionsMenuRegistration(
 						}
 					}
 
-					if (!IsNewStyleEnabled())
+					// Asset Columns Section
+					if (AssetView->GetColumnViewVisibility() == EVisibility::Visible)
 					{
-						// Asset Columns Section
-						if (AssetView->GetColumnViewVisibility() == EVisibility::Visible)
 						{
-							{
-								FToolMenuSection& Section = InMenu->FindOrAddSection("AssetColumns", LOCTEXT("ToggleColumnsHeading", "Columns"), FToolMenuInsert("Search", EToolMenuInsertType::After));
+							FToolMenuSection& Section = InMenu->FindOrAddSection("AssetColumns", LOCTEXT("ToggleColumnsHeading", "Columns"), FToolMenuInsert("Search", EToolMenuInsertType::After));
 
-								Section.AddMenuEntry(
-									"ResetColumns",
-									LOCTEXT("ResetColumns", "Reset Columns"),
-									LOCTEXT("ResetColumnsToolTip", "Reset all columns to be visible again."),
-									FSlateIcon(),
-									FUIAction(FExecuteAction::CreateSP(AssetView, &SAssetView::ResetColumns)),
-									EUserInterfaceActionType::Button
-								);
+							Section.AddMenuEntry(
+								"ResetColumns",
+								LOCTEXT("ResetColumns", "Reset Columns"),
+								LOCTEXT("ResetColumnsToolTip", "Reset all columns to be visible again."),
+								FSlateIcon(),
+								FUIAction(FExecuteAction::CreateSP(AssetView, &SAssetView::ResetColumns)),
+								EUserInterfaceActionType::Button
+							);
 
-								Section.AddMenuEntry(
-									"ExportColumns",
-									LOCTEXT("ExportColumns", "Export to CSV"),
-									LOCTEXT("ExportColumnsToolTip", "Export column data to CSV."),
-									FSlateIcon(),
-									FUIAction(FExecuteAction::CreateSP(AssetView, &SAssetView::ExportColumns)),
-									EUserInterfaceActionType::Button
-								);
-							}
+							Section.AddMenuEntry(
+								"ExportColumns",
+								LOCTEXT("ExportColumns", "Export to CSV"),
+								LOCTEXT("ExportColumnsToolTip", "Export column data to CSV."),
+								FSlateIcon(),
+								FUIAction(FExecuteAction::CreateSP(AssetView, &SAssetView::ExportColumns)),
+								EUserInterfaceActionType::Button
+							);
 						}
 					}
 
diff --git a/Engine/Source/Editor/ContentBrowser/Private/SAssetPicker.cpp b/Engine/Source/Editor/ContentBrowser/Private/SAssetPicker.cpp
--- a/Engine/Source/Editor/ContentBrowser/Private/SAssetPicker.cpp
+++ b/Engine/Source/Editor/ContentBrowser/Private/SAssetPicker.cpp
@@ -333,18 +333,6 @@ void SAssetPicker::Construct( const FArguments& InArgs )
 		TextFilter = MakeShared<FAssetTextFilter>();
 	}
 
-	// Avoid MajorAssetColumns creation if the ViewType is set to column and there are no Hidden or Custom columns
-	bool AvoidMajorAssetColumnsCreation = InArgs._AssetPickerConfig.AvoidMajorAssetColumnsCreation;
-	if (UE::Editor::ContentBrowser::IsNewStyleEnabled())
-	{
-		if (InArgs._AssetPickerConfig.InitialAssetViewType != EAssetViewType::Column &&
-			InArgs._AssetPickerConfig.HiddenColumnNames.IsEmpty() &&
-			InArgs._AssetPickerConfig.CustomColumns.IsEmpty())
-		{
-			AvoidMajorAssetColumnsCreation = true;
-		}
-	}
-
 	// clang-format off
 	VerticalBox->AddSlot()
 	.FillHeight(1.f)
@@ -389,7 +377,7 @@ void SAssetPicker::Construct( const FArguments& InArgs )
 		.AssetShowWarningText( InArgs._AssetPickerConfig.AssetShowWarningText)
 		.AllowFocusOnSync(false)	// Stop the asset view from stealing focus (we're in control of that)
 		.HiddenColumnNames(InArgs._AssetPickerConfig.HiddenColumnNames)
-		.AvoidMajorAssetColumnsCreation(AvoidMajorAssetColumnsCreation)
+
 		.CustomColumns(InArgs._AssetPickerConfig.CustomColumns)
 		.OnSearchOptionsChanged(this, &SAssetPicker::HandleSearchSettingsChanged)
 		.InitialThumbnailSize(InArgs._AssetPickerConfig.InitialThumbnailSize)
diff --git a/Engine/Source/Editor/ContentBrowser/Private/SAssetView.cpp b/Engine/Source/Editor/ContentBrowser/Private/SAssetView.cpp
--- a/Engine/Source/Editor/ContentBrowser/Private/SAssetView.cpp
+++ b/Engine/Source/Editor/ContentBrowser/Private/SAssetView.cpp
@@ -1194,7 +1194,6 @@ void SAssetView::Construct( const FArguments& InArgs )
 	bForceShowPluginContent = InArgs._ForceShowPluginContent;
 	bForceHideScrollbar = InArgs._ForceHideScrollbar;
 	bShowDisallowedAssetClassAsUnsupportedItems = InArgs._ShowDisallowedAssetClassAsUnsupportedItems;
-	bAvoidMajorAssetColumnsCreation = InArgs._AvoidMajorAssetColumnsCreation;
 
 	bPendingUpdateThumbnails = false;
 	bShouldNotifyNextAssetSync = true;
@@ -1269,6 +1268,7 @@ void SAssetView::Construct( const FArguments& InArgs )
 	bAllowDragging = InArgs._AllowDragging;
 	bAllowFocusOnSync = InArgs._AllowFocusOnSync;
 	HiddenColumnNames = DefaultHiddenColumnNames = InArgs._HiddenColumnNames;
+	ListHiddenColumnNames = TArray<FString>();
 	CustomColumns = InArgs._CustomColumns;
 	OnSearchOptionsChanged = InArgs._OnSearchOptionsChanged;
 	bShowPathViewFilters = InArgs._bShowPathViewFilters;
@@ -1277,14 +1277,7 @@ void SAssetView::Construct( const FArguments& InArgs )
 	
 	if (InArgs._InitialViewType >= 0 && InArgs._InitialViewType < EAssetViewType::MAX)
 	{
-		if (UE::Editor::ContentBrowser::IsNewStyleEnabled() && InArgs._InitialViewType == EAssetViewType::Column)
-		{
-			CurrentViewType = EAssetViewType::List;
-		}
-		else
-		{
-			CurrentViewType = InArgs._InitialViewType;
-		}
+		CurrentViewType = InArgs._InitialViewType;
 	}
 	else
 	{
@@ -1526,9 +1519,7 @@ void SAssetView::Construct( const FArguments& InArgs )
 	// Generalizing this is a bit difficult because the column ID is not accessible or is not known
 	// Currently I assume this won't work, if this view mode is not column. Otherwise, I don't think sorting by path
 	// is a good idea. 
-	if ((CurrentViewType == EAssetViewType::Column ||
-		(UE::Editor::ContentBrowser::IsNewStyleEnabled() && CurrentViewType == EAssetViewType::List))
-		&& bSortByPathInColumnView)
+	if (CurrentViewType == EAssetViewType::Column && bSortByPathInColumnView)
 	{
 		SortManager->SetSortColumnId(EColumnSortPriority::Primary, SortManager->PathColumnId);
 		SortManager->SetSortColumnId(EColumnSortPriority::Secondary, SortManager->NameColumnId);
@@ -2012,6 +2003,10 @@ void SAssetView::SaveSettings(const FString& IniFilename, const FString& IniSect
 	GConfig->SetFloat(*IniSection, *(SettingsString + TEXT(".ZoomScale")), ZoomScale, IniFilename);
 
 	GConfig->SetArray(*IniSection, *(SettingsString + TEXT(".HiddenColumns")), HiddenColumnNames, IniFilename);
+	if (UE::Editor::ContentBrowser::IsNewStyleEnabled())
+	{
+		GConfig->SetArray(*IniSection, *(SettingsString + TEXT(".ListHiddenColumns")), ListHiddenColumnNames, IniFilename);
+	}
 }
 
 void SAssetView::LoadSettings(const FString& IniFilename, const FString& IniSection, const FString& SettingsString)
@@ -2074,14 +2069,22 @@ void SAssetView::LoadSettings(const FString& IniFilename, const FString& IniSect
 				ColumnView->GetHeaderRow()->SetShowGeneratedColumn(Column.ColumnId, !HiddenColumnNames.Contains(Column.ColumnId.ToString()));
 			}
 		}
+	}
 
-		if (UE::Editor::ContentBrowser::IsNewStyleEnabled())
+	if (UE::Editor::ContentBrowser::IsNewStyleEnabled())
+	{
+		TArray<FString> LoadedListHiddenColumnNames;
+		GConfig->GetArray(*IniSection, *(SettingsString + TEXT(".ListHiddenColumns")), LoadedListHiddenColumnNames, IniFilename);
+		if (LoadedListHiddenColumnNames.Num() > 0)
 		{
+			ListHiddenColumnNames = LoadedListHiddenColumnNames;
+
+			// Also update the visibility of the columns we just loaded in (unless this is called before creation and ColumnView doesn't exist)
 			if (ListView)
 			{
 				for (const SHeaderRow::FColumn& ListColumn : ListView->GetHeaderRow()->GetColumns())
 				{
-					ListView->GetHeaderRow()->SetShowGeneratedColumn(ListColumn.ColumnId, !HiddenColumnNames.Contains(ListColumn.ColumnId.ToString()));
+					ListView->GetHeaderRow()->SetShowGeneratedColumn(ListColumn.ColumnId, !ListHiddenColumnNames.Contains(ListColumn.ColumnId.ToString()));
 				}
 			}
 		}
@@ -2222,8 +2225,7 @@ void SAssetView::Tick( const FGeometry& AllottedGeometry, const double InCurrent
 		else
 		{
 			// Update the columns in the column view now that we know the majority type
-			if (CurrentViewType == EAssetViewType::Column ||
-				(UE::Editor::ContentBrowser::IsNewStyleEnabled() && CurrentViewType == EAssetViewType::List))
+			if (CurrentViewType == EAssetViewType::Column)
 			{
 				int32 HighestCount = 0;
 				FName HighestType;
@@ -2456,8 +2458,7 @@ void SAssetView::ProcessItemsPendingFilter(const double TickStartTime)
 	int32 OldCount = FilteredAssetItems.Num();
 	Items->UpdateItemFiltering(FrontendFilterHelper, bFlushAllPendingItems ? MAX_dbl : TickStartTime + MaxSecondsPerFrame, FilteredAssetItems);
 
-	if (CurrentViewType == EAssetViewType::Column ||
-		(UE::Editor::ContentBrowser::IsNewStyleEnabled() && CurrentViewType == EAssetViewType::List))
+	if (CurrentViewType == EAssetViewType::Column)
 	{
 		for (int32 i = OldCount; i < FilteredAssetItems.Num(); ++i)
 		{
@@ -2688,7 +2689,7 @@ FReply SAssetView::OnMouseWheel( const FGeometry& MyGeometry, const FPointerEven
 			bool IsMoreThanMaxSize = (int32)ThumbnailSize + Delta == (int32)EThumbnailSize::MAX;
 			bool bWasSizeSupposedToChange = false;
 			// If shift is pressed or if we are in the List view the CTRL + Wheel should jump ThumbnailSize by design
-			if (MouseEvent.IsShiftDown() || CurrentViewType == EAssetViewType::List)
+			if (MouseEvent.IsShiftDown() || CurrentViewType != EAssetViewType::Tile)
 			{
 				bWasSizeSupposedToChange = true;
 				EThumbnailSize DesiredThumbnailSize = (EThumbnailSize)FMath::Clamp<int32>((int32)ThumbnailSize + Delta, 0, (int32)EThumbnailSize::MAX - 1);
@@ -2837,14 +2838,16 @@ TSharedRef<SAssetListView> SAssetView::CreateListView()
 
 	if (UE::Editor::ContentBrowser::IsNewStyleEnabled())
 	{
+		const TArray<FString>& HiddenColumnsToUse = CurrentViewType == EAssetViewType::List ? ListHiddenColumnNames : HiddenColumnNames;
+
 		{
-			const bool bIsColumnVisible = !HiddenColumnNames.Contains(SortManager->RevisionControlColumnId.ToString());
+			const bool bIsColumnVisible = !HiddenColumnsToUse.Contains(SortManager->RevisionControlColumnId.ToString());
 			NewListView->GetHeaderRow()->SetShowGeneratedColumn(SortManager->RevisionControlColumnId, bIsColumnVisible);
 		}
 
 		NewListView->GetHeaderRow()->SetOnGetMaxRowSizeForColumn(FOnGetMaxRowSizeForColumn::CreateRaw(NewListView.Get(), &SAssetColumnView::GetMaxRowSizeForColumn));
 
-		if (bShowTypeInColumnView)
+		if (bShowTypeInColumnView || CurrentViewType == EAssetViewType::List)
 		{
 			NewListView->GetHeaderRow()->AddColumn(
 					SHeaderRow::Column(SortManager->ClassColumnId)
@@ -2855,12 +2858,12 @@ TSharedRef<SAssetListView> SAssetView::CreateListView()
 					.DefaultLabel(LOCTEXT("Column_Class", "Type"))
 				);
 
-			const bool bIsColumnVisible = !HiddenColumnNames.Contains(SortManager->ClassColumnId.ToString());
+			const bool bIsColumnVisible = !HiddenColumnsToUse.Contains(SortManager->ClassColumnId.ToString());
 		
 			NewListView->GetHeaderRow()->SetShowGeneratedColumn(SortManager->ClassColumnId, bIsColumnVisible);
 		}
 
-		if (bShowPathInColumnView)
+		if (bShowPathInColumnView && CurrentViewType == EAssetViewType::Column)
 		{
 			NewListView->GetHeaderRow()->AddColumn(
 					SHeaderRow::Column(SortManager->PathColumnId)
@@ -3600,6 +3603,12 @@ void SAssetView::BindCommands()
 			FCanExecuteAction(),
 			FIsActionChecked::CreateSP(this, &SAssetView::IsCurrentViewType, EAssetViewType::List)
 		));
+
+		Commands->MapAction(FContentBrowserCommands::Get().ColumnViewShortcut, FUIAction(
+		FExecuteAction::CreateSP(this, &SAssetView::SetCurrentViewTypeFromMenu, EAssetViewType::Column),
+			FCanExecuteAction(),
+			FIsActionChecked::CreateSP(this, &SAssetView::IsCurrentViewType, EAssetViewType::Column)
+		));
 	}
 
 	FInputBindingManager::Get().RegisterCommandList(FContentBrowserCommands::Get().GetContextName(), Commands.ToSharedRef());
@@ -3786,9 +3795,7 @@ bool SAssetView::IsOrganizingFolders() const
 
 void SAssetView::SetMajorityAssetType(FName NewMajorityAssetType)
 {
-	if ((CurrentViewType != EAssetViewType::Column &&
-		!(UE::Editor::ContentBrowser::IsNewStyleEnabled() && CurrentViewType == EAssetViewType::List)) ||
-		(UE::Editor::ContentBrowser::IsNewStyleEnabled() && bAvoidMajorAssetColumnsCreation))
+	if (CurrentViewType != EAssetViewType::Column)
 	{
 		return;
 	}
@@ -3816,14 +3823,6 @@ void SAssetView::SetMajorityAssetType(FName NewMajorityAssetType)
 		TArray<FName> AddedColumns;
 		TSharedPtr<SListView<TSharedPtr<FAssetViewItem>>> ViewToUse = ColumnView;
 
-		if (UE::Editor::ContentBrowser::IsNewStyleEnabled())
-		{
-			if (CurrentViewType == EAssetViewType::List)
-			{
-				ViewToUse = ListView;
-			}
-		}
-
 		// Since the asset type has changed, remove all columns except name and class
 		const TIndirectArray<SHeaderRow::FColumn>& Columns = ViewToUse->GetHeaderRow()->GetColumns();
 
@@ -4795,14 +4794,11 @@ void SAssetView::SetCurrentViewType(EAssetViewType::Type NewType)
 		{
 			if (UE::Editor::ContentBrowser::IsNewStyleEnabled())
 			{
-				if (ThumbnailSize >= EThumbnailSize::Large)
+				if (ThumbnailSize >= EThumbnailSize::Small)
 				{
 					CurrentThumbnailSize = ListViewThumbnailSize;
 					bPendingUpdateThumbnails = true;
 				}
-				MajorityAssetType = NAME_None;
-				RefreshFilteredItems();
-				SortList();
 			}
 			else
 			{
@@ -4812,6 +4808,14 @@ void SAssetView::SetCurrentViewType(EAssetViewType::Type NewType)
 		}
 		else if ( NewType == EAssetViewType::Column )
 		{
+			if (UE::Editor::ContentBrowser::IsNewStyleEnabled())
+			{
+				if (ThumbnailSize >= EThumbnailSize::Small)
+				{
+					CurrentThumbnailSize = ListViewThumbnailSize;
+					bPendingUpdateThumbnails = true;
+				}
+			}
 			// No thumbnails, but we do need to refresh filtered items to determine a majority asset type
 			MajorityAssetType = NAME_None;
 			RefreshFilteredItems();
@@ -4860,7 +4864,14 @@ void SAssetView::CreateCurrentView()
 			NewView = CreateShadowOverlay(ListView.ToSharedRef());
 			break;
 		case EAssetViewType::Column:
-			ColumnView = CreateColumnView();
+			if (UE::Editor::ContentBrowser::IsNewStyleEnabled())
+			{
+				ColumnView = CreateListView();
+			}
+			else
+			{
+				ColumnView = CreateColumnView();
+			}
 			NewView = CreateShadowOverlay(ColumnView.ToSharedRef());
 			break;
 		case EAssetViewType::Custom:
@@ -5645,7 +5656,7 @@ bool SAssetView::IsThumbnailEditMode() const
 
 bool SAssetView::IsThumbnailEditModeAllowed() const
 {
-	return bAllowThumbnailEditMode && GetCurrentViewType() != EAssetViewType::Column;
+	return bAllowThumbnailEditMode && (UE::Editor::ContentBrowser::IsNewStyleEnabled() || GetCurrentViewType() != EAssetViewType::Column);
 }
 
 FReply SAssetView::EndThumbnailEditModeClicked()
@@ -5824,7 +5835,7 @@ void SAssetView::UpdateThumbnailSizeValue()
 
 bool SAssetView::IsThumbnailScalingAllowed() const
 {
-	return GetCurrentViewType() != EAssetViewType::Column && GetCurrentViewType() != EAssetViewType::Custom;
+	return (UE::Editor::ContentBrowser::IsNewStyleEnabled() || GetCurrentViewType() != EAssetViewType::Column) && GetCurrentViewType() != EAssetViewType::Custom;
 }
 
 float SAssetView::GetTileViewTypeNameHeight() const
@@ -5882,7 +5893,7 @@ float SAssetView::GetListViewItemHeight() const
 FMargin SAssetView::GetListViewItemPadding() const
 {
 	return UE::Editor::ContentBrowser::IsNewStyleEnabled()
-		? FMargin(0.f, (float)ListViewThumbnailPadding)
+		? ThumbnailSize == EThumbnailSize::Tiny ? FMargin(0.f) : FMargin(0.f, (float)ListViewThumbnailPadding)
 		: FMargin(0.f);
 }
 
@@ -6228,63 +6239,53 @@ void SAssetView::ExportColumns()
 
 void SAssetView::OnHiddenColumnsChanged()
 {
+	// Early out if this is called before creation (due to loading config etc)
+	TSharedPtr<SListView<TSharedPtr<FAssetViewItem>>> ViewToUse = ColumnView;
+	TArray<FString>* HiddenColumnsToUse = &HiddenColumnNames;
 	if (UE::Editor::ContentBrowser::IsNewStyleEnabled())
 	{
-		// Early out if this is called before creation (due to loading config etc)
-		if (!ListView)
+		if (CurrentViewType == EAssetViewType::List)
 		{
-			return;
+			ViewToUse = ListView;
+			HiddenColumnsToUse = &ListHiddenColumnNames;
 		}
+	}
 
-		// We can't directly update the hidden columns list, because some columns maybe hidden, but not created yet
-		TArray<FName> NewHiddenColumns = ListView->GetHeaderRow()->GetHiddenColumnIds();
+	if (!ViewToUse.IsValid())
+	{
+		return;
+	}
 
-		// So instead for each column that currently exists, we update its visibility state in the HiddenColumnNames array
-		for (const SHeaderRow::FColumn& ListColumn : ListView->GetHeaderRow()->GetColumns())
-		{
-			const bool bIsColumnVisible = NewHiddenColumns.Contains(ListColumn.ColumnId);
+	// We can't directly update the hidden columns list, because some columns maybe hidden, but not created yet
+	TArray<FName> NewHiddenColumns = ViewToUse->GetHeaderRow()->GetHiddenColumnIds();
 
-			if(bIsColumnVisible)
-			{
-				HiddenColumnNames.AddUnique(ListColumn.ColumnId.ToString());
-			}
-			else
-			{
-				HiddenColumnNames.Remove(ListColumn.ColumnId.ToString());
-			}
-		}
-	}
-	else
+	// So instead for each column that currently exists, we update its visibility state in the HiddenColumnNames array
+	for (const SHeaderRow::FColumn& Column : ViewToUse->GetHeaderRow()->GetColumns())
 	{
-		// Early out if this is called before creation (due to loading config etc)
-		if (!ColumnView)
+		const bool bIsColumnVisible = NewHiddenColumns.Contains(Column.ColumnId);
+
+		if (bIsColumnVisible)
 		{
-			return;
+			HiddenColumnsToUse->AddUnique(Column.ColumnId.ToString());
 		}
-
-		// We can't directly update the hidden columns list, because some columns maybe hidden, but not created yet
-		TArray<FName> NewHiddenColumns = ColumnView->GetHeaderRow()->GetHiddenColumnIds();
-
-		// So instead for each column that currently exists, we update its visibility state in the HiddenColumnNames array
-		for (const SHeaderRow::FColumn& Column : ColumnView->GetHeaderRow()->GetColumns())
+		else
 		{
-			const bool bIsColumnVisible = NewHiddenColumns.Contains(Column.ColumnId);
-
-			if(bIsColumnVisible)
-			{
-				HiddenColumnNames.AddUnique(Column.ColumnId.ToString());
-			}
-			else
-			{
-				HiddenColumnNames.Remove(Column.ColumnId.ToString());
-			}
+			HiddenColumnsToUse->Remove(Column.ColumnId.ToString());
 		}
 	}
 
 	if (FAssetViewInstanceConfig* Config = GetAssetViewConfig())
 	{
-		Config->HiddenColumns.Reset();
-		Algo::Transform(HiddenColumnNames, Config->HiddenColumns, [](const FString& Str) { return FName(*Str); });
+		if (CurrentViewType == EAssetViewType::List)
+		{
+			Config->ListHiddenColumns.Reset();
+			Algo::Transform(*HiddenColumnsToUse, Config->ListHiddenColumns, [](const FString& Str) { return FName(*Str); });
+		}
+		else
+		{
+			Config->HiddenColumns.Reset();
+			Algo::Transform(*HiddenColumnsToUse, Config->HiddenColumns, [](const FString& Str) { return FName(*Str); });
+		}
 
 		UAssetViewConfig::Get()->SaveEditorConfig();
 	}
diff --git a/Engine/Source/Editor/ContentBrowser/Public/AssetViewWidgets.h b/Engine/Source/Editor/ContentBrowser/Public/AssetViewWidgets.h
--- a/Engine/Source/Editor/ContentBrowser/Public/AssetViewWidgets.h
+++ b/Engine/Source/Editor/ContentBrowser/Public/AssetViewWidgets.h
@@ -513,6 +513,9 @@ private:
 	/** Get the index to use for the Thumbnail WidgetSwitcher, based on the ThumbnailSize */
 	int32 GetWidgetSwitcherIndex() const;
 
+	/** Get the visibility of the dirty badge in the list view based on the current size */
+	EVisibility GetDirtyBadgeVisibility() const;
+
 private:
 	/** The handle to the thumbnail that this item is rendering */
 	TSharedPtr<FAssetThumbnail> AssetThumbnail;
diff --git a/Engine/Source/Editor/ContentBrowser/Public/IContentBrowserSingleton.h b/Engine/Source/Editor/ContentBrowser/Public/IContentBrowserSingleton.h
--- a/Engine/Source/Editor/ContentBrowser/Public/IContentBrowserSingleton.h
+++ b/Engine/Source/Editor/ContentBrowser/Public/IContentBrowserSingleton.h
@@ -329,9 +329,6 @@ struct FAssetPickerConfig
 	/** If true, sort by path in column view. Only works if initial view type is Column */
 	bool bSortByPathInColumnView;
 
-	/** If true, will avoid adding the major assets column creation in the ListView */
-	bool AvoidMajorAssetColumnsCreation;
-
 	/** Can be used to freely change the Add Filter menu. This is executed on the dynamically instanced menu, not the registered menu. */
 	TDelegate<void(UToolMenu*)> OnExtendAddFilterMenu;
 	
@@ -367,7 +364,6 @@ struct FAssetPickerConfig
 		, bShowPathInColumnView(false)
 		, bShowTypeInColumnView(true)
 		, bSortByPathInColumnView(false)
-		, AvoidMajorAssetColumnsCreation(false)
 		, DefaultFilterMenuExpansion(EAssetTypeCategories::Basic)
 		, bUseSectionsForCustomFilterCategories(false)
 	{}
diff --git a/Engine/Source/Editor/ContentBrowser/Public/SAssetView.h b/Engine/Source/Editor/ContentBrowser/Public/SAssetView.h
--- a/Engine/Source/Editor/ContentBrowser/Public/SAssetView.h
+++ b/Engine/Source/Editor/ContentBrowser/Public/SAssetView.h
@@ -298,9 +298,6 @@ public:
 		/** Allow the asset view to display a custom view registered with the Content Browser module */
 		SLATE_ARGUMENT(bool, AllowCustomView)
 
-		/** If true avoid the creation of all the Major Asset columns */
-		SLATE_ARGUMENT(bool, AvoidMajorAssetColumnsCreation)
-
 		/** Called to check if an asset tag should be display in details view. */
 		SLATE_EVENT( FOnShouldDisplayAssetTag, OnAssetTagWantsToBeDisplayed )
 
@@ -1011,7 +1008,7 @@ private:
 	EAssetViewType::Type CurrentViewType;
 	TSharedPtr<SAssetListView> ListView;
 	TSharedPtr<SAssetTileView> TileView;
-	TSharedPtr<SAssetColumnView> ColumnView;
+	TSharedPtr<SListView<TSharedPtr<FAssetViewItem>>> ColumnView;
 	TSharedPtr<SBox> ViewContainer;
 
 	/** The content browser that created this asset view if any */
@@ -1250,9 +1247,6 @@ private:
 	/** Whether or not to notify about newly selected items on on the next asset sync */
 	bool bShouldNotifyNextAssetSync : 1;
 
-	/** Whether or not the major asset columns should be skipped in the ListView */
-	bool bAvoidMajorAssetColumnsCreation : 1;
-
 	/** The current selection mode used by the asset view */
 	ESelectionMode::Type SelectionMode;
 
@@ -1338,6 +1332,7 @@ private:
 	/** Column filtering state */
 	TArray<FString> DefaultHiddenColumnNames;
 	TArray<FString> HiddenColumnNames;
+	TArray<FString> ListHiddenColumnNames;
 
 	TArray<FAssetViewCustomColumn> CustomColumns;
 
diff --git a/Engine/Source/Editor/PropertyEditor/Private/UserInterface/PropertyEditor/SPropertyMenuAssetPicker.cpp b/Engine/Source/Editor/PropertyEditor/Private/UserInterface/PropertyEditor/SPropertyMenuAssetPicker.cpp
--- a/Engine/Source/Editor/PropertyEditor/Private/UserInterface/PropertyEditor/SPropertyMenuAssetPicker.cpp
+++ b/Engine/Source/Editor/PropertyEditor/Private/UserInterface/PropertyEditor/SPropertyMenuAssetPicker.cpp
@@ -143,10 +143,8 @@ void SPropertyMenuAssetPicker::Construct( const FArguments& InArgs )
 		// TODO: Remove "PropertyEditorHelpers.h" include when IsNewStyleEnabled will be removed
 		if (UE::Editor::ContentBrowser::IsNewStyleEnabled())
 		{
-			// Use the large size by default if the new style is enabled
-			AssetPickerConfig.InitialThumbnailSize = EThumbnailSize::Large;
-			// Avoid creating the MajorAsset specific columns
-			AssetPickerConfig.AvoidMajorAssetColumnsCreation = true;
+			// Use the medium size by default if the new style is enabled
+			AssetPickerConfig.InitialThumbnailSize = EThumbnailSize::Medium;
 		}
 		// The initial selection should be the current value
 		AssetPickerConfig.InitialAssetSelection = CurrentObject;
diff --git a/Engine/Source/Editor/UnrealEd/Private/AssetDefinitionDefault.cpp b/Engine/Source/Editor/UnrealEd/Private/AssetDefinitionDefault.cpp
--- a/Engine/Source/Editor/UnrealEd/Private/AssetDefinitionDefault.cpp
+++ b/Engine/Source/Editor/UnrealEd/Private/AssetDefinitionDefault.cpp
@@ -54,6 +54,29 @@ namespace UE::AssetDefinitionDefault::Status
 		return SourceControlBrush;
 	}
 
+	const FSlateBrush* GetSourceControlStatusOverlayBrush(const TSharedPtr<IAssetStatusInfoProvider> InAssetStatusInfoProvider)
+	{
+		const FSlateBrush* SourceControlStatusOverlay = FAppStyle::GetNoBrush();
+		if (!InAssetStatusInfoProvider.IsValid())
+		{
+			return SourceControlStatusOverlay;
+		}
+
+		if (ISourceControlModule::Get().IsEnabled() && ISourceControlModule::Get().GetProvider().IsAvailable())
+		{
+			const FString FileName = InAssetStatusInfoProvider->TryGetFilename();
+			if (FSourceControlStatePtr SourceControlState = ISourceControlModule::Get().GetProvider().GetState(FileName, EStateCacheUsage::Use))
+			{
+				FSlateIcon SCCIcon = SourceControlState->GetIcon();
+				if (SCCIcon.IsSet())
+				{
+					SourceControlStatusOverlay = SCCIcon.GetOverlayIcon();
+				}
+			}
+		}
+		return SourceControlStatusOverlay;
+	}
+
 	EVisibility GetSourceControlStatusVisibility(const TSharedPtr<IAssetStatusInfoProvider> InAssetStatusInfoProvider)
 	{
 		EVisibility SourceControlStatusVisibility = EVisibility::Collapsed;
@@ -133,7 +156,7 @@ void UAssetDefinitionDefault::GetAssetStatusInfo(const TSharedPtr<IAssetStatusIn
 {
 	FAssetDisplayInfo DirtyStatus;
 	DirtyStatus.StatusIcon = FAppStyle::GetBrush("ContentBrowser.ContentDirty");
-	DirtyStatus.Priority = FAssetStatusPriority(EStatusSeverity::Info, 1);
+	DirtyStatus.Priority = FAssetStatusPriority(EStatusSeverity::Info, 5);
 	DirtyStatus.StatusDescription = LOCTEXT("DirtyAssetTooltip", "Asset has unsaved changes");
 	DirtyStatus.IsVisible = TAttribute<EVisibility>::Create(TAttribute<EVisibility>::FGetter::CreateStatic(&UE::AssetDefinitionDefault::Status::GetDirtyStatusVisibility, InAssetStatusInfoProvider));
 	OutStatusInfo.Add(DirtyStatus);
@@ -141,6 +164,7 @@ void UAssetDefinitionDefault::GetAssetStatusInfo(const TSharedPtr<IAssetStatusIn
 	FAssetDisplayInfo SCCStatus;
 	SCCStatus.Priority = FAssetStatusPriority(EStatusSeverity::Info, 0);
 	SCCStatus.StatusIcon = TAttribute<const FSlateBrush*>::Create(TAttribute<const FSlateBrush*>::FGetter::CreateStatic(&UE::AssetDefinitionDefault::Status::GetSourceControlStatusBrush, InAssetStatusInfoProvider));
+	SCCStatus.StatusIconOverlay = TAttribute<const FSlateBrush*>::Create(TAttribute<const FSlateBrush*>::FGetter::CreateStatic(&UE::AssetDefinitionDefault::Status::GetSourceControlStatusOverlayBrush, InAssetStatusInfoProvider));
 	SCCStatus.IsVisible = TAttribute<EVisibility>::Create(TAttribute<EVisibility>::FGetter::CreateStatic(&UE::AssetDefinitionDefault::Status::GetSourceControlStatusVisibility, InAssetStatusInfoProvider));
 	SCCStatus.StatusDescription = TAttribute<FText>::Create(TAttribute<FText>::FGetter::CreateStatic(&UE::AssetDefinitionDefault::Status::GetSourceControlStatusDescription, InAssetStatusInfoProvider));
 	OutStatusInfo.Add(SCCStatus);
diff --git a/Engine/Source/Editor/UnrealEd/Private/AssetThumbnail.cpp b/Engine/Source/Editor/UnrealEd/Private/AssetThumbnail.cpp
--- a/Engine/Source/Editor/UnrealEd/Private/AssetThumbnail.cpp
+++ b/Engine/Source/Editor/UnrealEd/Private/AssetThumbnail.cpp
@@ -277,20 +277,32 @@ public:
 			];
 		}
 
-		if (ThumbnailClass.Get() && bIsClassType && !UE::Editor::ContentBrowser::IsNewStyleEnabled())
+		if (ThumbnailClass.Get() && bIsClassType)
 		{
-			OverlayWidget->AddSlot()
-			.VAlign(VAlign_Bottom)
-			.HAlign(HAlign_Right)
-			.Padding(GetClassIconPadding())
-			[
-				SAssignNew(ClassIconWidget, SBorder)
-				.BorderImage(FAppStyle::GetNoBrush())
+			if (UE::Editor::ContentBrowser::IsNewStyleEnabled())
+			{
+				FAssetDisplayInfo ClassInfo;
+				ClassInfo.StatusIcon = MakeAttributeSP(this, &SAssetThumbnail::GetClassIconBrush);
+				ClassInfo.Priority = FAssetStatusPriority(EStatusSeverity::Info, 1);
+				ClassInfo.StatusDescription = MakeAttributeSP(this, &SAssetThumbnail::GetClassName);
+				ClassInfo.IsVisible = EVisibility::Visible;
+				OverlayInfo.Add(ClassInfo);
+			}
+			else
+			{
+				OverlayWidget->AddSlot()
+				.VAlign(VAlign_Bottom)
+				.HAlign(HAlign_Right)
+				.Padding(GetClassIconPadding())
 				[
-					SNew(SImage)
-					.Image(this, &SAssetThumbnail::GetClassIconBrush)
-				]
-			];
+					SAssignNew(ClassIconWidget, SBorder)
+					.BorderImage(FAppStyle::GetNoBrush())
+					[
+						SNew(SImage)
+						.Image(this, &SAssetThumbnail::GetClassIconBrush)
+					]
+				];
+			}
 		}
 
 		if( bAllowHintText )
@@ -591,10 +603,16 @@ public:
 	{
 		const FAssetData& AssetData = AssetThumbnail->GetAssetData();
 		ThumbnailClass = MakeWeakObjectPtr(const_cast<UClass*>(FClassIconFinder::GetIconClassForAssetData(AssetData, &bIsClassType)));
+		if (ThumbnailClass.IsValid())
+		{
+			ClassName = FText::FromString(ThumbnailClass->GetName());
+		}
+
 		const FName AssetClassName = AssetThumbnail->GetAssetData().AssetClassPath.GetAssetName();
 
 		ClassIconBrush = nullptr;
 		ThumbnailBrush = nullptr;
+
 		if (AssetTypeActions)
 		{
 			if (const FSlateBrush* AssetTypeThumbnail = AssetTypeActions->GetThumbnailBrush(AssetData, AssetClassName))
@@ -690,20 +708,11 @@ public:
 		{
 			const FSlateBrush* ClassIcon = FAppStyle::GetDefaultBrush();
 			TOptional<FLinearColor> Color;
-			TArray<FAssetDisplayInfo> OutStatusInfo;
 			if (const UAssetDefinition* AssetDefinition = UAssetDefinitionRegistry::Get()->GetAssetDefinitionForClass(AssetData.GetClass()))
 			{
-				TSharedRef<FAssetStatusAssetDataInfoProvider> Provider = MakeShared<FAssetStatusAssetDataInfoProvider>(AssetData);
-				AssetDefinition->GetAssetStatusInfo(Provider, OutStatusInfo);
 				ClassIcon = AssetDefinition->GetIconBrush(AssetData, AssetData.AssetClassPath.GetAssetName());
 				Color = AssetDefinition->GetAssetColor();
 			}
-			// Get the default asset definition to retrieve some default information on it (ex: SourceControl)
-			else if (const UAssetDefinition* AssetDefinitionDefaultUObject = GetDefault<UAssetDefinitionDefault>())
-			{
-				TSharedRef<FAssetStatusAssetDataInfoProvider> Provider = MakeShared<FAssetStatusAssetDataInfoProvider>(AssetData);
-				AssetDefinitionDefaultUObject->GetAssetStatusInfo(Provider, OutStatusInfo);
-			}
 
 			if (AssetSystemInfoProvider.IsValid())
 			{
@@ -742,6 +751,7 @@ public:
 				]
 
 				+ SVerticalBox::Slot()
+				.Padding(0.f, 0.f, 0.f, 6.f)
 				.AutoHeight()
 				[
 					SNew(SHorizontalBox)
@@ -770,22 +780,13 @@ public:
 				]
 			];
 
-			// Separator
-			OverallTooltipVBox->AddSlot()
-			.Padding(0.f,6.f)
-			.AutoHeight()
-			[
-				SNew(SSeparator)
-				.Orientation(Orient_Horizontal)
-				.Thickness(1.f)
-				.ColorAndOpacity(COLOR("#484848FF"))
-				.SeparatorImage(FAppStyle::Get().GetBrush("WhiteBrush"))
-			];
-
 			TSharedRef<SVerticalBox> StatusVerticalBox = SNew(SVerticalBox);
 
-			for (const FAssetDisplayInfo& AssetStatusInfo : OutStatusInfo)
+			for (const FAssetDisplayInfo& AssetStatusInfo : OverlayInfo)
 			{
+				TSharedRef<SLayeredImage> StatusLayeredImage = SNew(SLayeredImage).Image(AssetStatusInfo.StatusIcon);
+				StatusLayeredImage->AddLayer(AssetStatusInfo.StatusIconOverlay);
+
 				StatusVerticalBox->AddSlot()
 					.Padding(0.f, 0.f, 0.f, 6.f)
 					.AutoHeight()
@@ -801,8 +802,7 @@ public:
 						   .WidthOverride(16.f)
 						   .HeightOverride(16.f)
 						   [
-							   SNew(SImage)
-							   .Image(AssetStatusInfo.StatusIcon)
+								StatusLayeredImage
 						   ]
 						]
 
@@ -823,6 +823,18 @@ public:
 				StatusVerticalBox
 			];
 
+			// Separator
+			OverallTooltipVBox->AddSlot()
+			.Padding(0.f,0.f, 0.f, 6.f)
+			.AutoHeight()
+			[
+				SNew(SSeparator)
+				.Orientation(Orient_Horizontal)
+				.Thickness(1.f)
+				.ColorAndOpacity(COLOR("#484848FF"))
+				.SeparatorImage(FAppStyle::Get().GetBrush("WhiteBrush"))
+			];
+
 			// More info
 			if (!OutSystemInfo.IsEmpty())
 			{
@@ -844,6 +856,7 @@ public:
 				ExtendedToolTipVerticalBox
 			];
 
+		bool bWasSeparatorAddedForCollection = false;
 		for (const FAssetDisplayInfo& SystemInfo : OutSystemInfo)
 		{
 			if (!SystemInfo.StatusTitle.IsSet() || !SystemInfo.StatusDescription.IsSet())
@@ -851,6 +864,26 @@ public:
 				continue;
 			}
 
+			// StatusTitle currently used to add a separator for Collection, need to be changed in future version to allow more configurability
+			const FName TitleName = FName(SystemInfo.StatusTitle.Get().ToString());
+			if (!bWasSeparatorAddedForCollection && TitleName == TEXT("Collection(s)"))
+			{
+				bWasSeparatorAddedForCollection = true;
+			
+				// Separator
+				ExtendedToolTipVerticalBox->AddSlot()
+				.Padding(0.f,0.f, 0.f, 6.f)
+				.AutoHeight()
+				[
+					SNew(SSeparator)
+					.Orientation(Orient_Horizontal)
+					.Thickness(1.f)
+					.ColorAndOpacity(COLOR("#484848FF"))
+					.SeparatorImage(FAppStyle::Get().GetBrush("WhiteBrush"))
+				];
+				continue;
+			}
+			
 			if ((SystemInfo.IsVisible.IsSet() && SystemInfo.IsVisible.Get().IsVisible()) || !SystemInfo.IsVisible.IsSet())
 			{
 				AddToExtendedToolTipInfoBox(ExtendedToolTipVerticalBox, SystemInfo.StatusIcon, SystemInfo.StatusTitle.Get(), SystemInfo.StatusDescription.Get());
@@ -1012,14 +1045,16 @@ private:
 
 	TSharedRef<SWidget> CreateStatusWidget(int32 StatusIndex, const FAssetDisplayInfo& InStatusInfo) const
 	{
+		TSharedRef<SLayeredImage> StatusLayeredImage = SNew(SLayeredImage).Image(InStatusInfo.StatusIcon);
+		StatusLayeredImage->AddLayer(InStatusInfo.StatusIconOverlay);
+
 		return SNew(SBox)
 			.WidthOverride(StatusSize)
 			.HeightOverride(StatusSize)
 			.Visibility(this, &SAssetThumbnail::GetStatusVisibility, StatusIndex)
 			.VAlign(VAlign_Center)
 			[
-				SNew(SImage)
-				.Image(InStatusInfo.StatusIcon)
+				StatusLayeredImage
 			];
 	}
 
@@ -1266,6 +1301,11 @@ private:
 		return (bHasRenderedThumbnail && ViewportFadeAnimation.IsAtEnd()) ? EVisibility::Collapsed : EVisibility::Visible;
 	}
 
+	FText GetClassName() const
+	{
+		return ClassName;
+	}
+
 	const FSlateBrush* GetClassIconBrush() const
 	{
 		return ClassIconBrush;
@@ -1561,6 +1601,7 @@ private:
 	const FSlateBrush* ThumbnailBrush;
 
 	const FSlateBrush* ClassIconBrush;
+	FText ClassName;
 
 	/** The class to use when finding the thumbnail. */
 	TWeakObjectPtr<UClass> ThumbnailClass;
