[ContentBrowser] Fixed navigation breadcrumb path text alignment

diff --git a/Engine/Source/Editor/ContentBrowser/Private/SContentBrowser.cpp b/Engine/Source/Editor/ContentBrowser/Private/SContentBrowser.cpp
--- a/Engine/Source/Editor/ContentBrowser/Private/SContentBrowser.cpp
+++ b/Engine/Source/Editor/ContentBrowser/Private/SContentBrowser.cpp
@@ -991,6 +991,7 @@ TSharedRef<SWidget> SContentBrowser::CreateAssetView(const FContentBrowserConfig
 	if (UE::Editor::ContentBrowser::IsNewStyleEnabled())
 	{
 		LegacyContentSourceWidgets->NavigationBar = SNew(SNavigationBar)
+						.BreadcrumbButtonContentPadding(FMargin(2, 0, 2, -1)) // This ensures proper vertical alignment of the text to fit the 24px height of the toolbar
 						.OnPathClicked(this, &SContentBrowser::OnPathClicked)
 						.GetPathMenuContent(this, &SContentBrowser::OnGetCrumbDelimiterContent)
 						.GetComboOptions(this, &SContentBrowser::GetRecentPaths)
