[AssetPicker] Change default to be Large and keep DetailsPanel Picker default size to medium

diff --git a/Engine/Source/Editor/ContentBrowser/Public/IContentBrowserSingleton.h b/Engine/Source/Editor/ContentBrowser/Public/IContentBrowserSingleton.h
--- a/Engine/Source/Editor/ContentBrowser/Public/IContentBrowserSingleton.h
+++ b/Engine/Source/Editor/ContentBrowser/Public/IContentBrowserSingleton.h
@@ -366,7 +366,13 @@ struct FAssetPickerConfig
 		, bSortByPathInColumnView(false)
 		, DefaultFilterMenuExpansion(EAssetTypeCategories::Basic)
 		, bUseSectionsForCustomFilterCategories(false)
-	{}
+	{
+		if (IsNewStyleEnabled())
+		{
+			// Default size in the new style AssetPicker will be Large
+			InitialThumbnailSize = EThumbnailSize::Large;
+		}
+	}
 
 	PRAGMA_DISABLE_DEPRECATION_WARNINGS
 
@@ -377,6 +383,22 @@ struct FAssetPickerConfig
 	FAssetPickerConfig& operator=(FAssetPickerConfig&&) = default;
 
 	PRAGMA_ENABLE_DEPRECATION_WARNINGS
+
+private:
+	static bool IsNewStyleEnabled()
+	{
+		static bool bIsNewStyleEnabled = [&]()
+		{
+			if (const IConsoleVariable* CVar = IConsoleManager::Get().FindConsoleVariable(TEXT("ContentBrowser.EnableNewStyle")))
+			{
+				ensureAlwaysMsgf(!EnumHasAnyFlags(CVar->GetFlags(), ECVF_Default), TEXT("The CVar should have already been set from commandline, @see: UnrealEdGlobals.cpp, UE::Editor::ContentBrowser::EnableContentBrowserNewStyleCVarRegistration."));
+				return CVar->GetBool();
+			}
+			return false;
+		}();
+
+		return bIsNewStyleEnabled;
+	}
 };
 
 /** A struct containing details about how the path picker should behave */
