[ContentBrowser] Path bar scaling/clipping fix

diff --git a/Engine/Source/Editor/ContentBrowser/Private/Menus/ContentBrowserMenus.cpp b/Engine/Source/Editor/ContentBrowser/Private/Menus/ContentBrowserMenus.cpp
--- a/Engine/Source/Editor/ContentBrowser/Private/Menus/ContentBrowserMenus.cpp
+++ b/Engine/Source/Editor/ContentBrowser/Private/Menus/ContentBrowserMenus.cpp
@@ -299,8 +299,6 @@ FDelayedAutoRegisterHelper SContentBrowser::ToolBarMenuRegistration(
 						return;
 					}
 
-					constexpr float NavigationPathMinWidth = 400.0f;
-
 					if (const UContentBrowserToolbarMenuContext* Context = InSection.FindContext<UContentBrowserToolbarMenuContext>())
 					{
 						if (Context->ContentBrowser.IsValid())
@@ -323,7 +321,6 @@ FDelayedAutoRegisterHelper SContentBrowser::ToolBarMenuRegistration(
 							PathBoxEntry.WidgetData.StyleParams.SizeRule = FSizeParam::SizeRule_StretchContent;
 							PathBoxEntry.WidgetData.StyleParams.FillSize = 1.0f;
 							PathBoxEntry.WidgetData.StyleParams.DesiredHeightOverride = 24.0f; // Match the height of the other widgets
-							PathBoxEntry.WidgetData.StyleParams.MinimumSize = NavigationPathMinWidth;
 						}
 					}
 				}));
