[ContentBrowser] Add Content context menu - replace icon chip with line

diff --git a/Engine/Source/Editor/ContentBrowser/Private/Widgets/SAssetMenuIcon.cpp b/Engine/Source/Editor/ContentBrowser/Private/Widgets/SAssetMenuIcon.cpp
--- a/Engine/Source/Editor/ContentBrowser/Private/Widgets/SAssetMenuIcon.cpp
+++ b/Engine/Source/Editor/ContentBrowser/Private/Widgets/SAssetMenuIcon.cpp
@@ -7,6 +7,7 @@
 #include "ContentBrowserStyle.h"
 #include "Misc/CommandLine.h"
 #include "SlateOptMacros.h"
+#include "Widgets/Colors/SColorBlock.h"
 
 void SAssetMenuIcon::Construct(const FArguments& InArgs, const UClass* InAssetClass, const FName InIconOverride)
 {
@@ -23,7 +24,7 @@ void SAssetMenuIcon::Construct(const FArguments& InArgs, const UClass* InAssetCl
 
 	if (UE::Editor::ContentBrowser::IsNewStyleEnabled())
 	{
-		static const FVector2D ChipSize = FVector2D(6, 6);
+		constexpr float AssetLineHeight = 2.0f; // Height of the colored line at the base of the icon
 
 		const FName ClassIconBrushOverride = InIconOverride;
 		const FSlateBrush* ClassIcon = nullptr;
@@ -64,14 +65,13 @@ void SAssetMenuIcon::Construct(const FArguments& InArgs, const UClass* InAssetCl
 				]
 
 				+ SOverlay::Slot()
-				.HAlign(HAlign_Right)
-				.VAlign(VAlign_Top)
+				.HAlign(HAlign_Fill)
+				.VAlign(VAlign_Bottom)
 				.Padding(0)
 				[
-					SNew(SImage)
-					.DesiredSizeOverride(ChipSize)
-					.Image(FAppStyle::GetBrush("AssetThumbnail.AssetTileViewWhiteChipSharp"))
-					.ColorAndOpacity(AssetColor)
+					SNew(SColorBlock)
+					.Size(FVector2D(1.0f, AssetLineHeight)) // The X size is ignored/overridden by HAlign_Fill
+					.Color(AssetColor)
 				]
 			]
 		];
