[ContentBrowser] 2nd batch of fixes based on feedbacks
Added ThumbnailSize x ViewType
Fixed ColumnView being unable to be the default view when LoadingSettings
Added HiddenColumns for ListView
Added HiddenColumn logic to take the default in consideration until changed by the user or reset to default

diff --git a/Engine/Source/Editor/ContentBrowser/Private/Menus/AssetViewMenus.cpp b/Engine/Source/Editor/ContentBrowser/Private/Menus/AssetViewMenus.cpp
--- a/Engine/Source/Editor/ContentBrowser/Private/Menus/AssetViewMenus.cpp
+++ b/Engine/Source/Editor/ContentBrowser/Private/Menus/AssetViewMenus.cpp
@@ -510,8 +510,8 @@ FDelayedAutoRegisterHelper SAssetView::AssetViewOptionsMenuRegistration(
 						}
 					}
 
-					// Asset Columns Section
-					if (AssetView->GetColumnViewVisibility() == EVisibility::Visible)
+					// Asset Columns/List Section
+					if (AssetView->GetColumnViewVisibility() == EVisibility::Visible || (IsNewStyleEnabled() && AssetView->GetListViewVisibility() == EVisibility::Visible))
 					{
 						{
 							FToolMenuSection& Section = InMenu->FindOrAddSection("AssetColumns", LOCTEXT("ToggleColumnsHeading", "Columns"), FToolMenuInsert("Search", EToolMenuInsertType::After));
@@ -524,6 +524,14 @@ FDelayedAutoRegisterHelper SAssetView::AssetViewOptionsMenuRegistration(
 								FUIAction(FExecuteAction::CreateSP(AssetView, &SAssetView::ResetColumns)),
 								EUserInterfaceActionType::Button
 							);
+						}
+					}
+
+					// Asset Columns Section
+					if (AssetView->GetColumnViewVisibility() == EVisibility::Visible)
+					{
+						{
+							FToolMenuSection& Section = InMenu->FindOrAddSection("AssetColumns", LOCTEXT("ToggleColumnsHeading", "Columns"), FToolMenuInsert("Search", EToolMenuInsertType::After));
 
 							Section.AddMenuEntry(
 								"ExportColumns",
diff --git a/Engine/Source/Editor/ContentBrowser/Private/SAssetPicker.cpp b/Engine/Source/Editor/ContentBrowser/Private/SAssetPicker.cpp
--- a/Engine/Source/Editor/ContentBrowser/Private/SAssetPicker.cpp
+++ b/Engine/Source/Editor/ContentBrowser/Private/SAssetPicker.cpp
@@ -377,7 +377,7 @@ void SAssetPicker::Construct( const FArguments& InArgs )
 		.AssetShowWarningText( InArgs._AssetPickerConfig.AssetShowWarningText)
 		.AllowFocusOnSync(false)	// Stop the asset view from stealing focus (we're in control of that)
 		.HiddenColumnNames(InArgs._AssetPickerConfig.HiddenColumnNames)
-
+		.ListHiddenColumnNames(InArgs._AssetPickerConfig.ListHiddenColumnNames)
 		.CustomColumns(InArgs._AssetPickerConfig.CustomColumns)
 		.OnSearchOptionsChanged(this, &SAssetPicker::HandleSearchSettingsChanged)
 		.InitialThumbnailSize(InArgs._AssetPickerConfig.InitialThumbnailSize)
diff --git a/Engine/Source/Editor/ContentBrowser/Private/SAssetView.cpp b/Engine/Source/Editor/ContentBrowser/Private/SAssetView.cpp
--- a/Engine/Source/Editor/ContentBrowser/Private/SAssetView.cpp
+++ b/Engine/Source/Editor/ContentBrowser/Private/SAssetView.cpp
@@ -1133,7 +1133,16 @@ void SAssetView::Construct( const FArguments& InArgs )
 	// Listen for when view settings are changed
 	UContentBrowserSettings::OnSettingChanged().AddSP(this, &SAssetView::HandleSettingChanged);
 
-	ThumbnailSize = InArgs._InitialThumbnailSize;
+	ThumbnailSizes = TMap<EAssetViewType::Type, EThumbnailSize>
+	{
+		{EAssetViewType::List, InArgs._InitialThumbnailSize},
+		{EAssetViewType::Tile, InArgs._InitialThumbnailSize},
+		// Force only the column default to be tiny like the older CB
+		{EAssetViewType::Column, EThumbnailSize::Tiny},
+		{EAssetViewType::Custom, InArgs._InitialThumbnailSize},
+		// Set a default for this case even though it should never land here
+		{EAssetViewType::MAX, InArgs._InitialThumbnailSize}
+	};
 
 	// Get desktop metrics
 	FDisplayMetrics DisplayMetrics;
@@ -1164,6 +1173,16 @@ void SAssetView::Construct( const FArguments& InArgs )
 
 	TileViewNameHeight = 50;
 
+	// Need to assign the ViewType before updating the ThumbnailSizeValue
+	if (InArgs._InitialViewType >= 0 && InArgs._InitialViewType < EAssetViewType::MAX)
+	{
+		CurrentViewType = InArgs._InitialViewType;
+	}
+	else
+	{
+		CurrentViewType = EAssetViewType::Tile;
+	}
+
 	UpdateThumbnailSizeValue();
 	MinThumbnailScale = 0.2f * ThumbnailScaleRangeScalar;
 	MaxThumbnailScale = 1.9f * ThumbnailScaleRangeScalar;
@@ -1268,21 +1287,12 @@ void SAssetView::Construct( const FArguments& InArgs )
 	bAllowDragging = InArgs._AllowDragging;
 	bAllowFocusOnSync = InArgs._AllowFocusOnSync;
 	HiddenColumnNames = DefaultHiddenColumnNames = InArgs._HiddenColumnNames;
-	ListHiddenColumnNames = TArray<FString>();
+	ListHiddenColumnNames = DefaultListHiddenColumnNames = InArgs._ListHiddenColumnNames;
 	CustomColumns = InArgs._CustomColumns;
 	OnSearchOptionsChanged = InArgs._OnSearchOptionsChanged;
 	bShowPathViewFilters = InArgs._bShowPathViewFilters;
 	OnExtendAssetViewOptionsMenuContext = InArgs._OnExtendAssetViewOptionsMenuContext;
 	AssetViewOptionsProfile = InArgs._AssetViewOptionsProfile;
-	
-	if (InArgs._InitialViewType >= 0 && InArgs._InitialViewType < EAssetViewType::MAX)
-	{
-		CurrentViewType = InArgs._InitialViewType;
-	}
-	else
-	{
-		CurrentViewType = EAssetViewType::Tile;
-	}
 
 	bPendingSortFilteredItems = false;
 	bQuickFrontendListRefreshRequested = false;
@@ -1986,9 +1996,30 @@ void SAssetView::RequestQuickFrontendListRefresh()
 	bQuickFrontendListRefreshRequested = true;
 }
 
-FString SAssetView::GetThumbnailScaleSettingPath(const FString& SettingsString) const
+FString SAssetView::GetThumbnailScaleSettingPath(const FString& SettingsString, const FString& InViewTypeString) const
 {
-	return SettingsString + TEXT(".ThumbnailSize");
+	return SettingsString + TEXT(".ThumbnailSize") + InViewTypeString;
+}
+
+void SAssetView::LoadScaleSetting(const FString& IniFilename, const FString& IniSection, const FString& SettingsString, const FString& InViewTypeString, EThumbnailSize& OutThumbnailSize)
+{
+	int32 ThumbnailSizeConfig = (int32)EThumbnailSize::Medium;
+	if (GConfig->GetInt(*IniSection, *GetThumbnailScaleSettingPath(SettingsString, InViewTypeString), ThumbnailSizeConfig, IniFilename))
+	{
+		// Clamp value to normal range and update state
+		ThumbnailSizeConfig = FMath::Clamp<int32>(ThumbnailSizeConfig, 0, (int32)EThumbnailSize::MAX-1);
+
+		// TODO: Remove this afterwards, current CB should hide new size
+		if (!UE::Editor::ContentBrowser::IsNewStyleEnabled())
+		{
+			if ((EThumbnailSize)ThumbnailSizeConfig == EThumbnailSize::XLarge)
+			{
+				ThumbnailSizeConfig -= 1;
+			}
+		}
+
+		OutThumbnailSize = (EThumbnailSize)ThumbnailSizeConfig;
+	}
 }
 
 FString SAssetView::GetCurrentViewTypeSettingPath(const FString& SettingsString) const
@@ -1998,37 +2029,46 @@ FString SAssetView::GetCurrentViewTypeSettingPath(const FString& SettingsString)
 
 void SAssetView::SaveSettings(const FString& IniFilename, const FString& IniSection, const FString& SettingsString) const
 {
-	GConfig->SetInt(*IniSection, *GetThumbnailScaleSettingPath(SettingsString), (int32)ThumbnailSize, IniFilename);
+	// ThumbnailSize saves
+	GConfig->SetInt(*IniSection, *GetThumbnailScaleSettingPath(SettingsString, GridViewSpecifier), (int32)ThumbnailSizes[EAssetViewType::Tile], IniFilename);
+	GConfig->SetInt(*IniSection, *GetThumbnailScaleSettingPath(SettingsString, ListViewSpecifier), (int32)ThumbnailSizes[EAssetViewType::List], IniFilename);
+	GConfig->SetInt(*IniSection, *GetThumbnailScaleSettingPath(SettingsString, CustomViewSpecifier), (int32)ThumbnailSizes[EAssetViewType::Custom], IniFilename);
+
+	// Save the ThumbnailSize config for the ColumnView only in the new CB
+	if (UE::Editor::ContentBrowser::IsNewStyleEnabled())
+	{
+		GConfig->SetInt(*IniSection, *GetThumbnailScaleSettingPath(SettingsString, ColumnViewSpecifier), (int32)ThumbnailSizes[EAssetViewType::Column], IniFilename);
+	}
+
 	GConfig->SetInt(*IniSection, *GetCurrentViewTypeSettingPath(SettingsString), CurrentViewType, IniFilename);
 	GConfig->SetFloat(*IniSection, *(SettingsString + TEXT(".ZoomScale")), ZoomScale, IniFilename);
 
 	GConfig->SetArray(*IniSection, *(SettingsString + TEXT(".HiddenColumns")), HiddenColumnNames, IniFilename);
 	if (UE::Editor::ContentBrowser::IsNewStyleEnabled())
 	{
+		// Used to discern if at some point the ColumnVisibility changed, if true the LoadedColumns will always be used instead
+		GConfig->SetBool(*IniSection, *(SettingsString + TEXT(".ListViewColumnsManuallyChangedOnce")), bListViewColumnsManuallyChangedOnce, IniFilename);
+		GConfig->SetBool(*IniSection, *(SettingsString + TEXT(".ColumnViewColumnsManuallyChangedOnce")), bColumnViewColumnsManuallyChangedOnce, IniFilename);
+
+		// ListView HiddenColumns
 		GConfig->SetArray(*IniSection, *(SettingsString + TEXT(".ListHiddenColumns")), ListHiddenColumnNames, IniFilename);
 	}
 }
 
 void SAssetView::LoadSettings(const FString& IniFilename, const FString& IniSection, const FString& SettingsString)
 {
-	int32 ThumbnailSizeConfig = (int32)EThumbnailSize::Medium;
-	if ( GConfig->GetInt(*IniSection, *GetThumbnailScaleSettingPath(SettingsString), ThumbnailSizeConfig, IniFilename) )
-	{
-		// Clamp value to normal range and update state
-		ThumbnailSizeConfig = FMath::Clamp<int32>(ThumbnailSizeConfig, 0, (int32)EThumbnailSize::MAX-1);
+	// Set the LoadSetting flag to true
+	TGuardValue<bool> ScopeGuard(bLoadingSettings, true);
 
-		// TODO: Remove this afterwards, current CB should hide new size
-		if (!UE::Editor::ContentBrowser::IsNewStyleEnabled())
-		{
-			if ((EThumbnailSize)ThumbnailSizeConfig == EThumbnailSize::XLarge)
-			{
-				ThumbnailSizeConfig -= 1;
-			}
-		}
+	// ThumbnailSize loadings
+	LoadScaleSetting(IniFilename, *IniSection, SettingsString, GridViewSpecifier, ThumbnailSizes[EAssetViewType::Tile]);
+	LoadScaleSetting(IniFilename, *IniSection, SettingsString, ListViewSpecifier, ThumbnailSizes[EAssetViewType::List]);
+	LoadScaleSetting(IniFilename, *IniSection, SettingsString, CustomViewSpecifier, ThumbnailSizes[EAssetViewType::Custom]);
 
-		ThumbnailSize = (EThumbnailSize) ThumbnailSizeConfig;
-		// Set the thumbnail value here after the Size enum is retrieved from the config
-		UpdateThumbnailSizeValue();
+	// Load the ThumbnailSize config for the ColumnView only in the new CB
+	if (UE::Editor::ContentBrowser::IsNewStyleEnabled())
+	{
+		LoadScaleSetting(IniFilename, *IniSection, SettingsString, ColumnViewSpecifier, ThumbnailSizes[EAssetViewType::Column]);
 	}
 
 	int32 ViewType = EAssetViewType::Tile;
@@ -2040,11 +2080,6 @@ void SAssetView::LoadSettings(const FString& IniFilename, const FString& IniSect
 			ViewType = EAssetViewType::Tile;
 		}
 
-		// If the new style is enabled do not allow the ColumnView
-		if (UE::Editor::ContentBrowser::IsNewStyleEnabled() && ViewType == 2 /** Column */)
-		{
-			ViewType = 1; // List
-		}
 		SetCurrentViewType((EAssetViewType::Type)ViewType);
 	}
 
@@ -2055,9 +2090,19 @@ void SAssetView::LoadSettings(const FString& IniFilename, const FString& IniSect
 		ZoomScale = FMath::Clamp<float>(Zoom, 0.f, 1.f);
 	}
 
+	if (UE::Editor::ContentBrowser::IsNewStyleEnabled())
+	{
+		bool bColumnChangedManually = false;
+		if (GConfig->GetBool(*IniSection, *(SettingsString + TEXT(".ColumnViewColumnsManuallyChangedOnce")), bColumnChangedManually, IniFilename) )
+		{
+			// Whether the column where changed by the user even once for this Config, if yes always use the LoadedColumns
+			bColumnViewColumnsManuallyChangedOnce = bColumnChangedManually;
+		}
+	}
+
 	TArray<FString> LoadedHiddenColumnNames;
 	GConfig->GetArray(*IniSection, *(SettingsString + TEXT(".HiddenColumns")), LoadedHiddenColumnNames, IniFilename);
-	if (LoadedHiddenColumnNames.Num() > 0)
+	if (LoadedHiddenColumnNames.Num() > 0 || bColumnViewColumnsManuallyChangedOnce)
 	{
 		HiddenColumnNames = LoadedHiddenColumnNames;
 
@@ -2073,9 +2118,16 @@ void SAssetView::LoadSettings(const FString& IniFilename, const FString& IniSect
 
 	if (UE::Editor::ContentBrowser::IsNewStyleEnabled())
 	{
+		bool bColumnChangedManually = false;
+		if ( GConfig->GetBool(*IniSection, *(SettingsString + TEXT(".ListViewColumnsManuallyChangedOnce")), bColumnChangedManually, IniFilename) )
+		{
+			// Whether the column where changed by the user even once for this Config, if yes always use the LoadedColumns
+			bListViewColumnsManuallyChangedOnce = bColumnChangedManually;
+		}
+
 		TArray<FString> LoadedListHiddenColumnNames;
 		GConfig->GetArray(*IniSection, *(SettingsString + TEXT(".ListHiddenColumns")), LoadedListHiddenColumnNames, IniFilename);
-		if (LoadedListHiddenColumnNames.Num() > 0)
+		if (LoadedListHiddenColumnNames.Num() > 0 || bListViewColumnsManuallyChangedOnce)
 		{
 			ListHiddenColumnNames = LoadedListHiddenColumnNames;
 
@@ -2685,16 +2737,16 @@ FReply SAssetView::OnMouseWheel( const FGeometry& MyGeometry, const FPointerEven
 		if (UE::Editor::ContentBrowser::IsNewStyleEnabled())
 		{
 			const int32 Delta = MouseEvent.GetWheelDelta() > 0 ? 1 : -1;
-			bool IsLessThanMinSize = (int32)ThumbnailSize + Delta < 0;
-			bool IsMoreThanMaxSize = (int32)ThumbnailSize + Delta == (int32)EThumbnailSize::MAX;
+			bool IsLessThanMinSize = (int32)ThumbnailSizes[CurrentViewType] + Delta < 0;
+			bool IsMoreThanMaxSize = (int32)ThumbnailSizes[CurrentViewType] + Delta == (int32)EThumbnailSize::MAX;
 			bool bWasSizeSupposedToChange = false;
 			// If shift is pressed or if we are in the List view the CTRL + Wheel should jump ThumbnailSize by design
 			if (MouseEvent.IsShiftDown() || CurrentViewType != EAssetViewType::Tile)
 			{
 				bWasSizeSupposedToChange = true;
-				EThumbnailSize DesiredThumbnailSize = (EThumbnailSize)FMath::Clamp<int32>((int32)ThumbnailSize + Delta, 0, (int32)EThumbnailSize::MAX - 1);
+				EThumbnailSize DesiredThumbnailSize = (EThumbnailSize)FMath::Clamp<int32>((int32)ThumbnailSizes[CurrentViewType] + Delta, 0, (int32)EThumbnailSize::MAX - 1);
 
-				if (DesiredThumbnailSize != ThumbnailSize)
+				if (DesiredThumbnailSize != ThumbnailSizes[CurrentViewType])
 				{
 					OnThumbnailSizeChanged(DesiredThumbnailSize);
 				}
@@ -2706,11 +2758,11 @@ FReply SAssetView::OnMouseWheel( const FGeometry& MyGeometry, const FPointerEven
 				{
 					bWasSizeSupposedToChange = true;
 					const int32 Step = (int32)FMath::Sign(NewDelta);
-					EThumbnailSize OldSize = ThumbnailSize;
-					ThumbnailSize = (EThumbnailSize)FMath::Clamp<int32>(((int32)ThumbnailSize + Step), 0, (int32)EThumbnailSize::MAX - 1);
-					if (OldSize != ThumbnailSize)
+					EThumbnailSize OldSize = ThumbnailSizes[CurrentViewType];
+					ThumbnailSizes[CurrentViewType] = (EThumbnailSize)FMath::Clamp<int32>(((int32)ThumbnailSizes[CurrentViewType] + Step), 0, (int32)EThumbnailSize::MAX - 1);
+					if (OldSize != ThumbnailSizes[CurrentViewType])
 					{
-						OnThumbnailSizeChanged(ThumbnailSize);
+						OnThumbnailSizeChanged(ThumbnailSizes[CurrentViewType]);
 						ZoomScale = NewDelta > 0.f ? 0.f : 1.f;
 					}
 				}
@@ -2743,7 +2795,7 @@ FReply SAssetView::OnMouseWheel( const FGeometry& MyGeometry, const FPointerEven
 			// Step up/down a level depending on the scroll wheel direction.
 			// Clamp value to enum min/max before updating.
 			const int32 Delta = MouseEvent.GetWheelDelta() > 0 ? 1 : -1;
-			EThumbnailSize DesiredThumbnailSize = (EThumbnailSize)FMath::Clamp<int32>((int32)ThumbnailSize + Delta, 0, (int32)EThumbnailSize::MAX - 1);
+			EThumbnailSize DesiredThumbnailSize = (EThumbnailSize)FMath::Clamp<int32>((int32)ThumbnailSizes[CurrentViewType] + Delta, 0, (int32)EThumbnailSize::MAX - 1);
 
 			// TODO: Remove this afterwards, current CB should hide new size
 			if (DesiredThumbnailSize == EThumbnailSize::XLarge)
@@ -2758,7 +2810,7 @@ FReply SAssetView::OnMouseWheel( const FGeometry& MyGeometry, const FPointerEven
 				}
 			}
 
-			if ( DesiredThumbnailSize != ThumbnailSize )
+			if ( DesiredThumbnailSize != ThumbnailSizes[CurrentViewType] )
 			{
 				OnThumbnailSizeChanged(DesiredThumbnailSize);
 			}
@@ -2874,7 +2926,7 @@ TSharedRef<SAssetListView> SAssetView::CreateListView()
 					.DefaultLabel(LOCTEXT("Column_Path", "Path"))
 				);
 
-			const bool bIsColumnVisible = !HiddenColumnNames.Contains(SortManager->PathColumnId.ToString());
+			const bool bIsColumnVisible = !HiddenColumnsToUse.Contains(SortManager->PathColumnId.ToString());
 		
 			NewListView->GetHeaderRow()->SetShowGeneratedColumn(SortManager->PathColumnId, bIsColumnVisible);
 		}
@@ -4794,7 +4846,7 @@ void SAssetView::SetCurrentViewType(EAssetViewType::Type NewType)
 		{
 			if (UE::Editor::ContentBrowser::IsNewStyleEnabled())
 			{
-				if (ThumbnailSize >= EThumbnailSize::Small)
+				if (ThumbnailSizes[CurrentViewType] >= EThumbnailSize::Small)
 				{
 					CurrentThumbnailSize = ListViewThumbnailSize;
 					bPendingUpdateThumbnails = true;
@@ -4810,7 +4862,7 @@ void SAssetView::SetCurrentViewType(EAssetViewType::Type NewType)
 		{
 			if (UE::Editor::ContentBrowser::IsNewStyleEnabled())
 			{
-				if (ThumbnailSize >= EThumbnailSize::Small)
+				if (ThumbnailSizes[CurrentViewType] >= EThumbnailSize::Small)
 				{
 					CurrentThumbnailSize = ListViewThumbnailSize;
 					bPendingUpdateThumbnails = true;
@@ -4822,6 +4874,9 @@ void SAssetView::SetCurrentViewType(EAssetViewType::Type NewType)
 			SortList();
 		}
 
+		// Update the size value when switching view to match the current view size
+		UpdateThumbnailSizeValue();
+
 		if (FAssetViewInstanceConfig* Config = GetAssetViewConfig())
 		{
 			Config->ViewType = (uint8) NewType;
@@ -4832,7 +4887,7 @@ void SAssetView::SetCurrentViewType(EAssetViewType::Type NewType)
 
 void SAssetView::SetCurrentThumbnailSize(EThumbnailSize NewThumbnailSize)
 {
-	if (ThumbnailSize != NewThumbnailSize)
+	if (ThumbnailSizes[CurrentViewType] != NewThumbnailSize)
 	{
 		OnThumbnailSizeChanged(NewThumbnailSize);
 	}
@@ -5738,7 +5793,7 @@ void SAssetView::ToggleThumbnailEditMode()
 
 void SAssetView::OnThumbnailSizeChanged(EThumbnailSize NewThumbnailSize)
 {
-	ThumbnailSize = NewThumbnailSize;
+	ThumbnailSizes[CurrentViewType] = NewThumbnailSize;
 	UpdateThumbnailSizeValue();
 
 	if (FAssetViewInstanceConfig* Config = GetAssetViewConfig())
@@ -5752,13 +5807,13 @@ void SAssetView::OnThumbnailSizeChanged(EThumbnailSize NewThumbnailSize)
 
 bool SAssetView::IsThumbnailSizeChecked(EThumbnailSize InThumbnailSize) const
 {
-	return ThumbnailSize == InThumbnailSize;
+	return ThumbnailSizes[CurrentViewType] == InThumbnailSize;
 }
 
 float SAssetView::GetThumbnailScale() const
 {
 	float BaseScale;
-	switch (ThumbnailSize)
+	switch (ThumbnailSizes[CurrentViewType])
 	{
 	case EThumbnailSize::Tiny:
 		BaseScale = 0.1f;
@@ -5793,7 +5848,7 @@ float SAssetView::GetThumbnailSizeValue() const
 
 void SAssetView::UpdateThumbnailSizeValue()
 {
-	switch (ThumbnailSize)
+	switch (ThumbnailSizes[CurrentViewType])
 	{
 	case EThumbnailSize::Tiny:
 		MinThumbnailSize = 64.f;
@@ -5842,7 +5897,7 @@ float SAssetView::GetTileViewTypeNameHeight() const
 {
 	if (UE::Editor::ContentBrowser::IsNewStyleEnabled())
 	{
-		if (ThumbnailSize == EThumbnailSize::Tiny)
+		if (ThumbnailSizes[CurrentViewType] == EThumbnailSize::Tiny)
 		{
 			return 0;
 		}
@@ -5857,15 +5912,15 @@ float SAssetView::GetTileViewTypeNameHeight() const
 		}
 		else
 		{
-			if (ThumbnailSize == EThumbnailSize::Small)
+			if (ThumbnailSizes[CurrentViewType] == EThumbnailSize::Small)
 			{
 				TypeNameHeight = 25;
 			}
-			else if (ThumbnailSize == EThumbnailSize::Medium)
+			else if (ThumbnailSizes[CurrentViewType] == EThumbnailSize::Medium)
 			{
 				TypeNameHeight = -5;
 			}
-			else if (ThumbnailSize > EThumbnailSize::Medium)
+			else if (ThumbnailSizes[CurrentViewType] > EThumbnailSize::Medium)
 			{
 				TypeNameHeight = -25;
 			}
@@ -5880,7 +5935,7 @@ float SAssetView::GetTileViewTypeNameHeight() const
 
 float SAssetView::GetSourceControlIconHeight() const
 {
-	return (float)(ThumbnailSize != EThumbnailSize::Tiny && ISourceControlModule::Get().IsEnabled() && ISourceControlModule::Get().GetProvider().IsAvailable() && !bShowTypeInTileView ? 17.0 : 0.0);
+	return (float)(ThumbnailSizes[CurrentViewType] != EThumbnailSize::Tiny && ISourceControlModule::Get().IsEnabled() && ISourceControlModule::Get().GetProvider().IsAvailable() && !bShowTypeInTileView ? 17.0 : 0.0);
 }
 
 float SAssetView::GetListViewItemHeight() const
@@ -5893,7 +5948,7 @@ float SAssetView::GetListViewItemHeight() const
 FMargin SAssetView::GetListViewItemPadding() const
 {
 	return UE::Editor::ContentBrowser::IsNewStyleEnabled()
-		? ThumbnailSize == EThumbnailSize::Tiny ? FMargin(0.f) : FMargin(0.f, (float)ListViewThumbnailPadding)
+		? ThumbnailSizes[CurrentViewType] == EThumbnailSize::Tiny ? FMargin(0.f) : FMargin(0.f, (float)ListViewThumbnailPadding)
 		: FMargin(0.f);
 }
 
@@ -6189,15 +6244,36 @@ bool SAssetView::PerformQuickJump(const bool bWasJumping)
 
 void SAssetView::ResetColumns()
 {
-	for (const SHeaderRow::FColumn &Column : ColumnView->GetHeaderRow()->GetColumns())
+	TSharedPtr<SListView<TSharedPtr<FAssetViewItem>>> ViewToUse = ColumnView;
+	TArray<FString>* HiddenColumnsToUse = &HiddenColumnNames;
+	TArray<FString> DefaultHiddenColumnToUse = DefaultHiddenColumnNames;
+	if (UE::Editor::ContentBrowser::IsNewStyleEnabled())
+	{
+		if (CurrentViewType == EAssetViewType::List)
+		{
+			ViewToUse = ListView;
+			HiddenColumnsToUse = &ListHiddenColumnNames;
+			DefaultHiddenColumnToUse = DefaultListHiddenColumnNames;
+			// When resetting list view columns, reset also this to use the default
+			bListViewColumnsManuallyChangedOnce = false;
+		}
+		else
+		{
+			// When resetting column view columns, reset also this to use the default
+			bColumnViewColumnsManuallyChangedOnce = false;
+		}
+	}
+
+	for (const SHeaderRow::FColumn &Column : ViewToUse->GetHeaderRow()->GetColumns())
 	{
-		ColumnView->GetHeaderRow()->SetShowGeneratedColumn(Column.ColumnId, !DefaultHiddenColumnNames.Contains(Column.ColumnId.ToString()));
+		ViewToUse->GetHeaderRow()->SetShowGeneratedColumn(Column.ColumnId, !DefaultHiddenColumnToUse.Contains(Column.ColumnId.ToString()));
 	}
 
 	// This is set after updating the column visibilties, because SetShowGeneratedColumn calls OnHiddenColumnsChanged indirectly which can mess up the list
-	HiddenColumnNames = DefaultHiddenColumnNames;
-	ColumnView->GetHeaderRow()->RefreshColumns();
-	ColumnView->RebuildList();
+	HiddenColumnsToUse->Reset(DefaultHiddenColumnToUse.Num());
+	HiddenColumnsToUse->Append(DefaultHiddenColumnToUse);
+	ViewToUse->GetHeaderRow()->RefreshColumns();
+	ViewToUse->RebuildList();
 }
 
 void SAssetView::ExportColumns()
@@ -6239,10 +6315,18 @@ void SAssetView::ExportColumns()
 
 void SAssetView::OnHiddenColumnsChanged()
 {
+	const bool bIsUsingNewStyle = UE::Editor::ContentBrowser::IsNewStyleEnabled();
+
+	// Early out if this is called before creation or during LoadSettings due to SetShowGeneratedColumn(due to loading config etc)
+	if (bIsUsingNewStyle && bLoadingSettings)
+	{
+		return;
+	}
+
 	// Early out if this is called before creation (due to loading config etc)
 	TSharedPtr<SListView<TSharedPtr<FAssetViewItem>>> ViewToUse = ColumnView;
 	TArray<FString>* HiddenColumnsToUse = &HiddenColumnNames;
-	if (UE::Editor::ContentBrowser::IsNewStyleEnabled())
+	if (bIsUsingNewStyle)
 	{
 		if (CurrentViewType == EAssetViewType::List)
 		{
@@ -6256,6 +6340,20 @@ void SAssetView::OnHiddenColumnsChanged()
 		return;
 	}
 
+	if (bIsUsingNewStyle)
+	{
+		if (CurrentViewType == EAssetViewType::List)
+		{
+			// Set this to true as soon as the first column is modified by the user
+			bListViewColumnsManuallyChangedOnce = true;
+		}
+		else
+		{
+			// Set this to true as soon as the first column is modified by the user
+			bColumnViewColumnsManuallyChangedOnce = true;
+		}
+	}
+
 	// We can't directly update the hidden columns list, because some columns maybe hidden, but not created yet
 	TArray<FName> NewHiddenColumns = ViewToUse->GetHeaderRow()->GetHiddenColumnIds();
 
diff --git a/Engine/Source/Editor/ContentBrowser/Public/IContentBrowserSingleton.h b/Engine/Source/Editor/ContentBrowser/Public/IContentBrowserSingleton.h
--- a/Engine/Source/Editor/ContentBrowser/Public/IContentBrowserSingleton.h
+++ b/Engine/Source/Editor/ContentBrowser/Public/IContentBrowserSingleton.h
@@ -198,6 +198,9 @@ struct FAssetPickerConfig
 	/** The names of columns to hide by default in the column view */
 	TArray<FString> HiddenColumnNames;
 
+	/** The names of columns to hide by default in the list view, by default the Type column is hidden */
+	TArray<FString> ListHiddenColumnNames;
+
 	/** List of custom columns that fill out data with a callback */
 	TArray<FAssetViewCustomColumn> CustomColumns;
 
@@ -343,6 +346,7 @@ struct FAssetPickerConfig
 
 	FAssetPickerConfig()
 		: SelectionMode( ESelectionMode::Multi )
+		, ListHiddenColumnNames( { TEXT("Class") } )
 		, ThumbnailLabel( EThumbnailLabel::ClassName )
 		, ThumbnailScale(0.1f)
 		, InitialThumbnailSize(EThumbnailSize::Medium)
diff --git a/Engine/Source/Editor/ContentBrowser/Public/SAssetView.h b/Engine/Source/Editor/ContentBrowser/Public/SAssetView.h
--- a/Engine/Source/Editor/ContentBrowser/Public/SAssetView.h
+++ b/Engine/Source/Editor/ContentBrowser/Public/SAssetView.h
@@ -304,9 +304,12 @@ public:
 		/** Called to add extra asset data to the asset view, to display virtual assets. These get treated similar to Class assets */
 		SLATE_EVENT( FOnGetCustomSourceAssets, OnGetCustomSourceAssets )
 
-		/** Columns to hide by default */
+		/** Columns to hide by default for the ColumnView */
 		SLATE_ARGUMENT( TArray<FString>, HiddenColumnNames )
-			
+
+		/** Columns to hide by default for the ListView, works only in the new style */
+		SLATE_ARGUMENT( TArray<FString>, ListHiddenColumnNames )
+
 		/** Custom columns that can be use specific */
 		SLATE_ARGUMENT( TArray<FAssetViewCustomColumn>, CustomColumns )
 
@@ -515,12 +518,14 @@ public:
 		const TOptional<const EColumnSortMode::Type>& InNewSortMode);
 
 private:
-
 	/** Sets the pending selection to the current selection (used when changing views or refreshing the view). */
 	void SyncToSelection( const bool bFocusOnSync = true );
 
 	/** @return the thumbnail scale setting path to use when looking up the setting in an ini. */
-	FString GetThumbnailScaleSettingPath(const FString& SettingsString) const;
+	FString GetThumbnailScaleSettingPath(const FString& SettingsString, const FString& InViewTypeString) const;
+
+	/** Load the config for each AssetViewType size into their own container */
+	void LoadScaleSetting(const FString& IniFilename, const FString& IniSection, const FString& SettingsString, const FString& InViewTypeString, EThumbnailSize& OutThumbnailSize);
 
 	/** @return the view type setting path to use when looking up the setting in an ini. */
 	FString GetCurrentViewTypeSettingPath(const FString& SettingsString) const;
@@ -859,7 +864,7 @@ private:
 	void UpdateThumbnailSizeValue();
 
 	/** Gets the current thumbnail size enum */
-	EThumbnailSize GetThumbnailSize() const { return ThumbnailSize; }
+	EThumbnailSize GetThumbnailSize() const { return ThumbnailSizes.Contains(CurrentViewType) ? ThumbnailSizes[CurrentViewType] : EThumbnailSize::Medium; }
 
 	/** Gets the spacing dedicated to support type names */
 	float GetTileViewTypeNameHeight() const;
@@ -1052,6 +1057,15 @@ private:
 	/** If true, the frontend items will be refreshed next frame. Much faster. */
 	bool bQuickFrontendListRefreshRequested;
 
+	/** True when loading settings, this is needed for the HiddenColumns for example to not double update them */
+	bool bLoadingSettings = false;
+
+	/** Set to true as soon as the user change the ListView visibility of a column, after that point the Loaded Hidden Columns will be used instead of any default given */
+	bool bListViewColumnsManuallyChangedOnce = false;
+
+	/** Set to true as soon as the user change the ColumnView visibility of a column, after that point the Loaded Hidden Columns will be used instead of any default given */
+	bool bColumnViewColumnsManuallyChangedOnce = false;
+
 	/** The list of items to sync next frame */
 	FSelectionData PendingSyncItems;
 
@@ -1169,8 +1183,14 @@ private:
 	/** Scalar applied to thumbnail sizes so that users thumbnails are scaled based on users display area size*/
 	float ThumbnailScaleRangeScalar;
 
-	/** Current thumbnail size */
-	EThumbnailSize ThumbnailSize;
+	/** Current thumbnail sizes */
+	TMap<EAssetViewType::Type, EThumbnailSize> ThumbnailSizes;
+
+	// Specifiers used for Config
+	const FString GridViewSpecifier = TEXT("Grid");
+	const FString ListViewSpecifier = TEXT("List");
+	const FString ColumnViewSpecifier = TEXT("Column");
+	const FString CustomViewSpecifier = TEXT("Custom");
 
 	/** The amount to scale each thumbnail so that the empty space is filled. */
 	float FillScale;
@@ -1333,6 +1353,7 @@ private:
 	TArray<FString> DefaultHiddenColumnNames;
 	TArray<FString> HiddenColumnNames;
 	TArray<FString> ListHiddenColumnNames;
+	TArray<FString> DefaultListHiddenColumnNames;
 
 	TArray<FAssetViewCustomColumn> CustomColumns;
 
