[ContentBrowser] Tiny Rename ESC will avoid creation and rename like in the old version

Pressing ESC during AssetRename will avoid the rename or the creation if the asset is being created like in the old CB

diff --git a/Engine/Source/Editor/ContentBrowser/Private/AssetViewWidgets.cpp b/Engine/Source/Editor/ContentBrowser/Private/AssetViewWidgets.cpp
--- a/Engine/Source/Editor/ContentBrowser/Private/AssetViewWidgets.cpp
+++ b/Engine/Source/Editor/ContentBrowser/Private/AssetViewWidgets.cpp
@@ -2403,7 +2403,6 @@ void SAssetTileItem::Construct( const FArguments& InArgs )
 		.IsSelected(InArgs._IsSelected)
 		);
 
-	bCreatedInTinySize = CurrentThumbnailSize.Get() == EThumbnailSize::Tiny;
 	bShowType = InArgs._ShowType;
 	AssetThumbnail = InArgs._AssetThumbnail;
 	ItemWidth = InArgs._ItemWidth;
@@ -2824,15 +2823,10 @@ void SAssetTileItem::OnRequestingRename()
 	FSlateApplication::Get().AddWindow(
 		SNew(SRenameWindow)
 		.ScreenPosition(Position)
-		.ForceRename(bCreatedInTinySize)
 		.OnBeginTextEdit(this, &SAssetTileItem::HandleBeginNameChange)
 		.OnVerifyTextChanged(this, &SAssetTileItem::HandleVerifyNameChanged)
 		.InitialText(InlineRenameWidget->GetText())
 		.OnTextCommitted(this, &SAssetTileItem::HandleNameCommitted));
-
-	// After the first time this is called, it shouldn't force the rename anymore.
-	// This is used on AssetCreation rename to avoid destroying the asset immediately if no name is given.
-	bCreatedInTinySize = false;
 }
 
 float SAssetTileItem::GetExtraStateIconWidth() const
diff --git a/Engine/Source/Editor/ContentBrowser/Private/SRenameWindow.cpp b/Engine/Source/Editor/ContentBrowser/Private/SRenameWindow.cpp
--- a/Engine/Source/Editor/ContentBrowser/Private/SRenameWindow.cpp
+++ b/Engine/Source/Editor/ContentBrowser/Private/SRenameWindow.cpp
@@ -27,6 +27,7 @@ class SRenameEditableTextBox : public SEditableTextBox
 		OnRenameCommittedDelegate = InArgs._OnTextCommitted;
 
 		SEditableTextBox::Construct(SEditableTextBox::FArguments()
+			.SelectAllTextWhenFocused(true)
 			.OnVerifyTextChanged(InArgs._OnVerifyTextChanged)
 			.Text(InArgs._InitialText)
 			.OnTextCommitted(this, &SRenameEditableTextBox::OnRenameTextCommitted));
@@ -45,11 +46,12 @@ class SRenameEditableTextBox : public SEditableTextBox
 	{
 		if (InKeyEvent.GetKey() == EKeys::Escape)
 		{
+			EndRename();
 			if (const TSharedPtr<SWindow>& Window = OwnerWindow.Pin())
 			{
-				EndRename();
-				return Window->OnKeyDown(MyGeometry, InKeyEvent);
+				Window->RequestDestroyWindow();
 			}
+			return FReply::Handled();
 		}
 		return SEditableTextBox::OnKeyDown(MyGeometry, InKeyEvent);
 	}
@@ -61,6 +63,10 @@ private:
 		{
 			bOnTextCommittedCalled = true;
 			OnRenameCommittedDelegate.ExecuteIfBound(InText, InCommitType);
+			if (const TSharedPtr<SWindow>& Window = OwnerWindow.Pin())
+			{
+				Window->RequestDestroyWindow();
+			}
 		}
 	}
 
@@ -77,8 +83,6 @@ private:
 
 void SRenameWindow::Construct(const FArguments& InArgs)
 {
-	bForceRename = InArgs._ForceRename;
-
 	float DPIScale = 1.f;
 	{
 		IMainFrameModule& MainFrameModule = FModuleManager::LoadModuleChecked<IMainFrameModule>(TEXT("MainFrame"));
@@ -120,21 +124,13 @@ void SRenameWindow::Construct(const FArguments& InArgs)
 	OnWindowDeactivated.BindSP(this, &SRenameWindow::DeactivateWindow);
 }
 
-FReply SRenameWindow::OnKeyDown(const FGeometry& MyGeometry, const FKeyEvent& InKeyEvent)
-{
-	if (InKeyEvent.GetKey() == EKeys::Escape)
-	{
-		DeactivateWindow();
-		return FReply::Handled();
-	}
-	return SWindow::OnKeyDown(MyGeometry, InKeyEvent);
-}
-
 void SRenameWindow::DeactivateWindow()
 {
 	if (RenameEditableTextBox.IsValid())
 	{
-		RenameEditableTextBox->EndRename(bForceRename);
+		// Always force it when deactivating the window
+		// This will not rename when pressing ESC since the rename already happened with TextCommit = OnCleared
+		RenameEditableTextBox->EndRename(true);
 	}
 	RequestDestroyWindow();
 }
diff --git a/Engine/Source/Editor/ContentBrowser/Private/SRenameWindow.h b/Engine/Source/Editor/ContentBrowser/Private/SRenameWindow.h
--- a/Engine/Source/Editor/ContentBrowser/Private/SRenameWindow.h
+++ b/Engine/Source/Editor/ContentBrowser/Private/SRenameWindow.h
@@ -18,7 +18,6 @@ class SRenameWindow : public SWindow
 		, _OnVerifyTextChanged()
 		, _OnBeginTextEdit()
 		, _InitialText()
-		, _ForceRename(false)
 	{}
 
 	SLATE_ARGUMENT(UE::Slate::FDeprecateVector2DParameter, ScreenPosition)
@@ -31,18 +30,13 @@ class SRenameWindow : public SWindow
 
 	SLATE_ARGUMENT(FText, InitialText)
 
-	SLATE_ARGUMENT(bool, ForceRename)
-
 	SLATE_END_ARGS()
 
 	void Construct(const FArguments& InArgs);
 
-	virtual FReply OnKeyDown(const FGeometry& MyGeometry, const FKeyEvent& InKeyEvent) override;
-
 private:
 	void DeactivateWindow();
 
 private:
 	TSharedPtr<SRenameEditableTextBox> RenameEditableTextBox;
-	bool bForceRename;
 };
diff --git a/Engine/Source/Editor/ContentBrowser/Public/AssetViewWidgets.h b/Engine/Source/Editor/ContentBrowser/Public/AssetViewWidgets.h
--- a/Engine/Source/Editor/ContentBrowser/Public/AssetViewWidgets.h
+++ b/Engine/Source/Editor/ContentBrowser/Public/AssetViewWidgets.h
@@ -844,9 +844,6 @@ private:
 
 	/** Current thumbnail size when this widget was generated */
 	TAttribute<EThumbnailSize> CurrentThumbnailSize;
-
-	/** Set to true if the current size during the tile item creation is Tiny */
-	bool bCreatedInTinySize = false;
 };
 
 /** An item in the asset column view */
