Fix issue where "Dock in Layout" wasn't aligned to the top.

#rnx

diff --git a/Engine/Source/Editor/ContentBrowser/Private/SContentBrowser.cpp b/Engine/Source/Editor/ContentBrowser/Private/SContentBrowser.cpp
--- a/Engine/Source/Editor/ContentBrowser/Private/SContentBrowser.cpp
+++ b/Engine/Source/Editor/ContentBrowser/Private/SContentBrowser.cpp
@@ -3563,7 +3563,8 @@ TSharedRef<SWidget> SContentBrowser::CreateLegacyAssetViewWidgets()
 									+ SHorizontalBox::Slot()
 									.AutoWidth()
 									.HAlign(HAlign_Right)
-									.VAlign(VAlign_Center)
+									.VAlign(VAlign_Top)
+									.Padding(0.f, 2.f, 0.f, 0.f)
 									[
 										CreateDrawerDockButton(&InitConfig)
 									]
