[ContentBrowser] Ctrl+RightClick will always create the AddMenu in the CB

diff --git a/Engine/Source/Editor/ContentBrowser/Private/SContentBrowser.cpp b/Engine/Source/Editor/ContentBrowser/Private/SContentBrowser.cpp
--- a/Engine/Source/Editor/ContentBrowser/Private/SContentBrowser.cpp
+++ b/Engine/Source/Editor/ContentBrowser/Private/SContentBrowser.cpp
@@ -3023,7 +3023,7 @@ FText SContentBrowser::GetAddNewToolTipText() const
 			return FText::Format(LOCTEXT("AddNewToolTip_CannotWrite", "Cannot write to path {0}..."), FText::FromString(CurrentPath));
 		}
 
-		return FText::Format( LOCTEXT("AddNewToolTip_AddNewContent", "Create a new content in {0}..."), FText::FromString(CurrentPath) );
+		return FText::Format( LOCTEXT("AddNewToolTip_AddNewContent", "Create new content in {0}...\nShortcut: Ctrl + RMB anywhere in the asset view"), FText::FromString(CurrentPath) );
 	}
 	else if ( ContentSources.GetVirtualPaths().Num() > 1 )
 	{
@@ -4982,7 +4982,11 @@ TSharedPtr<SWidget> SContentBrowser::GetItemContextMenu(TArrayView<const FConten
 		}
 	}
 
-	if (SelectedFolders.Num() > 0)
+	const bool bIsControlDown = FSlateApplication::Get().GetModifierKeys().IsControlDown();
+	const bool bIsAssetViewContext = ViewContext == EContentBrowserViewContext::AssetView;
+	const bool bShouldForceAddMenu = bIsControlDown && bIsAssetViewContext;
+
+	if (SelectedFolders.Num() > 0 && !bShouldForceAddMenu)
 	{
 		// Folders selected - show the folder menu
 
@@ -5073,13 +5077,13 @@ TSharedPtr<SWidget> SContentBrowser::GetItemContextMenu(TArrayView<const FConten
 		Context->SelectedPackagePaths = MoveTemp(SelectedPackagePaths);
 		return UToolMenus::Get()->GenerateWidget("ContentBrowser.FolderContextMenu", MenuContext);
 	}
-	else if (SelectedItems.Num() > 0)
+	else if (SelectedItems.Num() > 0 && !bShouldForceAddMenu)
 	{
 		// Files selected - show the file menu
 		checkf(ViewContext == EContentBrowserViewContext::AssetView, TEXT("File items were passed from a path view!"));
 		return LegacyContentSourceWidgets->AssetContextMenu->MakeContextMenu(SelectedItems, LegacyContentSourceWidgets->AssetViewPtr->GetContentSources(), Commands);
 	}
-	else if (ViewContext == EContentBrowserViewContext::AssetView)
+	else if (ViewContext == EContentBrowserViewContext::AssetView || bShouldForceAddMenu)
 	{
 		// Nothing selected - show the new asset menu
 		return MakeAddNewContextMenu(EContentBrowserDataMenuContext_AddNewMenuDomain::AssetView, nullptr);
