Content Browser Filter Bar:
- Fix for the All Filters menu to appear even if the Custom Filters menu is missing (happens in editors like Control Rig) -> Place the menu in Last instead of After
- Sort the classes before adding them to their sub menus

diff --git a/Engine/Source/Editor/ContentBrowser/Private/SFilterList.cpp b/Engine/Source/Editor/ContentBrowser/Private/SFilterList.cpp
--- a/Engine/Source/Editor/ContentBrowser/Private/SFilterList.cpp
+++ b/Engine/Source/Editor/ContentBrowser/Private/SFilterList.cpp
@@ -551,6 +551,15 @@ void SFilterList::PopulateAddFilterMenu(UToolMenu* Menu, TSharedPtr<FFilterCateg
 		return A->Title.CompareTo(B->Title) < 0;
 	});
 
+	// Sort each submenu internals 
+	for (TPair<TSharedPtr<FFilterCategory>, FFilterCategoryMenu>& CategoryToMenu : CategoryToMenuMap)
+	{
+		CategoryToMenu.Value.Classes.Sort([](const TSharedPtr<FCustomClassFilterData>& A, const TSharedPtr<FCustomClassFilterData>& B)
+		{
+			return A->GetName().CompareTo(B->GetName()) < 0;
+		});
+	}
+
 	FToolMenuSection& AdvancedSection = Menu->FindOrAddSection("FilterBarAdvanced");
 
 	FToolMenuEntry& AllFiltersSubMenu = AdvancedSection.AddSubMenu(
@@ -580,7 +589,7 @@ void SFilterList::PopulateAddFilterMenu(UToolMenu* Menu, TSharedPtr<FFilterCateg
 			this->PopulateCustomFilters(InMenu);
 		}));
 
-	AllFiltersSubMenu.InsertPosition = FToolMenuInsert("CustomFiltersSubMenu", EToolMenuInsertType::After);
+	AllFiltersSubMenu.InsertPosition.Position = EToolMenuInsertType::Last;
 }
 
 void SFilterList::CreateCustomFilterDialog(const FText& InText)
