[ContentBrowser] Re-Add hook for CustomTooltip

diff --git a/Engine/Source/Editor/ContentBrowser/Private/AssetViewWidgets.cpp b/Engine/Source/Editor/ContentBrowser/Private/AssetViewWidgets.cpp
--- a/Engine/Source/Editor/ContentBrowser/Private/AssetViewWidgets.cpp
+++ b/Engine/Source/Editor/ContentBrowser/Private/AssetViewWidgets.cpp
@@ -301,12 +301,20 @@ public:
 	void Construct(const FArguments& InArgs)
 	{
 		AssetViewItem = InArgs._AssetViewItem;
+		FMargin TooltipPadding = FMargin(0.f);
+		if (const TSharedPtr<SAssetViewItem> AssetViewItemPinned = AssetViewItem.Pin())
+		{
+			if (AssetViewItemPinned->IsFolder())
+			{
+				TooltipPadding = FMargin(12.f, 8.f, 12.f, 2.f);
+			}
+		}
 
 		SToolTip::Construct(
 			SToolTip::FArguments()
 			.TextMargin(
 				UE::Editor::ContentBrowser::IsNewStyleEnabled()
-				? FMargin(12.f, 8.f, 12.f, 2.f)
+				? TooltipPadding
 				: FMargin(1.0f))
 			.BorderImage(
 				UE::Editor::ContentBrowser::IsNewStyleEnabled()
@@ -322,15 +330,27 @@ public:
 
 	virtual void OnOpening() override
 	{
-		TSharedPtr<SAssetViewItem> AssetViewItemPin = AssetViewItem.Pin();
-		if (AssetViewItemPin.IsValid())
+		if (const TSharedPtr<SAssetViewItem> AssetViewItemPinned = AssetViewItem.Pin())
 		{
-			SetContentWidget(AssetViewItemPin->CreateToolTipWidget());
+			if (AssetViewItemPinned->ThumbnailWidgetITooltip.IsValid())
+			{
+				// Update the Thumbnail tooltip
+				AssetViewItemPinned->ThumbnailWidgetITooltip->OnOpening();
+			}
+			SetContentWidget(AssetViewItemPinned->CreateToolTipWidget());
 		}
 	}
 
 	virtual void OnClosed() override
 	{
+		if (const TSharedPtr<SAssetViewItem> AssetViewItemPinned = AssetViewItem.Pin())
+		{
+			if (AssetViewItemPinned->ThumbnailWidgetITooltip.IsValid())
+			{
+				// Reset the content of the Thumbnail tooltip
+				AssetViewItemPinned->ThumbnailWidgetITooltip->ResetContentWidget();
+			}
+		}
 		ResetContentWidget();
 	}
 
@@ -806,7 +826,6 @@ EVisibility SAssetViewItem::GetThumbnailEditModeUIVisibility() const
 	return !IsFolder() && ThumbnailEditMode.Get() ? EVisibility::Visible : EVisibility::Collapsed;
 }
 
-
 TSharedRef<SWidget> SAssetViewItem::CreateToolTipWidget() const
 {
 	if ( AssetItem.IsValid() )
@@ -828,6 +847,12 @@ TSharedRef<SWidget> SAssetViewItem::CreateToolTipWidget() const
 		// TODO: Remove this special caseness so that folders can also have visible attributes
 		if(AssetItem->IsFile())
 		{
+			// New style tooltip based on the AssetThumbnail one, won't be called in the old version as it is never set
+			if (ThumbnailWidgetITooltip.IsValid())
+			{
+				return ThumbnailWidgetITooltip->GetContentWidget();
+			}
+
 			// The tooltip contains the name, class, path, asset registry tags and source control status
 			const FText NameText = GetNameText();
 			const FText ClassText = FText::Format(LOCTEXT("ClassName", "({0})"), GetAssetClassText());
@@ -1924,8 +1949,10 @@ void SAssetListItem::Construct( const FArguments& InArgs )
 			// Use the same tooltip as the Thumbnail
 			if (const TSharedPtr<IToolTip>& ThumbnailTooltip = ThumbnailWidget->GetToolTip())
 			{
-				SetToolTip(ThumbnailTooltip);
+				ThumbnailWidgetITooltip = ThumbnailTooltip;
 			}
+			// reset the thumbnail tooltip to avoid calling that, we handle our tooltip
+			ThumbnailWidget->SetToolTip(TSharedPtr<IToolTip>());
 		}
 	}
 	else
@@ -2480,8 +2507,10 @@ void SAssetTileItem::Construct( const FArguments& InArgs )
 			// Use the same tooltip as the Thumbnail
 			if (const TSharedPtr<IToolTip>& ThumbnailTooltip = Thumbnail->GetToolTip())
 			{
-				SetToolTip(ThumbnailTooltip);
+				ThumbnailWidgetITooltip = ThumbnailTooltip;
 			}
+			// reset the thumbnail tooltip to avoid calling that, we handle our tooltip
+			Thumbnail->SetToolTip(TSharedPtr<IToolTip>());
 		}
 	}
 	else
diff --git a/Engine/Source/Editor/ContentBrowser/Public/AssetViewWidgets.h b/Engine/Source/Editor/ContentBrowser/Public/AssetViewWidgets.h
--- a/Engine/Source/Editor/ContentBrowser/Public/AssetViewWidgets.h
+++ b/Engine/Source/Editor/ContentBrowser/Public/AssetViewWidgets.h
@@ -355,6 +355,9 @@ protected:
 	/** Widget for the source control state */
 	TSharedPtr<SLayeredImage> SCCStateWidget;
 
+	/** Thumbnail tooltip widget used for asset if no given custom widget is passed */
+	TSharedPtr<IToolTip> ThumbnailWidgetITooltip;
+
 	/** Whether the item has a valid scc state brush (not empty) */
 	bool bHasCCStateBrush = false;
 
