Build Engine Localization: Localization warnings in AssetSystemContentBrowserInfoProvider

diff --git a/Engine/Source/Editor/ContentBrowser/Private/AssetSystemContentBrowserInfoProvider.cpp b/Engine/Source/Editor/ContentBrowser/Private/AssetSystemContentBrowserInfoProvider.cpp
--- a/Engine/Source/Editor/ContentBrowser/Private/AssetSystemContentBrowserInfoProvider.cpp
+++ b/Engine/Source/Editor/ContentBrowser/Private/AssetSystemContentBrowserInfoProvider.cpp
@@ -242,7 +242,7 @@ void FAssetSystemContentBrowserInfoProvider::PopulateAssetInfo(TArray<FAssetDisp
 							bAddedCollectionHeader = true;
 							FAssetDisplayInfo CollectionHeader = FAssetDisplayInfo();
 							CollectionHeader.StatusTitle = LOCTEXT("CollectionHeaderTitle", "Collection(s)");
-							CollectionHeader.StatusDescription = LOCTEXT("CollectionHeaderDescription", "");
+							CollectionHeader.StatusDescription = FText::GetEmpty();
 							OutAssetDisplayInfo.Add(CollectionHeader);
 						}
 
