[ContentBrowser] CustomColumns in the CB do not show the values

diff --git a/Engine/Source/Editor/ContentBrowser/Private/SAssetView.cpp b/Engine/Source/Editor/ContentBrowser/Private/SAssetView.cpp
--- a/Engine/Source/Editor/ContentBrowser/Private/SAssetView.cpp
+++ b/Engine/Source/Editor/ContentBrowser/Private/SAssetView.cpp
@@ -5056,6 +5056,12 @@ TSharedRef<ITableRow> SAssetView::MakeListViewWidget(TSharedPtr<FAssetViewItem>
 		return SNew( STableRow<TSharedPtr<FAssetViewItem>>, OwnerTable );
 	}
 
+	if (UE::Editor::ContentBrowser::IsNewStyleEnabled() && CurrentViewType == EAssetViewType::Column)
+	{
+		// Update the cached custom data
+		AssetItem->CacheCustomColumns(CustomColumns, /* bUpdateSortData */ false, /* bUpdateDisplayText */ true, /* bUpdateExisting */ false);
+	}
+
 	VisibleItems.Add(AssetItem);
 	bPendingUpdateThumbnails = true;
 	const FTableRowStyle& ListViewStyle = UE::Editor::ContentBrowser::IsNewStyleEnabled()
