[ContentBrowser] Add ColorStrip at small size

diff --git a/Engine/Source/Editor/ContentBrowser/Private/AssetViewWidgets.cpp b/Engine/Source/Editor/ContentBrowser/Private/AssetViewWidgets.cpp
--- a/Engine/Source/Editor/ContentBrowser/Private/AssetViewWidgets.cpp
+++ b/Engine/Source/Editor/ContentBrowser/Private/AssetViewWidgets.cpp
@@ -1918,7 +1918,7 @@ void SAssetListItem::Construct( const FArguments& InArgs )
 			ThumbnailConfig.GenericThumbnailSize = MakeAttributeSP(this, &SAssetListItem::GetGenericThumbnailSize);
 			ThumbnailConfig.AssetSystemInfoProvider = MakeShared<FAssetSystemContentBrowserInfoProvider>(AssetItem);
 			ThumbnailConfig.AllowAssetStatusThumbnailOverlay = MakeAttributeSP(this, &SAssetListItem::ShouldAllowAssetStatusOverlay, ThumbnailConfig.bForceGenericThumbnail);
-			ThumbnailConfig.ShowAssetColor = MakeAttributeSP(this, &SAssetListItem::ShouldUseAssetColor);
+			ThumbnailConfig.ShowAssetColor = true;
 			ThumbnailConfig.AlwaysExpandTooltip = MakeAttributeSP(this, &SAssetListItem::ShouldExpandTooltip);
 		}
 
@@ -2209,11 +2209,6 @@ bool SAssetListItem::ShouldAllowAssetActionOverlayIndicator(bool bForceGenericTh
 	return !bForceGenericThumbnail && CurrentThumbnailSize.Get() >= EThumbnailSize::Small;
 }
 
-bool SAssetListItem::ShouldUseAssetColor() const
-{
-	return CurrentThumbnailSize.Get() >= EThumbnailSize::Medium;
-}
-
 TSharedRef<SWidget> SAssetListItem::GenerateWidgetForColumnNewStyle(const FName& ColumnName, FIsSelected InIsSelected)
 {
 	TSharedPtr<SWidget> Content;
diff --git a/Engine/Source/Editor/ContentBrowser/Public/AssetViewWidgets.h b/Engine/Source/Editor/ContentBrowser/Public/AssetViewWidgets.h
--- a/Engine/Source/Editor/ContentBrowser/Public/AssetViewWidgets.h
+++ b/Engine/Source/Editor/ContentBrowser/Public/AssetViewWidgets.h
@@ -485,9 +485,6 @@ protected:
 	/** Get if the AssetAction indicator should be visible */
 	bool ShouldAllowAssetActionOverlayIndicator(bool bForceGenericThumbnail) const;
 
-	/** Get if the Asset Color should be visible */
-	bool ShouldUseAssetColor() const;
-
 private:
 	/** New Style columns generation */
 	TSharedRef<SWidget> GenerateWidgetForColumnNewStyle(const FName& ColumnName, FIsSelected InIsSelected);
