[ContentBrowser] Size get resetted to Medium if changing the ThumbnailSize closing and re-opening the CB/Picker

diff --git a/Engine/Source/Editor/ContentBrowser/Private/SAssetView.cpp b/Engine/Source/Editor/ContentBrowser/Private/SAssetView.cpp
--- a/Engine/Source/Editor/ContentBrowser/Private/SAssetView.cpp
+++ b/Engine/Source/Editor/ContentBrowser/Private/SAssetView.cpp
@@ -2083,6 +2083,10 @@ void SAssetView::LoadSettings(const FString& IniFilename, const FString& IniSect
 		SetCurrentViewType((EAssetViewType::Type)ViewType);
 	}
 
+	// Update the size value after loading the config of the CurrentViewType and the sizes
+	// Since if the View was the same as before it won't get called during SetCurrentViewType
+	UpdateThumbnailSizeValue();
+
 	float Zoom = 0;
 	if ( GConfig->GetFloat(*IniSection, *(SettingsString + TEXT(".ZoomScale")), Zoom, IniFilename) )
 	{
