Added support for pasting Verse paths in the Content Browser address bar for navigating to their corresponding asset or module.

diff --git a/Engine/Source/Editor/ContentBrowser/Private/ContentBrowserUtils.cpp b/Engine/Source/Editor/ContentBrowser/Private/ContentBrowserUtils.cpp
--- a/Engine/Source/Editor/ContentBrowser/Private/ContentBrowserUtils.cpp
+++ b/Engine/Source/Editor/ContentBrowser/Private/ContentBrowserUtils.cpp
@@ -9,6 +9,7 @@
 #include "CollectionViewUtils.h"
 #include "Containers/Map.h"
 #include "Containers/StringView.h"
+#include "Containers/VersePath.h"
 #include "ContentBrowserConfig.h"
 #include "ContentBrowserDataFilter.h"
 #include "ContentBrowserDataSource.h"
@@ -27,11 +28,13 @@
 #include "HAL/PlatformCrt.h"
 #include "HAL/PlatformMisc.h"
 #include "HAL/PlatformProcess.h"
+#include "IAssetTools.h"
 #include "ICollectionContainer.h"
 #include "ICollectionManager.h"
 #include "ICollectionSource.h"
 #include "IContentBrowserDataModule.h"
 #include "Input/Reply.h"
+#include "Interfaces/IPluginManager.h"
 #include "Layout/Children.h"
 #include "Layout/Margin.h"
 #include "Layout/SlateRect.h"
@@ -1214,7 +1217,7 @@ FContentBrowserItem ContentBrowserUtils::TryGetItemFromUserProvidedPath(FStringV
 	FStringView ObjectPathView = RequestedPathView;
 	if (FPackageName::IsValidObjectPath(ObjectPathView) || FPackageName::ParseExportTextPath(RequestedPathView, nullptr, &ObjectPathView))
 	{
-		VirtualPath = ContentBrowserData->ConvertInternalPathToVirtual(FName(ObjectPathView));
+		ContentBrowserData->ConvertInternalPathToVirtual(ObjectPathView, VirtualPath);
 		Item = ContentBrowserData->GetItemAtPath(VirtualPath, EContentBrowserItemTypeFilter::IncludeFiles);
 		if (Item.IsValid())
 		{
@@ -1222,11 +1225,37 @@ FContentBrowserItem ContentBrowserUtils::TryGetItemFromUserProvidedPath(FStringV
 		}
 	}
 
+	const IAssetTools& AssetTools = IAssetTools::Get();
+	UE::Core::FVersePath VersePath;
+	if (AssetTools.ShowingContentVersePath())
+	{
+		UE::Core::FVersePath::TryMake(VersePath, FString(RequestedPathView));
+	}
+
+	// If the string is a complete Verse path to an asset, sync to that asset
+	if (VersePath.IsValid())
+	{
+		FAssetData AssetData = AssetTools.FindAssetByVersePath(VersePath);
+		if (AssetData.IsValid())
+		{
+			FNameBuilder NameBuilder;
+			AssetData.GetSoftObjectPath().ToString(NameBuilder);
+
+			ContentBrowserData->ConvertInternalPathToVirtual(NameBuilder, VirtualPath);
+			Item = ContentBrowserData->GetItemAtPath(VirtualPath, EContentBrowserItemTypeFilter::IncludeFiles);
+			if (Item.IsValid())
+			{
+				return Item;
+			}
+		}
+	}
+
 	auto GetItemFromPackageName = [ContentBrowserData](const FString& PackageName) -> FContentBrowserItem {
 		// Packages like /Game/Characters/Knight do not map to virtual paths in data source, assets like /Game/Characters/Knight.Knight do.
 		// See if there's an item if we duplicate the last part of the path 
-		FName InternalPath(TStringBuilder<320>(InPlace, PackageName, TEXTVIEW("."), FPackageName::GetShortName(PackageName)));
-		FName VirtualPath = ContentBrowserData->ConvertInternalPathToVirtual(InternalPath);
+		TStringBuilder<320> InternalPath(InPlace, PackageName, TEXTVIEW("."), FPackageName::GetShortName(PackageName));
+		FName VirtualPath;
+		ContentBrowserData->ConvertInternalPathToVirtual(InternalPath, VirtualPath);
 		FContentBrowserItem Item = ContentBrowserData->GetItemAtPath(VirtualPath, EContentBrowserItemTypeFilter::IncludeFiles);
 		if (Item.IsValid())
 		{
@@ -1235,7 +1264,7 @@ FContentBrowserItem ContentBrowserUtils::TryGetItemFromUserProvidedPath(FStringV
 
 		// Otherwise go up to the package path and enumerate items to see if there's an asset with the desired package name 
 		FString PackagePath = FPackageName::GetLongPackagePath(PackageName);
-		VirtualPath = ContentBrowserData->ConvertInternalPathToVirtual(FName(PackagePath));
+		ContentBrowserData->ConvertInternalPathToVirtual(PackagePath, VirtualPath);
 		FContentBrowserDataFilter Filter;
 		Filter.bRecursivePaths = false;
 		Filter.ItemTypeFilter = EContentBrowserItemTypeFilter::IncludeFiles;
@@ -1286,6 +1315,37 @@ FContentBrowserItem ContentBrowserUtils::TryGetItemFromUserProvidedPath(FStringV
 		}
 	}
 
+	// If the string is a complete Verse path to a module, sync to that folder
+	if (VersePath.IsValid())
+	{
+		IPluginManager::Get().EnumerateVersePathMountPoints(VersePath, [ContentBrowserData, &VirtualPath, &Item](FStringView MountedAssetPath, FStringView LeafPath)
+		{
+			FNameBuilder NameBuilder;
+			NameBuilder.Append(MountedAssetPath);
+			if (ensure(NameBuilder.Len() > 0 && NameBuilder.LastChar() == TEXT('/')))
+			{
+				if (LeafPath.IsEmpty())
+				{
+					// ContentBrowserData expects no trailing /.
+					NameBuilder.RemoveSuffix(1);
+				}
+				else
+				{
+					NameBuilder.Append(LeafPath);
+				}
+
+				ContentBrowserData->ConvertInternalPathToVirtual(NameBuilder, VirtualPath);
+				Item = ContentBrowserData->GetItemAtPath(VirtualPath, EContentBrowserItemTypeFilter::IncludeFolders);
+			}
+			return !Item.IsValid();
+		});
+
+		if (Item.IsValid())
+		{
+			return Item;
+		}
+	}
+
 	// Try and remove elements from the end of the path until it's a valid virtual path 
 	FPathViews::IterateAncestors(RequestedPathView, [RequestedPathView, ContentBrowserData, &Item](FStringView InAncestor){
 		if (RequestedPathView == InAncestor)
