Convert each path segment of a Content Browser virtual path into its display name when showing virtual paths so they match what is displayed in the navigation bar.
Don't add virtual paths that no longer exist to the Content Browser navigation bar.
Fixed AssetDialog hotkeys either not working or operating on the wrong widget.

diff --git a/Engine/Plugins/Editor/ContentBrowser/ContentBrowserFileDataSource/Source/ContentBrowserFileDataSource/Private/ContentBrowserFileDataSource.cpp b/Engine/Plugins/Editor/ContentBrowser/ContentBrowserFileDataSource/Source/ContentBrowserFileDataSource/Private/ContentBrowserFileDataSource.cpp
--- a/Engine/Plugins/Editor/ContentBrowser/ContentBrowserFileDataSource/Source/ContentBrowserFileDataSource/Private/ContentBrowserFileDataSource.cpp
+++ b/Engine/Plugins/Editor/ContentBrowser/ContentBrowserFileDataSource/Source/ContentBrowserFileDataSource/Private/ContentBrowserFileDataSource.cpp
@@ -1379,14 +1379,14 @@ bool UContentBrowserFileDataSource::CanRenameItem(const FContentBrowserItemData&
 	// We cannot rename the file mounts themselves
 	if (IsRootFileMount(InternalPath))
 	{
-		ContentBrowserFileData::SetOptionalErrorMessage(OutErrorMsg, FText::Format(LOCTEXT("Error_FolderIsFileRoot", "Folder '{0}' is a root file path and cannot be renamed"), FText::FromName(InItem.GetVirtualPath())));
+		ContentBrowserFileData::SetOptionalErrorMessage(OutErrorMsg, FText::Format(LOCTEXT("Error_FolderIsFileRoot", "Folder '{0}' is a root file path and cannot be renamed"), FText::FromName(InternalPath)));
 		return false;
 	}
 
 	// We cannot rename things outside of our file mounts
 	if (!IsKnownFileMount(InternalPath))
 	{
-		ContentBrowserFileData::SetOptionalErrorMessage(OutErrorMsg, FText::Format(LOCTEXT("Error_FolderIsNotFilePath", "Folder '{0}' is not a known file path"), FText::FromName(InItem.GetVirtualPath())));
+		ContentBrowserFileData::SetOptionalErrorMessage(OutErrorMsg, FText::Format(LOCTEXT("Error_FolderIsNotFilePath", "Folder '{0}' is not a known file path"), FText::FromName(InternalPath)));
 		return false;
 	}
 
diff --git a/Engine/Source/Editor/ContentBrowser/Private/AssetSystemContentBrowserInfoProvider.cpp b/Engine/Source/Editor/ContentBrowser/Private/AssetSystemContentBrowserInfoProvider.cpp
--- a/Engine/Source/Editor/ContentBrowser/Private/AssetSystemContentBrowserInfoProvider.cpp
+++ b/Engine/Source/Editor/ContentBrowser/Private/AssetSystemContentBrowserInfoProvider.cpp
@@ -15,10 +15,12 @@
 #include "CollectionManagerModule.h"
 #include "Containers/VersePath.h"
 #include "ContentBrowserDataSource.h"
+#include "ContentBrowserDataSubsystem.h"
 #include "ContentBrowserUtils.h"
 #include "IAssetTools.h"
 #include "ICollectionContainer.h"
 #include "ICollectionManager.h"
+#include "IContentBrowserDataModule.h"
 #include "Misc/EngineBuildSettings.h"
 
 #define LOCTEXT_NAMESPACE "AssetSystemContentBrowserInfoProvider"
@@ -77,7 +79,7 @@ void FAssetSystemContentBrowserInfoProvider::PopulateAssetInfo(TArray<FAssetDisp
 		}
 		else
 		{
-			PathInfo.StatusDescription = FText::FromName(AssetItem->GetItem().GetVirtualPath());
+			PathInfo.StatusDescription = IContentBrowserDataModule::Get().GetSubsystem()->ConvertVirtualPathToDisplay(AssetItem->GetItem());
 		}
 		OutAssetDisplayInfo.Add(PathInfo);
 
diff --git a/Engine/Source/Editor/ContentBrowser/Private/AssetViewWidgets.cpp b/Engine/Source/Editor/ContentBrowser/Private/AssetViewWidgets.cpp
--- a/Engine/Source/Editor/ContentBrowser/Private/AssetViewWidgets.cpp
+++ b/Engine/Source/Editor/ContentBrowser/Private/AssetViewWidgets.cpp
@@ -24,6 +24,7 @@
 #include "Containers/UnrealString.h"
 #include "Containers/VersePath.h"
 #include "ContentBrowserDataSource.h"
+#include "ContentBrowserDataSubsystem.h"
 #include "ContentBrowserItem.h"
 #include "ContentBrowserItemData.h"
 #include "ContentBrowserModule.h"
@@ -40,6 +41,7 @@
 #include "IAssetTools.h"
 #include "ICollectionContainer.h"
 #include "ICollectionManager.h"
+#include "IContentBrowserDataModule.h"
 #include "ISourceControlModule.h"
 #include "ISourceControlProvider.h"
 #include "ISourceControlState.h"
@@ -874,7 +876,7 @@ TSharedRef<SWidget> SAssetViewItem::CreateToolTipWidget() const
 			}
 			else
 			{
-				AddToToolTipInfoBox(InfoBox, LOCTEXT("TileViewTooltipPath", "Path"), FText::FromName(AssetItem->GetItem().GetVirtualPath()), false);
+				AddToToolTipInfoBox(InfoBox, LOCTEXT("TileViewTooltipPath", "Path"), IContentBrowserDataModule::Get().GetSubsystem()->ConvertVirtualPathToDisplay(AssetItem->GetItem()), false);
 			}
 
 			if (ItemAssetData.IsValid() && FAssetToolsModule::GetModule().Get().ShowingContentVersePath())
@@ -1193,7 +1195,7 @@ TSharedRef<SWidget> SAssetViewItem::CreateToolTipWidget() const
 		else
 		{
 			const FText FolderName = GetNameText();
-			const FText FolderPath = FText::FromName(AssetItem->GetItem().GetVirtualPath());
+			const FText FolderPath = IContentBrowserDataModule::Get().GetSubsystem()->ConvertVirtualPathToDisplay(AssetItem->GetItem());
 
 			// Create a box to hold every line of info in the body of the tooltip
 			TSharedRef<SVerticalBox> InfoBox = SNew(SVerticalBox);
diff --git a/Engine/Source/Editor/ContentBrowser/Private/NewAssetOrClassContextMenu.cpp b/Engine/Source/Editor/ContentBrowser/Private/NewAssetOrClassContextMenu.cpp
--- a/Engine/Source/Editor/ContentBrowser/Private/NewAssetOrClassContextMenu.cpp
+++ b/Engine/Source/Editor/ContentBrowser/Private/NewAssetOrClassContextMenu.cpp
@@ -99,11 +99,15 @@ void FNewAssetOrClassContextMenu::MakeContextMenu(
 			{
 				if(bIsValidNewFolderPath)
 				{
-					NewFolderToolTip = FText::Format(LOCTEXT("NewFolderTooltip_CreateIn", "Create a new folder in {0}."), FText::FromName(FirstSelectedPath));
+					NewFolderToolTip = FText::Format(
+						LOCTEXT("NewFolderTooltip_CreateIn", "Create a new folder in {0}."),
+						IContentBrowserDataModule::Get().GetSubsystem()->ConvertVirtualPathToDisplay(FirstSelectedPath, EContentBrowserItemTypeFilter::IncludeFolders));
 				}
 				else
 				{
-					NewFolderToolTip = FText::Format(LOCTEXT("NewFolderTooltip_InvalidPath", "Cannot create new folders in {0}."), FText::FromName(FirstSelectedPath));
+					NewFolderToolTip = FText::Format(
+						LOCTEXT("NewFolderTooltip_InvalidPath", "Cannot create new folders in {0}."),
+						IContentBrowserDataModule::Get().GetSubsystem()->ConvertVirtualPathToDisplay(FirstSelectedPath, EContentBrowserItemTypeFilter::IncludeFolders));
 				}
 			}
 			else
diff --git a/Engine/Source/Editor/ContentBrowser/Private/PathContextMenu.cpp b/Engine/Source/Editor/ContentBrowser/Private/PathContextMenu.cpp
--- a/Engine/Source/Editor/ContentBrowser/Private/PathContextMenu.cpp
+++ b/Engine/Source/Editor/ContentBrowser/Private/PathContextMenu.cpp
@@ -126,7 +126,7 @@ void FPathContextMenu::MakePathViewContextMenu(UToolMenu* Menu)
 				{
 					if(Context->bCanBeModified)
 					{
-						NewAssetToolTip = FText::Format(LOCTEXT("NewAssetTooltip_CreateIn", "Create a new item in {0}."), FText::FromName(SelectedFolders[0].GetVirtualPath()));
+						NewAssetToolTip = FText::Format(LOCTEXT("NewAssetTooltip_CreateIn", "Create a new item in {0}."), IContentBrowserDataModule::Get().GetSubsystem()->ConvertVirtualPathToDisplay(SelectedFolders[0]));
 					}
 				}
 				else
diff --git a/Engine/Source/Editor/ContentBrowser/Private/SAssetDialog.cpp b/Engine/Source/Editor/ContentBrowser/Private/SAssetDialog.cpp
--- a/Engine/Source/Editor/ContentBrowser/Private/SAssetDialog.cpp
+++ b/Engine/Source/Editor/ContentBrowser/Private/SAssetDialog.cpp
@@ -88,7 +88,6 @@ SAssetDialog::SAssetDialog()
 	, bLastInputValidityCheckSuccessful(false)
 	, bPendingFocusNextFrame(true)
 	, bValidAssetsChosen(false)
-	, OpenedContextMenuWidget(EOpenedContextMenuWidget::None)
 {
 }
 
@@ -121,7 +120,7 @@ void SAssetDialog::Construct(const FArguments& InArgs, const FSharedAssetDialogC
 	PathPickerConfig.bFocusSearchBoxWhenOpened = false;
 	PathPickerConfig.OnPathSelected = FOnPathSelected::CreateSP(this, &SAssetDialog::SetCurrentlySelectedPath, EContentBrowserPathType::Virtual);
 	PathPickerConfig.SetPathsDelegates.Add(&SetPathsDelegate);
-	PathPickerConfig.OnGetFolderContextMenu = FOnGetFolderContextMenu::CreateSP(this, &SAssetDialog::OnGetFolderContextMenu);	
+	PathPickerConfig.OnGetFolderContextMenu = FOnGetFolderContextMenu::CreateSP(this, &SAssetDialog::OnGetFolderContextMenu, EAssetDialogCommandContext::PathView);
 	PathPickerConfig.bOnPathSelectedPassesVirtualPaths = true;
 
 	FAssetPickerConfig AssetPickerConfig;
@@ -137,7 +136,7 @@ void SAssetDialog::Construct(const FArguments& InArgs, const FSharedAssetDialogC
 	AssetPickerConfig.bCanShowDevelopersFolder = true;
 	AssetPickerConfig.OnFolderEntered = FOnPathSelected::CreateSP(this, &SAssetDialog::SetCurrentlySelectedPath, EContentBrowserPathType::Virtual);
 	AssetPickerConfig.OnGetAssetContextMenu = FOnGetAssetContextMenu::CreateSP(this, &SAssetDialog::OnGetAssetContextMenu);
-	AssetPickerConfig.OnGetFolderContextMenu = FOnGetFolderContextMenu::CreateSP(this, &SAssetDialog::OnGetFolderContextMenu);
+	AssetPickerConfig.OnGetFolderContextMenu = FOnGetFolderContextMenu::CreateSP(this, &SAssetDialog::OnGetFolderContextMenu, EAssetDialogCommandContext::AssetView);
 
 	OnPathSelected = InConfig.OnPathSelected;
 
@@ -440,9 +439,22 @@ FReply SAssetDialog::OnKeyDown( const FGeometry& MyGeometry, const FKeyEvent& In
 		CloseDialog();
 		return FReply::Handled();
 	}
-	else if (Commands->ProcessCommandBindings(InKeyEvent))
+	else
 	{
-		return FReply::Handled();
+		if (FSlateApplication::Get().HasFocusedDescendants(PathPicker->GetPathView().ToSharedRef()))
+		{
+			if (PathViewCommands->ProcessCommandBindings(InKeyEvent))
+			{
+				return FReply::Handled();
+			}
+		}
+		else if (FSlateApplication::Get().HasFocusedDescendants(AssetPicker->GetAssetView().ToSharedRef()))
+		{
+			if (AssetViewCommands->ProcessCommandBindings(InKeyEvent))
+			{
+				return FReply::Handled();
+			}
+		}
 	}
 
 	return SCompoundWidget::OnKeyDown(MyGeometry, InKeyEvent);
@@ -484,67 +496,78 @@ FReply SAssetDialog::OnMouseButtonDoubleClick(const FGeometry& InMyGeometry, con
 
 void SAssetDialog::BindCommands()
 {
-	Commands = TSharedPtr< FUICommandList >(new FUICommandList);
-
-	Commands->MapAction(FGenericCommands::Get().Rename, FUIAction(
-		FExecuteAction::CreateSP(this, &SAssetDialog::ExecuteRename),
-		FCanExecuteAction::CreateSP(this, &SAssetDialog::CanExecuteRename)
-	));
-
-	Commands->MapAction(FGenericCommands::Get().Delete, FUIAction(
-		FExecuteAction::CreateSP(this, &SAssetDialog::ExecuteDelete),
-		FCanExecuteAction::CreateSP(this, &SAssetDialog::CanExecuteDelete)
-	));
-
-	Commands->MapAction(FContentBrowserCommands::Get().CreateNewFolder, FUIAction(
-		FExecuteAction::CreateSP(this, &SAssetDialog::ExecuteCreateNewFolder),
-		FCanExecuteAction::CreateSP(this, &SAssetDialog::CanExecuteCreateNewFolder)
-	));
+	auto MapActions = [this](EAssetDialogCommandContext CommandContext, TSharedPtr<FUICommandList>& OutCommands)
+	{
+		OutCommands = MakeShared<FUICommandList>();
+
+		OutCommands->MapAction(FGenericCommands::Get().Rename, FUIAction(
+			FExecuteAction::CreateSP(this, &SAssetDialog::ExecuteRename, CommandContext),
+			FCanExecuteAction::CreateSP(this, &SAssetDialog::CanExecuteRename, CommandContext)
+		));
+
+		OutCommands->MapAction(FGenericCommands::Get().Delete, FUIAction(
+			FExecuteAction::CreateSP(this, &SAssetDialog::ExecuteDelete, CommandContext),
+			FCanExecuteAction::CreateSP(this, &SAssetDialog::CanExecuteDelete, CommandContext)
+		));
+
+		OutCommands->MapAction(FContentBrowserCommands::Get().CreateNewFolder, FUIAction(
+			FExecuteAction::CreateSP(this, &SAssetDialog::ExecuteCreateNewFolder, CommandContext),
+			FCanExecuteAction::CreateSP(this, &SAssetDialog::CanExecuteCreateNewFolder, CommandContext)
+		));
+	};
+
+	MapActions(EAssetDialogCommandContext::AssetView, AssetViewCommands);
+	MapActions(EAssetDialogCommandContext::PathView, PathViewCommands);
 }
 
-bool SAssetDialog::CanExecuteRename() const
+bool SAssetDialog::CanExecuteRename(EAssetDialogCommandContext CommandContext) const
 {
-	switch (OpenedContextMenuWidget)
+	switch (CommandContext)
 	{
-		case EOpenedContextMenuWidget::AssetView: return ContentBrowserUtils::CanRenameFromAssetView(AssetPicker->GetAssetView());
-		case EOpenedContextMenuWidget::PathView: return ContentBrowserUtils::CanRenameFromPathView(PathPicker->GetPathView());
+		case EAssetDialogCommandContext::AssetView: return ContentBrowserUtils::CanRenameFromAssetView(AssetPicker->GetAssetView());
+		case EAssetDialogCommandContext::PathView: return ContentBrowserUtils::CanRenameFromPathView(PathPicker->GetPathView());
 	}
 
 	return false;
 }
 
-void SAssetDialog::ExecuteRename()
+void SAssetDialog::ExecuteRename(EAssetDialogCommandContext CommandContext)
 {
-	const TArray<FContentBrowserItem> SelectedItems = AssetPicker->GetAssetView()->GetSelectedItems();
-	if (SelectedItems.Num() > 0)
+	switch (CommandContext)
 	{
+	case EAssetDialogCommandContext::AssetView:
+	{
+		const TArray<FContentBrowserItem> SelectedItems = AssetPicker->GetAssetView()->GetSelectedItems();
 		if (SelectedItems.Num() == 1)
 		{
 			AssetPicker->GetAssetView()->RenameItem(SelectedItems[0]);
 		}
+		break;
 	}
-	else
-	{		
+	case EAssetDialogCommandContext::PathView:
+	{
 		const TArray<FContentBrowserItem> SelectedFolders = PathPicker->GetPathView()->GetSelectedFolderItems();
 		if (SelectedFolders.Num() == 1)
 		{
 			PathPicker->GetPathView()->RenameFolderItem(SelectedFolders[0]);
 		}
+		break;
+	}
 	}
 }
 
-bool SAssetDialog::CanExecuteDelete() const
+bool SAssetDialog::CanExecuteDelete(EAssetDialogCommandContext CommandContext) const
 {
-	switch (OpenedContextMenuWidget)
+	switch (CommandContext)
 	{
-		case EOpenedContextMenuWidget::AssetView: return ContentBrowserUtils::CanDeleteFromAssetView(AssetPicker->GetAssetView());
-		case EOpenedContextMenuWidget::PathView: return ContentBrowserUtils::CanDeleteFromPathView(PathPicker->GetPathView());
+		case EAssetDialogCommandContext::AssetView: return ContentBrowserUtils::CanDeleteFromAssetView(AssetPicker->GetAssetView());
+		case EAssetDialogCommandContext::PathView: return ContentBrowserUtils::CanDeleteFromPathView(PathPicker->GetPathView());
 	}
 
 	return false;
 }
 
-void SAssetDialog::ExecuteDelete()
+void SAssetDialog::ExecuteDelete(EAssetDialogCommandContext CommandContext)
 {
 	// Don't allow asset deletion during PIE
 	if (GIsEditor)
@@ -560,7 +583,7 @@ void SAssetDialog::ExecuteDelete()
 		}
 	}
 
-	if (OpenedContextMenuWidget != EOpenedContextMenuWidget::PathView)
+	if (CommandContext != EAssetDialogCommandContext::PathView)
 	{
 		const TArray<FContentBrowserItem> SelectedFiles = AssetPicker->GetAssetView()->GetSelectedFileItems();
 
@@ -598,7 +621,7 @@ void SAssetDialog::ExecuteDelete()
 	FText FirstFolderDisplayName;
 	TArray<FString> SelectedFolderInternalPaths;
 	{
-		const TArray<FContentBrowserItem> SelectedFolderItems = (OpenedContextMenuWidget == EOpenedContextMenuWidget::PathView) ?
+		const TArray<FContentBrowserItem> SelectedFolderItems = (CommandContext == EAssetDialogCommandContext::PathView) ?
 			PathPicker->GetPathView()->GetSelectedFolderItems() :
 			AssetPicker->GetAssetView()->GetSelectedFolderItems();
 
@@ -635,7 +658,7 @@ void SAssetDialog::ExecuteDelete()
 			Prompt = FText::Format(LOCTEXT("FolderDeleteConfirm_Multiple", "Delete {0} folders?"), SelectedFolderInternalPaths.Num());
 		}
 
-		const bool bResetSelection = OpenedContextMenuWidget == EOpenedContextMenuWidget::PathView;
+		const bool bResetSelection = CommandContext == EAssetDialogCommandContext::PathView;
 
 		// Spawn a confirmation dialog since this is potentially a highly destructive operation
 		ContentBrowserUtils::DisplayConfirmationPopup(
@@ -682,73 +705,114 @@ void SAssetDialog::SelectDefaultPaths()
 	AssetPicker->GetAssetView()->SetContentSources(DefaultContentSources);
 }
 
-bool SAssetDialog::ExecuteExploreInternal(bool bTest)
+void SAssetDialog::ExecuteExplore(EAssetDialogCommandContext Widget)
 {
-	bool bCanExplore = false;
+	TArray<FContentBrowserItem> SelectedItems;
+	TSharedPtr<SWidget> ParentContent;
+	switch (Widget)
+	{
+	case EAssetDialogCommandContext::AssetView:
+		SelectedItems = AssetPicker->GetAssetView()->GetSelectedItems();
+		ParentContent = AssetPicker->GetAssetView();
+		break;
+	case EAssetDialogCommandContext::PathView:
+		SelectedItems = PathPicker->GetPathView()->GetSelectedFolderItems();
+		ParentContent = PathPicker->GetPathView();
+		break;
+	default:
+		return;
+	}
+	ContentBrowserUtils::ExploreFolders(SelectedItems, ParentContent.ToSharedRef());
+}
 
-	TArray<FContentBrowserItem> SelectedItems = AssetPicker->GetAssetView()->GetSelectedItems();
-	if (SelectedItems.Num() == 0)
+bool SAssetDialog::CanExecuteExplore(EAssetDialogCommandContext Widget)
+{
+	TArray<FContentBrowserItem> SelectedItems;
+	switch (Widget)
 	{
+	case EAssetDialogCommandContext::AssetView:
+		SelectedItems = AssetPicker->GetAssetView()->GetSelectedItems();
+		break;
+	case EAssetDialogCommandContext::PathView:
 		SelectedItems = PathPicker->GetPathView()->GetSelectedFolderItems();
+		break;
+	default:
+		return false;
 	}
+	return ContentBrowserUtils::CanExploreFolders(SelectedItems);
+}
 
-	for (const FContentBrowserItem& SelectedItem : SelectedItems)
+FContentBrowserItem SAssetDialog::GetCreateNewFolderParent(EAssetDialogCommandContext CommandContext) const
+{
+	if (CommandContext == EAssetDialogCommandContext::AssetView)
 	{
-		FString ItemFilename;
-		if (SelectedItem.GetItemPhysicalPath(ItemFilename))
+		TArray<FContentBrowserItem> SelectedItems = AssetPicker->GetAssetView()->GetSelectedItems();
+		if (!SelectedItems.IsEmpty())
 		{
-			const bool bExists = SelectedItem.IsFile() ? FPaths::FileExists(ItemFilename) : FPaths::DirectoryExists(ItemFilename);
-			if (bExists)
+			if (SelectedItems.Num() == 1 && SelectedItems[0].IsFolder())
 			{
-				bCanExplore = true;
-
-				if (!bTest)
-				{
-					FPlatformProcess::ExploreFolder(*IFileManager::Get().ConvertToAbsolutePathForExternalAppForRead(*ItemFilename));
-				}
-				else
-				{
-					break;
-				}
+				return SelectedItems[0];
 			}
+			return {};
 		}
+		// If nothing is selected in the AssetView, use the parent folder from the PathView.
 	}
 
-	return bCanExplore;
-}
+	TArray<FContentBrowserItem> SelectedFolders = PathPicker->GetPathView()->GetSelectedFolderItems();
+	// We can only create folders when we have a single path selected
+	if (SelectedFolders.Num() == 1)
+	{
+		return SelectedFolders[0];
+	}
 
-void SAssetDialog::ExecuteExplore()
-{
-	//ExecuteExploreInternal(/*bTest*/ false);
-	ContentBrowserUtils::ExploreFolders(AssetPicker->GetAssetView()->GetSelectedItems(), AssetPicker->GetAssetView().ToSharedRef());
+	return {};
 }
 
-bool SAssetDialog::CanExecuteExplore()
+bool SAssetDialog::CanExecuteCreateNewFolder(EAssetDialogCommandContext CommandContext) const
 {
-	return ExecuteExploreInternal(/*bTest*/ true);
-}
+	const FContentBrowserItem CreateNewFolderParent = GetCreateNewFolderParent(CommandContext);
+	if (!CreateNewFolderParent.IsValid())
+	{
+		return false;
+	}
 
-bool SAssetDialog::CanExecuteCreateNewFolder() const
-{	
-	// We can only create folders when we have a single path selected
 	UContentBrowserDataSubsystem* ContentBrowserData = IContentBrowserDataModule::Get().GetSubsystem();
-	return ContentBrowserData->CanCreateFolder(GetCurrentSelectedVirtualPath(), nullptr);
+	return ContentBrowserData->CanCreateFolder(CreateNewFolderParent.GetVirtualPath(), nullptr);
 }
 
-void SAssetDialog::ExecuteCreateNewFolder()
+void SAssetDialog::ExecuteCreateNewFolder(EAssetDialogCommandContext CommandContext)
 {
-	PathPicker->CreateNewFolder(GetCurrentSelectedVirtualPath().ToString(), CurrentContextMenuCreateNewFolderDelegate);
-}
+	const FContentBrowserItem CreateNewFolderParent = GetCreateNewFolderParent(CommandContext);
+	if (!CreateNewFolderParent.IsValid())
+	{
+		return;
+	}
 
-TSharedPtr<SWidget> SAssetDialog::OnGetFolderContextMenu(const TArray<FString>& SelectedPaths, FContentBrowserMenuExtender_SelectedPaths InMenuExtender, FOnCreateNewFolder InOnCreateNewFolder)
-{
-	if (FSlateApplication::Get().HasFocusedDescendants(PathPicker.ToSharedRef()))
+	FOnCreateNewFolder OnCreateNewFolder;
+	switch (CommandContext)
 	{
-		OpenedContextMenuWidget = EOpenedContextMenuWidget::PathView;
+	case EAssetDialogCommandContext::AssetView:
+		OnCreateNewFolder = FOnCreateNewFolder::CreateSP(AssetPicker->GetAssetView().ToSharedRef(), &SAssetView::NewFolderItemRequested);
+		break;
+	case EAssetDialogCommandContext::PathView:
+		OnCreateNewFolder = FOnCreateNewFolder::CreateSP(PathPicker->GetPathView().ToSharedRef(), &SPathView::NewFolderItemRequested);
+		break;
 	}
-	else if (FSlateApplication::Get().HasFocusedDescendants(AssetPicker.ToSharedRef()))
+
+	PathPicker->CreateNewFolder(CreateNewFolderParent.GetVirtualPath().ToString(), OnCreateNewFolder);
+}
+
+TSharedPtr<SWidget> SAssetDialog::OnGetFolderContextMenu(const TArray<FString>& SelectedPaths, FContentBrowserMenuExtender_SelectedPaths InMenuExtender, FOnCreateNewFolder InOnCreateNewFolder, EAssetDialogCommandContext CommandContext)
+{
+	TSharedPtr<const FUICommandList> Commands;
+	switch (CommandContext)
 	{
-		OpenedContextMenuWidget = EOpenedContextMenuWidget::AssetView;
+	case EAssetDialogCommandContext::AssetView:
+		Commands = AssetViewCommands;
+		break;
+	case EAssetDialogCommandContext::PathView:
+		Commands = PathViewCommands;
+		break;
 	}
 
 	TSharedPtr<FExtender> Extender;
@@ -757,53 +821,46 @@ TSharedPtr<SWidget> SAssetDialog::OnGetFolderContextMenu(const TArray<FString>&
 		Extender = InMenuExtender.Execute(SelectedPaths);
 	}
 
-	if (FSlateApplication::Get().HasFocusedDescendants(PathPicker.ToSharedRef()))
-	{
-		PathPicker->SetPaths(SelectedPaths);
-	}
-
-	CurrentContextMenuCreateNewFolderDelegate = InOnCreateNewFolder;
-
 	FMenuBuilder MenuBuilder(true /*bInShouldCloseWindowAfterMenuSelection*/, Commands, Extender);
-	SetupContextMenuContent(MenuBuilder, SelectedPaths);
+	SetupContextMenuContent(CommandContext, MenuBuilder);
 
 	return MenuBuilder.MakeWidget();
 }
 
 TSharedPtr<SWidget> SAssetDialog::OnGetAssetContextMenu(const TArray<FAssetData>& SelectedAssets)
 {
-	OpenedContextMenuWidget = EOpenedContextMenuWidget::AssetView;
-
-	FMenuBuilder MenuBuilder(true /*bInShouldCloseWindowAfterMenuSelection*/, Commands);
+	FMenuBuilder MenuBuilder(true /*bInShouldCloseWindowAfterMenuSelection*/, AssetViewCommands);
 
-	CurrentContextMenuCreateNewFolderDelegate = FOnCreateNewFolder::CreateSP(AssetPicker->GetAssetView().Get(), &SAssetView::NewFolderItemRequested);
-
-	TArray<FString> Paths;
-	SetupContextMenuContent(MenuBuilder, Paths);
+	SetupContextMenuContent(EAssetDialogCommandContext::AssetView, MenuBuilder);
 
 	return MenuBuilder.MakeWidget();
 }
 
-void SAssetDialog::SetupContextMenuContent(FMenuBuilder& MenuBuilder, const TArray<FString>& SelectedPaths)
+void SAssetDialog::SetupContextMenuContent(EAssetDialogCommandContext CommandContext, FMenuBuilder& MenuBuilder)
 {
 	const FName ContentBrowserStyleSetName = UE::ContentBrowser::Private::FContentBrowserStyle::Get().GetStyleSetName();
 
 	FText NewFolderToolTip;
 
-	if (SelectedPaths.Num() > 0)
+	const FContentBrowserItem CreateNewFolderParent = GetCreateNewFolderParent(CommandContext);
+	if (CreateNewFolderParent.IsValid())
 	{
-		if (CanExecuteCreateNewFolder())
+		if (CanExecuteCreateNewFolder(CommandContext))
 		{
-			NewFolderToolTip = FText::Format(LOCTEXT("NewFolderTooltip_CreateIn", "Create a new folder in {0}."), FText::FromString(SelectedPaths[0]));
+			NewFolderToolTip = FText::Format(
+				LOCTEXT("NewFolderTooltip_CreateIn", "Create a new folder in {0}."),
+				IContentBrowserDataModule::Get().GetSubsystem()->ConvertVirtualPathToDisplay(CreateNewFolderParent));
 		}
 		else
 		{
-			NewFolderToolTip = FText::Format(LOCTEXT("NewFolderTooltip_InvalidPath", "Cannot create new folders in {0}."), FText::FromString(SelectedPaths[0]));
+			NewFolderToolTip = FText::Format(
+				LOCTEXT("NewFolderTooltip_InvalidPath", "Cannot create new folders in {0}."),
+				IContentBrowserDataModule::Get().GetSubsystem()->ConvertVirtualPathToDisplay(CreateNewFolderParent));
 		}
 	}
 	else
 	{
-		NewFolderToolTip = FText(LOCTEXT("NewFolderTooltip_InvalidAction", "Cannot create new folders when an asset is selected."));
+		NewFolderToolTip = FText(LOCTEXT("NewFolderTooltip_InvalidAction", "Can only create folders when there is a single path selected."));
 	}
 
 	MenuBuilder.BeginSection("AssetDialogOptions", LOCTEXT("AssetDialogMenuHeading", "Options"));
@@ -814,13 +871,13 @@ void SAssetDialog::SetupContextMenuContent(FMenuBuilder& MenuBuilder, const TArr
 
 	MenuBuilder.EndSection();
 
-	if (CanExecuteExplore())
+	if (CanExecuteExplore(CommandContext))
 	{
 		MenuBuilder.BeginSection("AssetDialogExplore", LOCTEXT("AssetDialogExploreHeading", "Explore"));
 		MenuBuilder.AddMenuEntry(ContentBrowserUtils::GetExploreFolderText(),
 			LOCTEXT("ExploreTooltip", "Finds this folder on disk."),
 			FSlateIcon(ContentBrowserStyleSetName, "SystemWideCommands.FindInContentBrowser"),
-			FUIAction(FExecuteAction::CreateSP(this, &SAssetDialog::ExecuteExplore)));
+			FUIAction(FExecuteAction::CreateSP(this, &SAssetDialog::ExecuteExplore, CommandContext)));
 		MenuBuilder.EndSection();
 	}
 }
@@ -907,7 +964,9 @@ void SAssetDialog::OnUpdateHistoryData(FHistoryData& HistoryData) const
 	const FAssetViewContentSources& ContentSources = AssetPicker->GetAssetView()->GetContentSources();
 	const TArray<FContentBrowserItem> SelectedItems = AssetPicker->GetAssetView()->GetSelectedItems();
 
-	const FText NewSource = ContentSources.HasVirtualPaths() ? FText::FromName(ContentSources.GetVirtualPaths()[0]) : (ContentSources.HasCollections() ? FText::FromName(ContentSources.GetCollections()[0].Name) : LOCTEXT("AllAssets", "All Assets"));
+	const FText NewSource = ContentSources.HasVirtualPaths() ?
+		IContentBrowserDataModule::Get().GetSubsystem()->ConvertVirtualPathToDisplay(ContentSources.GetVirtualPaths()[0], EContentBrowserItemTypeFilter::IncludeFolders) :
+		(ContentSources.HasCollections() ? FText::FromName(ContentSources.GetCollections()[0].Name) : LOCTEXT("AllAssets", "All Assets"));
 
 	HistoryData.HistoryDesc = NewSource;
 	HistoryData.ContentSources = ContentSources;
@@ -926,7 +985,7 @@ FText SAssetDialog::GetAssetNameText() const
 
 FText SAssetDialog::GetPathNameText() const
 {
-	return FText::FromName(GetCurrentSelectedVirtualPath());
+	return IContentBrowserDataModule::Get().GetSubsystem()->ConvertVirtualPathToDisplay(GetCurrentSelectedVirtualPath(), EContentBrowserItemTypeFilter::IncludeFolders);
 }
 
 void SAssetDialog::OnAssetNameTextCommited(const FText& InText, ETextCommit::Type InCommitType)
@@ -973,7 +1032,7 @@ void SAssetDialog::OnPathTextEdited(const FString& NewPath)
 	}
 }
 
-TArray<FString> SAssetDialog::OnCompletePathPrefix(const FString& Prefix) const
+TArray<FNavigationBarComboOption> SAssetDialog::OnCompletePathPrefix(const FString& Prefix) const
 {
 	FStringView PrefixView = Prefix;
 	
@@ -994,14 +1053,14 @@ TArray<FString> SAssetDialog::OnCompletePathPrefix(const FString& Prefix) const
 		*PathPicker->GetPathView()
 	);
 	
-	TArray<FString> Results;
+	TArray<FNavigationBarComboOption> Results;
 	for (const FContentBrowserItem& Item : SubItems)
 	{
 		FName Path = Item.GetVirtualPath();
 		FNameBuilder PathBuilder(Path);
 		if (PathBuilder.ToView().StartsWith(Prefix))
 		{
-			Results.Add(Item.GetVirtualPath().ToString());
+			Results.Emplace(Path);
 		}
 	}
 	return Results;
@@ -1071,9 +1130,20 @@ void SAssetDialog::OnCrumbDelimiterItemClicked(FString ClickedPath)
 	SetCurrentlySelectedPath(ClickedPath, EContentBrowserPathType::Virtual);
 }
 
-TArray<FString> SAssetDialog::GetRecentPaths() const
+TArray<FNavigationBarComboOption> SAssetDialog::GetRecentPaths() const
 {
-	return RecentDirectories;
+	TArray<FNavigationBarComboOption> RecentPaths;
+	RecentPaths.Reserve(RecentDirectories.Num());
+	UContentBrowserDataSubsystem* ContentBrowserData = IContentBrowserDataModule::Get().GetSubsystem();
+	for (const FString& RecentDirectory : RecentDirectories)
+	{
+		FContentBrowserItem Item = ContentBrowserData->GetItemAtPath(FName(RecentDirectory), EContentBrowserItemTypeFilter::IncludeFolders);
+		if (Item.IsValid())
+		{
+			RecentPaths.Emplace(ContentBrowserData->ConvertVirtualPathToDisplay(Item), RecentDirectory);
+		}
+	}
+	return RecentPaths;
 }
 
 void SAssetDialog::SetCurrentlySelectedPath(const FString& NewPath, const EContentBrowserPathType InPathType)
diff --git a/Engine/Source/Editor/ContentBrowser/Private/SAssetDialog.h b/Engine/Source/Editor/ContentBrowser/Private/SAssetDialog.h
--- a/Engine/Source/Editor/ContentBrowser/Private/SAssetDialog.h
+++ b/Engine/Source/Editor/ContentBrowser/Private/SAssetDialog.h
@@ -31,13 +31,13 @@ class SWidget;
 struct FAssetData;
 struct FGeometry;
 struct FKeyEvent;
+struct FNavigationBarComboOption;
 struct FTopLevelAssetPath;
 
-enum class EOpenedContextMenuWidget : uint8
+enum class EAssetDialogCommandContext
 {
 	AssetView,
 	PathView,
-	None
 };
 
 class SAssetDialog : public SCompoundWidget
@@ -117,14 +117,14 @@ private:
 	void OnPathTextEdited(const FString& NewPath);
 
 	/** Provides autocomplete for text-based path editing in the navigation bar */
-	TArray<FString> OnCompletePathPrefix(const FString& Prefix) const;
+	TArray<FNavigationBarComboOption> OnCompletePathPrefix(const FString& Prefix) const;
 	
 	/** Provide content for the navigation bar's breadcrumb menues. */
 	TSharedRef<SWidget> OnGetCrumbDelimiterContent(const FString& CrumbData) const;
 	void OnCrumbDelimiterItemClicked(FString ClickedPath);
 
 	/** Get a list of recently visited paths for the navigation bar. */
-	TArray<FString> GetRecentPaths() const;
+	TArray<FNavigationBarComboOption> GetRecentPaths() const;
 
 	
 	// --- Selection ---
@@ -153,34 +153,35 @@ private:
 
 	/** Will generate the context menu if an asset or a folder is selected, either from the PathView or AssetView */
 	TSharedPtr<SWidget> OnGetAssetContextMenu(const TArray<FAssetData>& SelectedAssets);
-	TSharedPtr<SWidget> OnGetFolderContextMenu(const TArray<FString>& SelectedPaths, FContentBrowserMenuExtender_SelectedPaths InMenuExtender, FOnCreateNewFolder InOnCreateNewFolder);
+	TSharedPtr<SWidget> OnGetFolderContextMenu(const TArray<FString>& SelectedPaths, FContentBrowserMenuExtender_SelectedPaths InMenuExtender, FOnCreateNewFolder InOnCreateNewFolder, EAssetDialogCommandContext CommandContext);
 
 	/** Handler to check to see if a rename command is allowed */
-	bool CanExecuteRename() const;
+	bool CanExecuteRename(EAssetDialogCommandContext CommandContext) const;
 
 	/** Handler for Rename */
-	void ExecuteRename();
+	void ExecuteRename(EAssetDialogCommandContext CommandContext);
 
 	/** Handler to check to see if a delete command is allowed */
-	bool CanExecuteDelete() const;
+	bool CanExecuteDelete(EAssetDialogCommandContext CommandContext) const;
 
 	/** Handler for Delete */
-	void ExecuteDelete();
+	void ExecuteDelete(EAssetDialogCommandContext CommandContext);
 	FReply ExecuteDeleteFolderConfirmed(const TArray<FString> SelectedFolderInternalPaths, const bool bResetSelection);
 
+	FContentBrowserItem GetCreateNewFolderParent(EAssetDialogCommandContext CommandContext) const;
+
 	/** Handler to check to see if a create new folder command is allowed */
-	bool CanExecuteCreateNewFolder() const;
+	bool CanExecuteCreateNewFolder(EAssetDialogCommandContext CommandContext) const;
 
 	/** Handler for creating new folder */
-	void ExecuteCreateNewFolder();
+	void ExecuteCreateNewFolder(EAssetDialogCommandContext CommandContext);
 
 	/** Handler for show in explorer */
-	void ExecuteExplore();
-	bool CanExecuteExplore();
-	bool ExecuteExploreInternal(bool bTest=false);
+	void ExecuteExplore(EAssetDialogCommandContext Widget);
+	bool CanExecuteExplore(EAssetDialogCommandContext Widget);
 
 	/** Setup function for the context menu creation of folder and assets */
-	void SetupContextMenuContent(FMenuBuilder& MenuBuilder, const TArray<FString>& SelectedPaths);
+	void SetupContextMenuContent(EAssetDialogCommandContext CommandContext, FMenuBuilder& MenuBuilder);
 	
 	void BindCommands();
 
@@ -266,7 +267,8 @@ private:
 	bool bValidAssetsChosen;
 
 	/** Commands handled by this widget */
-	TSharedPtr< FUICommandList > Commands;
+	TSharedPtr<FUICommandList> AssetViewCommands;
+	TSharedPtr<FUICommandList> PathViewCommands;
 
 	/** Path Picker used by the dialog */
 	TSharedPtr<SPathPicker> PathPicker;
@@ -277,12 +279,6 @@ private:
 	/** Asset Picker used by the dialog */
 	TSharedPtr<SAssetPicker> AssetPicker;
 
-	/** CreateNewFolder delegate used when user select create new folder from the context menu */
-	FOnCreateNewFolder CurrentContextMenuCreateNewFolderDelegate;
-
-	/** Utility member to know if the context menu was opened on the asset view or the path view */
-	EOpenedContextMenuWidget OpenedContextMenuWidget;
-
 	/** Callback for refreshing content */
 	FOnPathSelected OnPathSelected;
 };
diff --git a/Engine/Source/Editor/ContentBrowser/Private/SAssetPicker.cpp b/Engine/Source/Editor/ContentBrowser/Private/SAssetPicker.cpp
--- a/Engine/Source/Editor/ContentBrowser/Private/SAssetPicker.cpp
+++ b/Engine/Source/Editor/ContentBrowser/Private/SAssetPicker.cpp
@@ -851,12 +851,12 @@ void SAssetPicker::HandleSearchSettingsChanged()
 TSharedPtr<SWidget> SAssetPicker::GetItemContextMenu(TArrayView<const FContentBrowserItem> SelectedItems)
 {
 	// We may only open the file or folder context menu (folder takes priority), so see whether we have any folders selected
-	TArray<FContentBrowserItem> SelectedFolders;
+	TArray<const FContentBrowserItem*> SelectedFolders;
 	for (const FContentBrowserItem& SelectedItem : SelectedItems)
 	{
 		if (SelectedItem.IsFolder())
 		{
-			SelectedFolders.Add(SelectedItem);
+			SelectedFolders.Add(&SelectedItem);
 		}
 	}
 
@@ -865,9 +865,10 @@ TSharedPtr<SWidget> SAssetPicker::GetItemContextMenu(TArrayView<const FContentBr
 		// Folders selected - show the folder menu
 
 		TArray<FString> SelectedPackagePaths;
-		for (const FContentBrowserItem& SelectedFolder : SelectedFolders)
+		SelectedPackagePaths.Reserve(SelectedFolders.Num());
+		for (const FContentBrowserItem* SelectedFolder : SelectedFolders)
 		{
-			FName PackagePath = SelectedFolder.GetInvariantPath();
+			const FName PackagePath = SelectedFolder->GetInternalPath();
 			if (!PackagePath.IsNone())
 			{
 				SelectedPackagePaths.Add(PackagePath.ToString());
diff --git a/Engine/Source/Editor/ContentBrowser/Private/SContentBrowser.cpp b/Engine/Source/Editor/ContentBrowser/Private/SContentBrowser.cpp
--- a/Engine/Source/Editor/ContentBrowser/Private/SContentBrowser.cpp
+++ b/Engine/Source/Editor/ContentBrowser/Private/SContentBrowser.cpp
@@ -1603,9 +1603,6 @@ void SContentBrowser::SyncToItems(TArrayView<const FContentBrowserItem> ItemsToS
 {
 	// For now we just forcefully enable the legacy content source when a sync is requested so it succeeds
 	ContentSourcesContainer->ActivateLegacyContentSource();
-	
-	UContentBrowserDataSubsystem* ContentBrowserData = IContentBrowserDataModule::Get().GetSubsystem();
-	const TSharedRef<FPathPermissionList>& FolderPermissions = FAssetToolsModule::GetModule().Get().GetFolderPermissionList();
 
 	TArray<FContentBrowserItem> NewItemsToSync = ContentBrowserUtils::FilterOrAliasItems(ItemsToSync);
 	ItemsToSync = NewItemsToSync;
@@ -2058,7 +2055,19 @@ void SContentBrowser::LoadSettings(const FName& InInstanceName)
 	}
 	LegacyContentSourceWidgets->AssetViewPtr->LoadSettings(GEditorPerProjectIni, SettingsIniSection, SettingsString);
 
-	GConfig->GetArray(*SettingsIniSection, *(SettingsString + TEXT(".JumpMRU")), JumpMRU, GEditorPerProjectIni);
+	TArray<FString> TempJumpMRU;
+	GConfig->GetArray(*SettingsIniSection, *(SettingsString + TEXT(".JumpMRU")), TempJumpMRU, GEditorPerProjectIni);
+
+	// We previously allowed some non-normalized paths, fix those up.
+	JumpMRU.Reset(TempJumpMRU.Num());
+	for (FString& Path : TempJumpMRU)
+	{
+		FPaths::NormalizeDirectoryName(Path);
+		if (!JumpMRU.Contains(Path))
+		{
+			JumpMRU.Add(MoveTemp(Path));
+		}
+	}
 }
 
 void SContentBrowser::SourcesChanged(const TArray<FString>& SelectedPaths, const TArray<FCollectionRef>& SelectedCollections)
@@ -2269,9 +2278,15 @@ void SContentBrowser::SetSelectedPaths(const TArray<FString>& FolderPaths, bool
 			LegacyContentSourceWidgets->FavoritePathViewPtr->Populate();
 		}
 
-		LegacyContentSourceWidgets->PathViewPtr->SetSelectedPaths(FolderPaths);
-		LegacyContentSourceWidgets->FavoritePathViewPtr->SetSelectedPaths(FolderPaths);
-		PathSelected(FolderPaths[0]);
+		TArray<FString> NormalizedFolderPaths = FolderPaths;
+		for (FString& NormalizedFolderPath : NormalizedFolderPaths)
+		{
+			FPaths::NormalizeDirectoryName(NormalizedFolderPath);
+		}
+
+		LegacyContentSourceWidgets->PathViewPtr->SetSelectedPaths(NormalizedFolderPaths);
+		LegacyContentSourceWidgets->FavoritePathViewPtr->SetSelectedPaths(NormalizedFolderPaths);
+		PathSelected(NormalizedFolderPaths[0]);
 	}
 }
 
@@ -2326,7 +2341,8 @@ void SContentBrowser::OnUpdateHistoryData(FHistoryData& HistoryData) const
 	const FAssetViewContentSources& ContentSources = LegacyContentSourceWidgets->AssetViewPtr->GetContentSources();
 	const TArray<FContentBrowserItem> SelectedItems = LegacyContentSourceWidgets->AssetViewPtr->GetSelectedItems();
 
-	const FText NewSource = ContentSources.HasVirtualPaths() ? FText::FromName(ContentSources.GetVirtualPaths()[0]) :
+	const FText NewSource = ContentSources.HasVirtualPaths() ?
+		IContentBrowserDataModule::Get().GetSubsystem()->ConvertVirtualPathToDisplay(ContentSources.GetVirtualPaths()[0], EContentBrowserItemTypeFilter::IncludeFolders) :
 		(ContentSources.HasCollections() ? FText::FromName(ContentSources.GetCollections()[0].Name) : LOCTEXT("AllAssets", "All Assets"));
 
 	HistoryData.HistoryDesc = NewSource;
@@ -2579,9 +2595,20 @@ void SContentBrowser::OnPathClicked( const FString& CrumbData )
 	}
 }
 
-TArray<FString> SContentBrowser::GetRecentPaths() const
+TArray<FNavigationBarComboOption> SContentBrowser::GetRecentPaths() const
 {
-	return JumpMRU;	
+	TArray<FNavigationBarComboOption> RecentPaths;
+	RecentPaths.Reserve(JumpMRU.Num());
+	UContentBrowserDataSubsystem* ContentBrowserData = IContentBrowserDataModule::Get().GetSubsystem();
+	for (const FString& RecentPath : JumpMRU)
+	{
+		FContentBrowserItem Item = ContentBrowserData->GetItemAtPath(FName(RecentPath), EContentBrowserItemTypeFilter::IncludeFolders);
+		if (Item.IsValid())
+		{
+			RecentPaths.Emplace(ContentBrowserData->ConvertVirtualPathToDisplay(Item), RecentPath);
+		}
+	}
+	return RecentPaths;
 }
 
 void SContentBrowser::OnPathMenuItemClicked(FString ClickedPath)
@@ -2600,7 +2627,7 @@ bool SContentBrowser::OnCanEditPathAsText(const FString& Text) const
 	return true;
 }
 
-TArray<FString> SContentBrowser::OnCompletePathPrefix(const FString& Prefix) const
+TArray<FNavigationBarComboOption> SContentBrowser::OnCompletePathPrefix(const FString& Prefix) const
 {
 	FStringView PrefixView = Prefix;
 	
@@ -2620,14 +2647,14 @@ TArray<FString> SContentBrowser::OnCompletePathPrefix(const FString& Prefix) con
 		LegacyContentSourceWidgets->PathViewPtr->GetContentBrowserItemAttributeFilter(),
 		InstanceName,
 		*LegacyContentSourceWidgets->PathViewPtr); 
-	TArray<FString> Results;
+	TArray<FNavigationBarComboOption> Results;
 	for (const FContentBrowserItem& Item : SubItems)
 	{
 		FName Path = Item.GetVirtualPath();
 		FNameBuilder PathBuilder(Path);
 		if (PathBuilder.ToView().StartsWith(Prefix))
 		{
-			Results.Add(Item.GetVirtualPath().ToString());
+			Results.Emplace(Item.GetVirtualPath());
 		}
 	}
 	return Results;
@@ -5108,11 +5135,11 @@ void SContentBrowser::PopulateFolderContextMenu(UToolMenu* Menu)
 	{
 		if(bCanCreateNewFolder)
 		{
-			NewFolderToolTip = FText::Format(LOCTEXT("NewFolderTooltip_CreateIn", "Create a new folder in {0}."), FText::FromName(SelectedFolders[0].GetVirtualPath()));
+			NewFolderToolTip = FText::Format(LOCTEXT("NewFolderTooltip_CreateIn", "Create a new folder in {0}."), ContentBrowserData->ConvertVirtualPathToDisplay(SelectedFolders[0]));
 		}
 		else
 		{
-			NewFolderToolTip = FText::Format(LOCTEXT("NewFolderTooltip_InvalidPath", "Cannot create new folders in {0}."), FText::FromName(SelectedFolders[0].GetVirtualPath()));
+			NewFolderToolTip = FText::Format(LOCTEXT("NewFolderTooltip_InvalidPath", "Cannot create new folders in {0}."), ContentBrowserData->ConvertVirtualPathToDisplay(SelectedFolders[0]));
 		}
 	}
 	else
diff --git a/Engine/Source/Editor/ContentBrowser/Private/SContentBrowser.h b/Engine/Source/Editor/ContentBrowser/Private/SContentBrowser.h
--- a/Engine/Source/Editor/ContentBrowser/Private/SContentBrowser.h
+++ b/Engine/Source/Editor/ContentBrowser/Private/SContentBrowser.h
@@ -76,6 +76,7 @@ struct FContentBrowserItem;
 struct FContentBrowserItemPath;
 struct FGeometry;
 struct FKeyEvent;
+struct FNavigationBarComboOption;
 struct FPointerEvent;
 struct FSlateBrush;
 struct FToolMenuContext;
@@ -378,7 +379,7 @@ private:
 	bool ParseCollectionCrumbData(const FString& CrumbData, FCollectionSource*& OutCollectionSource, FCollectionNameType& OutCollection) const;
 
 	/** Return a list of recently visited locations which can be used with OnNavigateToPath to return to those locations */
-	TArray<FString> GetRecentPaths() const;	
+	TArray<FNavigationBarComboOption> GetRecentPaths() const;
 	
 	/** Navigate to a location by text identifier. E.g. a virtual path such as /All/Game/MyFolder */
 	void OnNavigateToPath(const FString& Path);
@@ -390,7 +391,7 @@ private:
 	bool OnCanEditPathAsText(const FString& Text) const;
 	
 	/** Returns a list of valid paths starting with Prefix which can be navigated to with OnNavigateToPath */
-	TArray<FString> OnCompletePathPrefix(const FString& Prefix) const;
+	TArray<FNavigationBarComboOption> OnCompletePathPrefix(const FString& Prefix) const;
 
 	/** Register the context objects needed for the "Add New" menu */
 	void AppendNewMenuContextObjects(const EContentBrowserDataMenuContext_AddNewMenuDomain InDomain, const TArray<FName>& InSelectedPaths, FToolMenuContext& InOutMenuContext, UContentBrowserToolbarMenuContext* CommonContext, bool bCanBeModified);
diff --git a/Engine/Source/Editor/ContentBrowser/Private/SNavigationBar.cpp b/Engine/Source/Editor/ContentBrowser/Private/SNavigationBar.cpp
--- a/Engine/Source/Editor/ContentBrowser/Private/SNavigationBar.cpp
+++ b/Engine/Source/Editor/ContentBrowser/Private/SNavigationBar.cpp
@@ -30,6 +30,7 @@ enum class ELocationSource
 // Type for list view control
 struct FLocationItem
 {
+	FText DisplayPath;
 	FString VirtualPath;
 	ELocationSource Source;
 };
@@ -68,14 +69,14 @@ public:
 		return FReply::Unhandled();
 	}
 	
-	void SetItems(TArray<FString> InItems, ELocationSource InSource)
+	void SetItems(TArray<FNavigationBarComboOption> InItems, ELocationSource InSource)
 	{
 		RequestListRefresh();
 		ClearSelection();
 		Items.Reset(InItems.Num());
-		for (FString& Item : InItems)
+		for (FNavigationBarComboOption& Item : InItems)
 		{
-			Items.Emplace(MakeShared<FLocationItem>(FLocationItem{MoveTemp(Item), InSource}));
+			Items.Emplace(MakeShared<FLocationItem>(FLocationItem{ MoveTemp(Item.DisplayPath), MoveTemp(Item.VirtualPath), InSource }));
 		}
 	}
 	
@@ -103,6 +104,18 @@ protected:
 	TArray<TSharedPtr<FLocationItem>> Items;
 };
 
+FNavigationBarComboOption::FNavigationBarComboOption(FName InVirtualPath)
+{
+	VirtualPath = InVirtualPath.ToString();
+	DisplayPath = FText::FromString(VirtualPath);
+}
+
+FNavigationBarComboOption::FNavigationBarComboOption(const FText& InDisplayPath, const FString& InVirtualPath)
+	: DisplayPath(InDisplayPath)
+	, VirtualPath(InVirtualPath)
+{
+}
+
 void SNavigationBar::Construct(const FArguments& InArgs)
 {
 	OnGetComboOptions = InArgs._GetComboOptions;
@@ -297,7 +310,7 @@ FReply SNavigationBar::OnButtonClicked()
 void SNavigationBar::GenerateHistoryOptions()
 {
 	PopupHeading = LOCTEXT("NavigationBar.HistoryHeader", "HISTORY");
-	TArray<FString> Options;
+	TArray<FNavigationBarComboOption> Options;
 	if (OnGetComboOptions.IsBound())
 	{
 		Options = OnGetComboOptions.Execute();
@@ -308,7 +321,7 @@ void SNavigationBar::GenerateHistoryOptions()
 void SNavigationBar::GenerateCompletionOptions(const FString& Prefix)
 {
 	PopupHeading = LOCTEXT("NavigationBar.SuggestionsHeader", "SUGGESTIONS"); 
-	TArray<FString> Options;
+	TArray<FNavigationBarComboOption> Options;
 	if (OnCompletePrefix.IsBound())
 	{
 		Options = OnCompletePrefix.Execute(Prefix);
@@ -504,7 +517,7 @@ TSharedRef<ITableRow> SNavigationBar::HandleGenerateComboRow(TSharedPtr<FLocatio
 			.Padding(2.0f, 0.0f)
 			[
 				SNew(STextBlock)
-				.Text(FText::FromString(ForItem->VirtualPath))
+				.Text(ForItem->DisplayPath)
 			]
 		];
 }
diff --git a/Engine/Source/Editor/ContentBrowser/Private/SNavigationBar.h b/Engine/Source/Editor/ContentBrowser/Private/SNavigationBar.h
--- a/Engine/Source/Editor/ContentBrowser/Private/SNavigationBar.h
+++ b/Engine/Source/Editor/ContentBrowser/Private/SNavigationBar.h
@@ -22,10 +22,20 @@ class SLocationListView;
 class STableViewBase;
 class USlateWidgetStyleAsset;
 
+struct FNavigationBarComboOption
+{
+	FNavigationBarComboOption() = default;
+	explicit FNavigationBarComboOption(FName InVirtualPath);
+	FNavigationBarComboOption(const FText& InDisplayPath, const FString& InVirtualPath);
+
+	FText DisplayPath;
+	FString VirtualPath;
+};
+
 using FOnPathClicked = TDelegate<void(const FString&)>;
 using FGetPathMenuContent = TDelegate<TSharedRef<SWidget>(const FString&)>;
-using FGetComboOptions = TDelegate<TArray<FString>(void)>;
-using FOnCompletePrefix = TDelegate<TArray<FString>(const FString&)>;
+using FGetComboOptions = TDelegate<TArray<FNavigationBarComboOption>(void)>;
+using FOnCompletePrefix = TDelegate<TArray<FNavigationBarComboOption>(const FString&)>;
 using FOnNavigateToPath = TDelegate<void(const FString&)>;
 using FCanEditPathAsText = TDelegate<bool(const FString&)>;
 
diff --git a/Engine/Source/Editor/ContentBrowser/Private/SPathPicker.cpp b/Engine/Source/Editor/ContentBrowser/Private/SPathPicker.cpp
--- a/Engine/Source/Editor/ContentBrowser/Private/SPathPicker.cpp
+++ b/Engine/Source/Editor/ContentBrowser/Private/SPathPicker.cpp
@@ -103,13 +103,7 @@ void SPathPicker::OnItemSelectionChanged(const FContentBrowserItem& SelectedItem
 
 TSharedPtr<SWidget> SPathPicker::GetItemContextMenu(TArrayView<const FContentBrowserItem> SelectedItems)
 {
-	TArray<FString> SelectedPackagePaths;
-	for (const FContentBrowserItem& SelectedItem : SelectedItems)
-	{
-		SelectedPackagePaths.Add(SelectedItem.GetVirtualPath().ToString());
-	}
-
-	if (SelectedPackagePaths.Num() == 0)
+	if (SelectedItems.IsEmpty())
 	{
 		return nullptr;
 	}
@@ -118,13 +112,13 @@ TSharedPtr<SWidget> SPathPicker::GetItemContextMenu(TArrayView<const FContentBro
 
 	if (OnGetFolderContextMenu.IsBound())
 	{
-		return OnGetFolderContextMenu.Execute(SelectedPackagePaths, OnGetPathContextMenuExtender, OnCreateNewFolder);
+		return OnGetFolderContextMenu.Execute(GetInternalPaths(SelectedItems), OnGetPathContextMenuExtender, OnCreateNewFolder);
 	}
 
-	return GetFolderContextMenu(SelectedPackagePaths, OnGetPathContextMenuExtender, OnCreateNewFolder);
+	return GetFolderContextMenu(SelectedItems, OnGetPathContextMenuExtender, OnCreateNewFolder);
 }
 
-TSharedPtr<SWidget> SPathPicker::GetFolderContextMenu(const TArray<FString>& SelectedPaths, FContentBrowserMenuExtender_SelectedPaths InMenuExtender, FOnCreateNewFolder InOnCreateNewFolder)
+TSharedPtr<SWidget> SPathPicker::GetFolderContextMenu(TArrayView<const FContentBrowserItem> SelectedItems, FContentBrowserMenuExtender_SelectedPaths InMenuExtender, FOnCreateNewFolder InOnCreateNewFolder)
 {
 	UContentBrowserDataSubsystem* ContentBrowserData = IContentBrowserDataModule::Get().GetSubsystem();
 
@@ -132,7 +126,7 @@ TSharedPtr<SWidget> SPathPicker::GetFolderContextMenu(const TArray<FString>& Sel
 	if (InMenuExtender.IsBound())
 	{
 		// Code using extenders here currently expects internal paths
-		Extender = InMenuExtender.Execute(ContentBrowserData->TryConvertVirtualPathsToInternal(SelectedPaths));
+		Extender = InMenuExtender.Execute(GetInternalPaths(SelectedItems));
 	}
 
 	const bool bInShouldCloseWindowAfterSelection = true;
@@ -140,18 +134,22 @@ TSharedPtr<SWidget> SPathPicker::GetFolderContextMenu(const TArray<FString>& Sel
 	FMenuBuilder MenuBuilder(bInShouldCloseWindowAfterSelection, nullptr, Extender, bCloseSelfOnly);
 
 	// We can only create folders when we have a single path selected
-	const bool bCanCreateNewFolder = SelectedPaths.Num() == 1 && ContentBrowserData->CanCreateFolder(*SelectedPaths[0], nullptr);
+	const bool bCanCreateNewFolder = SelectedItems.Num() == 1 && ContentBrowserData->CanCreateFolder(SelectedItems[0].GetVirtualPath(), nullptr);
 
 	FText NewFolderToolTip;
-	if(SelectedPaths.Num() == 1)
+	if(SelectedItems.Num() == 1)
 	{
 		if(bCanCreateNewFolder)
 		{
-			NewFolderToolTip = FText::Format(LOCTEXT("NewFolderTooltip_CreateIn", "Create a new folder in {0}."), FText::FromString(SelectedPaths[0]));
+			NewFolderToolTip = FText::Format(
+				LOCTEXT("NewFolderTooltip_CreateIn", "Create a new folder in {0}."),
+				IContentBrowserDataModule::Get().GetSubsystem()->ConvertVirtualPathToDisplay(SelectedItems[0]));
 		}
 		else
 		{
-			NewFolderToolTip = FText::Format(LOCTEXT("NewFolderTooltip_InvalidPath", "Cannot create new folders in {0}."), FText::FromString(SelectedPaths[0]));
+			NewFolderToolTip = FText::Format(
+				LOCTEXT("NewFolderTooltip_InvalidPath", "Cannot create new folders in {0}."),
+				IContentBrowserDataModule::Get().GetSubsystem()->ConvertVirtualPathToDisplay(SelectedItems[0]));
 		}
 	}
 	else
@@ -165,7 +163,7 @@ TSharedPtr<SWidget> SPathPicker::GetFolderContextMenu(const TArray<FString>& Sel
 		NewFolderToolTip,
 		FSlateIcon(UE::ContentBrowser::Private::FContentBrowserStyle::Get().GetStyleSetName(), "ContentBrowser.NewFolderIcon"),
 		FUIAction(
-			FExecuteAction::CreateSP(this, &SPathPicker::CreateNewFolder, SelectedPaths.Num() > 0 ? SelectedPaths[0] : FString(), InOnCreateNewFolder),
+			FExecuteAction::CreateSP(this, &SPathPicker::CreateNewFolder, SelectedItems.Num() > 0 ? SelectedItems[0].GetVirtualPath().ToString() : FString(), InOnCreateNewFolder),
 			FCanExecuteAction::CreateLambda( [bCanCreateNewFolder] { return bCanCreateNewFolder; } )
 			),
 		"FolderContext"
@@ -174,6 +172,21 @@ TSharedPtr<SWidget> SPathPicker::GetFolderContextMenu(const TArray<FString>& Sel
 	return MenuBuilder.MakeWidget();
 }
 
+TArray<FString> SPathPicker::GetInternalPaths(TArrayView<const FContentBrowserItem> SelectedItems)
+{
+	TArray<FString> InternalPaths;
+	InternalPaths.Reserve(SelectedItems.Num());
+	for (const FContentBrowserItem& SelectedItem : SelectedItems)
+	{
+		const FName InternalPath = SelectedItem.GetInternalPath();
+		if (!InternalPath.IsNone())
+		{
+			InternalPaths.Add(InternalPath.ToString());
+		}
+	}
+	return InternalPaths;
+}
+
 void SPathPicker::ExecuteRenameFolder()
 {
 	if (PathViewPtr.IsValid())
diff --git a/Engine/Source/Editor/ContentBrowser/Private/SPathPicker.h b/Engine/Source/Editor/ContentBrowser/Private/SPathPicker.h
--- a/Engine/Source/Editor/ContentBrowser/Private/SPathPicker.h
+++ b/Engine/Source/Editor/ContentBrowser/Private/SPathPicker.h
@@ -62,7 +62,8 @@ private:
 
 	/** Handler for the context menu for folder items */
 	TSharedPtr<SWidget> GetItemContextMenu(TArrayView<const FContentBrowserItem> SelectedItems);
-	TSharedPtr<SWidget> GetFolderContextMenu(const TArray<FString>& SelectedPaths, FContentBrowserMenuExtender_SelectedPaths InMenuExtender, FOnCreateNewFolder InOnCreateNewFolder);
+	TSharedPtr<SWidget> GetFolderContextMenu(TArrayView<const FContentBrowserItem> SelectedItems, FContentBrowserMenuExtender_SelectedPaths InMenuExtender, FOnCreateNewFolder InOnCreateNewFolder);
+	static TArray<FString> GetInternalPaths(TArrayView<const FContentBrowserItem> SelectedItems);
 
 private:
 
diff --git a/Engine/Source/Editor/ContentBrowser/Private/SourcesViewWidgets.cpp b/Engine/Source/Editor/ContentBrowser/Private/SourcesViewWidgets.cpp
--- a/Engine/Source/Editor/ContentBrowser/Private/SourcesViewWidgets.cpp
+++ b/Engine/Source/Editor/ContentBrowser/Private/SourcesViewWidgets.cpp
@@ -8,6 +8,7 @@
 #include "CollectionViewUtils.h"
 #include "Containers/UnrealString.h"
 #include "ContentBrowserDataSource.h"
+#include "ContentBrowserDataSubsystem.h"
 #include "ContentBrowserItem.h"
 #include "ContentBrowserItemData.h"
 #include "ContentBrowserModule.h"
@@ -20,6 +21,7 @@
 #include "DragDropHandler.h"
 #include "Fonts/SlateFontInfo.h"
 #include "GenericPlatform/ICursor.h"
+#include "IContentBrowserDataModule.h"
 #include "Input/DragAndDrop.h"
 #include "Layout/Children.h"
 #include "Layout/Margin.h"
@@ -380,7 +382,7 @@ FText SAssetTreeItem::GetToolTipText() const
 	{
 		// If this item is a plugin folder, append the plugin description to the tooltip
 		const FContentBrowserItem& Item = TreeItemPin->GetItem();
-		FText PathText = FText::FromName(Item.GetVirtualPath());
+		FText PathText = IContentBrowserDataModule::Get().GetSubsystem()->ConvertVirtualPathToDisplay(Item);
 		if (Item.IsInPlugin())
 		{
 			FNameBuilder ItemPath{Item.GetInternalPath()};
diff --git a/Engine/Source/Editor/ContentBrowserData/Private/ContentBrowserDataDragDropOp.cpp b/Engine/Source/Editor/ContentBrowserData/Private/ContentBrowserDataDragDropOp.cpp
--- a/Engine/Source/Editor/ContentBrowserData/Private/ContentBrowserDataDragDropOp.cpp
+++ b/Engine/Source/Editor/ContentBrowserData/Private/ContentBrowserDataDragDropOp.cpp
@@ -171,7 +171,8 @@ FText FContentBrowserDataDragDropOp::GetFirstItemText() const
 
 	if (DraggedFolders.Num() > 0)
 	{
-		return FText::FromName(DraggedFolders[0].GetVirtualPath());
+		UContentBrowserDataSubsystem* ContentBrowserData = GEditor->GetEditorSubsystem<UContentBrowserDataSubsystem>();
+		return ContentBrowserData->ConvertVirtualPathToDisplay(DraggedFolders[0]);
 	}
 
 	return FText::GetEmpty();
diff --git a/Engine/Source/Editor/ContentBrowserData/Private/ContentBrowserDataSubsystem.cpp b/Engine/Source/Editor/ContentBrowserData/Private/ContentBrowserDataSubsystem.cpp
--- a/Engine/Source/Editor/ContentBrowserData/Private/ContentBrowserDataSubsystem.cpp
+++ b/Engine/Source/Editor/ContentBrowserData/Private/ContentBrowserDataSubsystem.cpp
@@ -19,6 +19,7 @@
 #include "Misc/EnumClassFlags.h"
 #include "Misc/PackageName.h"
 #include "Misc/PathViews.h"
+#include "Misc/PathViews.h"
 #include "Misc/StringBuilder.h"
 #include "PluginDescriptor.h"
 #include "Settings/ContentBrowserSettings.h"
@@ -1143,6 +1144,93 @@ TArray<FString> UContentBrowserDataSubsystem::ConvertInternalPathsToVirtual(cons
 	return VirtualPaths;
 }
 
+FText UContentBrowserDataSubsystem::ConvertVirtualPathToDisplay(const FStringView InVirtualPath, const EContentBrowserItemTypeFilter InLeafType) const
+{
+	FString DisplayPath;
+	DisplayPath.Reserve(InVirtualPath.Len());
+
+	TArray<FStringView, TInlineAllocator<8>> Components;
+	FPathViews::IterateComponents(InVirtualPath, [&Components](FStringView Component)
+	{
+		if (!Component.IsEmpty())
+		{
+			Components.Add(Component);
+		}
+	});
+
+	FNameBuilder PathBuilder;
+	for (int32 Index = 0; Index < Components.Num(); ++Index)
+	{
+		PathBuilder.AppendChar(TEXT('/'));
+		PathBuilder.Append(Components[Index]);
+
+		DisplayPath += TEXT('/');
+		const FContentBrowserItem Item = GetItemAtPath(FName(PathBuilder), Index + 1 == Components.Num() ? InLeafType : EContentBrowserItemTypeFilter::IncludeFolders);
+		if (Item.IsValid())
+		{
+			DisplayPath += Item.GetDisplayName().ToString();
+		}
+		else
+		{
+			DisplayPath += Components[Index];
+		}
+	}
+
+	return FText::FromString(MoveTemp(DisplayPath));
+}
+
+FText UContentBrowserDataSubsystem::ConvertVirtualPathToDisplay(FName InVirtualPath, const EContentBrowserItemTypeFilter InLeafType) const
+{
+	FNameBuilder VirtualPath;
+	InVirtualPath.AppendString(VirtualPath);
+	return ConvertVirtualPathToDisplay(VirtualPath.ToView(), InLeafType);
+}
+
+FText UContentBrowserDataSubsystem::ConvertVirtualPathToDisplay(const FContentBrowserItem& InItem) const
+{
+	FNameBuilder VirtualPath;
+	InItem.GetVirtualPath().AppendString(VirtualPath);
+
+	FString DisplayPath;
+	DisplayPath.Reserve(VirtualPath.Len());
+
+	TArray<FStringView, TInlineAllocator<8>> Components;
+	FPathViews::IterateComponents(VirtualPath, [&Components](FStringView Component)
+	{
+		if (!Component.IsEmpty())
+		{
+			Components.Add(Component);
+		}
+	});
+
+	FNameBuilder PathBuilder;
+	for (int32 Index = 0; Index < Components.Num(); ++Index)
+	{
+		PathBuilder.AppendChar(TEXT('/'));
+		PathBuilder.Append(Components[Index]);
+
+		DisplayPath += TEXT('/');
+		if (Index + 1 == Components.Num())
+		{
+			DisplayPath += InItem.GetDisplayName().ToString();
+		}
+		else
+		{
+			const FContentBrowserItem Item = GetItemAtPath(FName(PathBuilder), EContentBrowserItemTypeFilter::IncludeFolders);
+			if (Item.IsValid())
+			{
+				DisplayPath += Item.GetDisplayName().ToString();
+			}
+			else
+			{
+				DisplayPath += Components[Index];
+			}
+		}
+	}
+
+	return FText::FromString(MoveTemp(DisplayPath));
+}
+
 void UContentBrowserDataSubsystem::SetGenerateVirtualPathPrefixDelegate(const FContentBrowserGenerateVirtualPathDelegate& InDelegate)
 {
 	GenerateVirtualPathPrefixDelegate = InDelegate;
diff --git a/Engine/Source/Editor/ContentBrowserData/Public/ContentBrowserDataSubsystem.h b/Engine/Source/Editor/ContentBrowserData/Public/ContentBrowserDataSubsystem.h
--- a/Engine/Source/Editor/ContentBrowserData/Public/ContentBrowserDataSubsystem.h
+++ b/Engine/Source/Editor/ContentBrowserData/Public/ContentBrowserDataSubsystem.h
@@ -353,6 +353,14 @@ public:
 	FName ConvertInternalPathToVirtual(FName InPath);
 	TArray<FString> ConvertInternalPathsToVirtual(const TArray<FString>& InPaths);
 
+	/**
+	 * Converts a virtual path to a display path, where each path segment is replaced with the display name of its corresponding FContentBrowserItem.
+	 * This path will display similar to the Content Browser's navigation bar.
+	 */
+	FText ConvertVirtualPathToDisplay(const FStringView InVirtualPath, const EContentBrowserItemTypeFilter InLeafType) const;
+	FText ConvertVirtualPathToDisplay(FName InVirtualPath, const EContentBrowserItemTypeFilter InLeafType) const;
+	FText ConvertVirtualPathToDisplay(const FContentBrowserItem& InItem) const;
+
 	/**
 	 * Converts virtual path back into an internal or invariant path
 	 */
