Updated AssetViewItem to support returning the ContentItem's underlying EAssetAccessSpecifier
Updated SAssetView to support rendering the EAssetAccessSpecifier in tile, list and column views.
Added an option in UContentBrowserConfig to toggle viewing the EAssetAccessSpecifier. Off by default.

#rnx

diff --git a/Engine/Source/Editor/ContentBrowser/Private/AssetSystemContentBrowserInfoProvider.cpp b/Engine/Source/Editor/ContentBrowser/Private/AssetSystemContentBrowserInfoProvider.cpp
--- a/Engine/Source/Editor/ContentBrowser/Private/AssetSystemContentBrowserInfoProvider.cpp
+++ b/Engine/Source/Editor/ContentBrowser/Private/AssetSystemContentBrowserInfoProvider.cpp
@@ -116,17 +116,11 @@ void FAssetSystemContentBrowserInfoProvider::PopulateAssetInfo(TArray<FAssetDisp
 			CookPathLenInfo.StatusDescription = FText::Format(LOCTEXT("TileViewTooltipPathLengthForCookingValue", "{0} / {1}"), FText::AsNumber(PackageNameLengthForCooking), FText::AsNumber(MaxCookPathLen));
 			OutAssetDisplayInfo.Add(CookPathLenInfo);
 			
-			if (ItemAssetData.GetAssetAccessSpecifier() == EAssetAccessSpecifier::Public)
+			PublicStateText = AssetItem->GetItemAssetAccessSpecifierText();
+			if (!PublicStateText.IsEmpty())
 			{
-				PublicStateText = LOCTEXT("PublicAssetState", "Public");
-			}
-			else if (ItemAssetData.GetAssetAccessSpecifier() == EAssetAccessSpecifier::EpicInternal)
-			{
-				PublicStateText = LOCTEXT("EpicInternalAssetState", "Epic Internal");
-			}
-			else
-			{
-				PublicStateText = LOCTEXT("PrivateAssetState", "Private");
+				// An icon is required for the PublicStateText to display. So, set one if there's a value to show.
+				PublicStateIcon = FAppStyle::GetBrush("Icons.Info");
 			}
 		}
 
diff --git a/Engine/Source/Editor/ContentBrowser/Private/AssetViewSortManager.cpp b/Engine/Source/Editor/ContentBrowser/Private/AssetViewSortManager.cpp
--- a/Engine/Source/Editor/ContentBrowser/Private/AssetViewSortManager.cpp
+++ b/Engine/Source/Editor/ContentBrowser/Private/AssetViewSortManager.cpp
@@ -229,6 +229,44 @@ protected:
 	}
 };
 
+struct FCompareFAssetItemByAssetAccessSpecifier : public FCompareFAssetItemBase
+{
+public:
+	FCompareFAssetItemByAssetAccessSpecifier(bool bInAscending, const FName& InTag) : FCompareFAssetItemBase(bInAscending, InTag) {}
+	virtual ~FCompareFAssetItemByAssetAccessSpecifier() override = default;
+
+protected:
+	FORCEINLINE virtual bool Compare(const TSharedPtr<FAssetViewItem>& A, const TSharedPtr<FAssetViewItem>& B) const override
+	{
+		FText Value1 = A->GetItemAssetAccessSpecifierText();
+		FText Value2 = B->GetItemAssetAccessSpecifierText();
+
+		// Depending if we're sorting ascending or descending it's quicker to flip the compares incase values are missing
+		const bool bFoundValue1 = bAscending ? !Value1.IsEmpty() : !Value2.IsEmpty();
+		if (!bFoundValue1)
+		{
+			return true;
+		}
+
+		const bool bFoundValue2 = bAscending ? !Value2.IsEmpty() : !Value1.IsEmpty();
+		if (!bFoundValue2)
+		{
+			return false;
+		}
+
+		const int32 Result = Value1.CompareTo(Value2);
+		if (Result < 0)
+		{
+			return bAscending;
+		}
+		else if (Result > 0)
+		{
+			return !bAscending;
+		}
+		return FCompareFAssetItemBase::Compare(A, B);
+	}
+};
+
 struct FCompareFAssetItemByTag : public FCompareFAssetItemBase
 {
 public:
@@ -365,6 +403,7 @@ const FName FAssetViewSortManager::ClassColumnId = "Class";
 const FName FAssetViewSortManager::PathColumnId = "Path";
 const FName FAssetViewSortManager::RevisionControlColumnId = "RevisionControl";
 const FName FAssetViewSortManager::DiskSizeColumnId = "Size";
+const FName FAssetViewSortManager::AssetAccessSpecifierColumnId = "AssetAccessSpecifier";
 
 FAssetViewSortManager::FAssetViewSortManager()
 {
@@ -437,6 +476,10 @@ void FAssetViewSortManager::SortList(TArray<TSharedPtr<FAssetViewItem>>& AssetIt
 		{
 			SortMethod.Add(MakeUnique<FCompareFAssetItemByDiskSize>(bAscending, Tag));
 		}
+		else if (Tag == AssetAccessSpecifierColumnId)
+		{
+			SortMethod.Add(MakeUnique<FCompareFAssetItemByAssetAccessSpecifier>(bAscending, Tag));
+		}
 		else
 		{
 			UObject::FAssetRegistryTag::ETagType TagType = UObject::FAssetRegistryTag::TT_Hidden;
diff --git a/Engine/Source/Editor/ContentBrowser/Private/AssetViewTypes.cpp b/Engine/Source/Editor/ContentBrowser/Private/AssetViewTypes.cpp
--- a/Engine/Source/Editor/ContentBrowser/Private/AssetViewTypes.cpp
+++ b/Engine/Source/Editor/ContentBrowser/Private/AssetViewTypes.cpp
@@ -7,10 +7,13 @@
 #include "ContentBrowserDelegates.h"
 #include "ContentBrowserItemData.h"
 #include "HAL/Platform.h"
+#include "IAssetTools.h"
 #include "Internationalization/Text.h"
 #include "Misc/AssertionMacros.h"
 #include "Templates/UnrealTemplate.h"
 
+#define LOCTEXT_NAMESPACE "ContentBrowser"
+
 FAssetViewItem::FAssetViewItem(int32 InIndex)
 	: Index(InIndex)
 {
@@ -197,6 +200,52 @@ bool FAssetViewItem::GetTagValue(const FName Tag, FString& OutString, UObject::F
 	return false;
 }
 
+bool FAssetViewItem::GetItemAssetAccessSpecifier(EAssetAccessSpecifier& OutAssetAccessSpecifier) const
+{
+	return Item.GetItemAssetAccessSpecifier(OutAssetAccessSpecifier);
+}
+
+FText FAssetViewItem::GetItemAssetAccessSpecifierText() const
+{
+	// Items that don't have a modifable access specifier shouldn't render anything.
+	if (!CanModifyItemAssetAccessSpecifier())
+	{
+		return FText::GetEmpty();
+	}
+
+	EAssetAccessSpecifier AssetAccessSpecifier;
+	if (GetItemAssetAccessSpecifier(AssetAccessSpecifier))
+	{
+		switch (AssetAccessSpecifier)
+		{
+			case EAssetAccessSpecifier::Public:
+			{
+				return LOCTEXT("PublicAssetState", "Public");
+			}
+
+			case EAssetAccessSpecifier::EpicInternal:
+			{
+				return LOCTEXT("PrivateAssetState", "Epic Internal");
+			}
+
+			case EAssetAccessSpecifier::Private:
+			{
+				return LOCTEXT("EpicInternalAssetState", "Private");
+			}
+		}
+	}
+
+	// This would only happen if there's something wrong with the underlying FContentBrowserItem.
+	// It somehow returned true for CanModifyItemAssetAccessSpecifier, but didn't implement its GetItemAssetAccessSpecifier or did and returned something unknown.
+	ensureMsgf(0, TEXT("Failed to read the EAssetAccessSpecifier value for FContentBrowserItem: '%s'"), *Item.GetItemName().ToString());
+	return LOCTEXT("UnknownAssetState", "Unknown");
+}
+
+bool FAssetViewItem::CanModifyItemAssetAccessSpecifier() const
+{
+	return Item.CanModifyItemAssetAccessSpecifier();
+}
+
 const FContentBrowserItem& FAssetViewItem::GetItem() const
 {
 	return Item;
@@ -240,3 +289,5 @@ FString FAssetViewItem::ItemToString_Debug(TSharedPtr<FAssetViewItem> AssetItem)
 	}
 	return TEXT("nullptr");
 }
+
+#undef LOCTEXT_NAMESPACE
diff --git a/Engine/Source/Editor/ContentBrowser/Private/AssetViewWidgets.cpp b/Engine/Source/Editor/ContentBrowser/Private/AssetViewWidgets.cpp
--- a/Engine/Source/Editor/ContentBrowser/Private/AssetViewWidgets.cpp
+++ b/Engine/Source/Editor/ContentBrowser/Private/AssetViewWidgets.cpp
@@ -136,18 +136,18 @@ bool IsTopLevelFolder(const FName InFolderPath)
 // FAssetViewModeUtils
 ///////////////////////////////
 
-TSharedRef<SWidget> FAssetViewItemHelper::CreateListItemContents(SAssetListItem* const InListItem, const TSharedRef<SWidget>& InThumbnail, FName& OutItemShadowBorder, const FMargin& InFolderOverlayPadding)
+TSharedRef<SWidget> FAssetViewItemHelper::CreateListItemContents(SAssetListItem* const InListItem, const TSharedRef<SWidget>& InThumbnail, FName& OutItemShadowBorder, const FMargin& InFolderOverlayPadding, bool bShowAssetAccessSpecifier)
 {
-	return CreateListTileItemContents(InListItem, InThumbnail, OutItemShadowBorder, InFolderOverlayPadding);
+	return CreateListTileItemContents(InListItem, InThumbnail, OutItemShadowBorder, InFolderOverlayPadding, bShowAssetAccessSpecifier);
 }
 
-TSharedRef<SWidget> FAssetViewItemHelper::CreateTileItemContents(SAssetTileItem* const InTileItem, const TSharedRef<SWidget>& InThumbnail, FName& OutItemShadowBorder, const FMargin& InFolderOverlayPadding)
+TSharedRef<SWidget> FAssetViewItemHelper::CreateTileItemContents(SAssetTileItem* const InTileItem, const TSharedRef<SWidget>& InThumbnail, FName& OutItemShadowBorder, const FMargin& InFolderOverlayPadding, bool bShowAssetAccessSpecifier)
 {
-	return CreateListTileItemContents(InTileItem, InThumbnail, OutItemShadowBorder, InFolderOverlayPadding);
+	return CreateListTileItemContents(InTileItem, InThumbnail, OutItemShadowBorder, InFolderOverlayPadding, bShowAssetAccessSpecifier);
 }
 
 template <typename T>
-TSharedRef<SWidget> FAssetViewItemHelper::CreateListTileItemContents(T* const InTileOrListItem, const TSharedRef<SWidget>& InThumbnail, FName& OutItemShadowBorder, const FMargin& InFolderOverlayPadding)
+TSharedRef<SWidget> FAssetViewItemHelper::CreateListTileItemContents(T* const InTileOrListItem, const TSharedRef<SWidget>& InThumbnail, FName& OutItemShadowBorder, const FMargin& InFolderOverlayPadding, bool bShowAssetAccessSpecifier)
 {
 	TSharedRef<SOverlay> ItemContentsOverlay = SNew(SOverlay);
 
@@ -255,6 +255,22 @@ TSharedRef<SWidget> FAssetViewItemHelper::CreateListTileItemContents(T* const In
 			]
 		];
 
+		// Add the icon for the asset's EAssetAccessSpecifier
+		if (bShowAssetAccessSpecifier)
+		{
+			ItemContentsOverlay->AddSlot()
+			.HAlign(HAlign_Right)
+			.VAlign(VAlign_Bottom)
+			[
+				SNew(SBox)
+				.MaxDesiredWidth(InTileOrListItem, &T::GetExtraStateIconMaxSize)
+				.MaxDesiredHeight(InTileOrListItem, &T::GetExtraStateIconMaxSize)
+				[
+					InTileOrListItem->GetAssetAccessSpecifierIconBox()
+				]
+			];
+		}
+
 		if (!UE::Editor::ContentBrowser::IsNewStyleEnabled())
 		{
 			// Dirty state
@@ -427,6 +443,9 @@ void SAssetViewItem::Construct( const FArguments& InArgs )
 	OnVisualizeAssetToolTip = InArgs._OnVisualizeAssetToolTip;
 	OnAssetToolTipClosing = InArgs._OnAssetToolTipClosing;
 	IsSelected = InArgs._IsSelected;
+	bShowAssetAccessSpecifier = InArgs._ShowAssetAccessSpecifier;
+
+	AssetAccessSpecifierIconBox = SNew(SBox);
 
 	bDraggedOver = false;
 
@@ -903,19 +922,7 @@ TSharedRef<SWidget> SAssetViewItem::CreateToolTipWidget() const
 				AddToToolTipInfoBox(InfoBox, LOCTEXT("TileViewTooltipPathLengthForCookingKey", "Cooking Filepath Length"), FText::Format(LOCTEXT("TileViewTooltipPathLengthForCookingValue", "{0} / {1}"),
 					FText::AsNumber(PackageNameLengthForCooking), FText::AsNumber(MaxCookPathLen)), PackageNameLengthForCooking > MaxCookPathLen ? true : false);
 
-				const EAssetAccessSpecifier AssetAccessSpecifier = ItemAssetData.GetAssetAccessSpecifier();
-				if (AssetAccessSpecifier == EAssetAccessSpecifier::Private)
-				{
-					PublicStateText = LOCTEXT("PrivateAssetState", "Private");
-				}
-				else if (AssetAccessSpecifier == EAssetAccessSpecifier::EpicInternal)
-				{
-					PublicStateText = LOCTEXT("EpicInternalAssetState", "Epic Internal");
-				}
-				else
-				{
-					PublicStateText = LOCTEXT("PublicAssetState", "Public");
-				}
+				PublicStateText = AssetItem->GetItemAssetAccessSpecifierText();
 			}
 
 			if (!AssetItem->GetItem().CanEdit())
@@ -1865,6 +1872,64 @@ void SAssetViewItem::OnToolTipClosing()
 	OnAssetToolTipClosing.ExecuteIfBound();
 }
 
+TSharedRef<SWidget> SAssetViewItem::GetAssetAccessSpecifierIconBox() const
+{
+	if (ensure(AssetAccessSpecifierIconBox.IsValid()))
+	{
+		return AssetAccessSpecifierIconBox.ToSharedRef();
+	}
+
+	return SNullWidget::NullWidget;
+}
+
+void SAssetViewItem::UpdateAssetAccessSpecifierIcon()
+{
+	TSharedRef<SWidget> IconWidget = SNullWidget::NullWidget;
+
+	// If there's a valid, modifiable EAssetAccessSpecifier, render its state. If there isn't one, or it isn't modifiable, it's pointless to show it.
+	EAssetAccessSpecifier AssetAccessSpecifier;
+	if (AssetItem.IsValid() && AssetItem->CanModifyItemAssetAccessSpecifier() && AssetItem->GetItemAssetAccessSpecifier(AssetAccessSpecifier))
+	{
+		switch (AssetAccessSpecifier)
+		{
+			case EAssetAccessSpecifier::Public:
+			{
+				IconWidget = SNew(SImage).Image(FAppStyle::GetBrush(TEXT("ContentBrowser.Migrate"))).ColorAndOpacity(FLinearColor::Red);
+				break;
+			}
+
+			case EAssetAccessSpecifier::Private:
+			{
+				IconWidget = SNew(SImage).Image(FAppStyle::GetBrush(TEXT("ContentBrowser.Private"))).ColorAndOpacity(FLinearColor::Red);
+				break;
+			}
+
+			case EAssetAccessSpecifier::EpicInternal:
+			{
+				IconWidget = SNew(SImage).Image(FAppStyle::GetBrush(TEXT("ContentBrowser.ContentDirty"))).ColorAndOpacity(FLinearColor::Red);
+				break;
+			}
+
+			default:
+			{
+				ensureMsgf(0, TEXT("Unrecognized EAssetAccessSpecifier value '%d' found for asset: '%s'"), AssetAccessSpecifier , *AssetItem->GetItem().GetItemName().ToString());
+				break;
+			}
+		}
+	}
+
+	if (ensure(AssetAccessSpecifierIconBox.IsValid()))
+	{
+		AssetAccessSpecifierIconBox->SetContent(IconWidget);
+	}
+}
+
+FText SAssetViewItem::GetAssetAccessSpecifierText() const
+{
+	return AssetItem.IsValid() ? AssetItem->GetItemAssetAccessSpecifierText()
+		: FText::GetEmpty();
+}
+
 ///////////////////////////////
 // SAssetListItem
 ///////////////////////////////
@@ -1885,6 +1950,7 @@ void SAssetListItem::Construct( const FArguments& InArgs )
 		.OnItemDestroyed(InArgs._OnItemDestroyed)
 		.ShouldAllowToolTip(InArgs._ShouldAllowToolTip)
 		.ThumbnailEditMode(InArgs._ThumbnailEditMode)
+		.ShowAssetAccessSpecifier(InArgs._ShowAssetAccessSpecifier)
 		.HighlightText(InArgs._HighlightText)
 		.OnIsAssetValidForCustomToolTip(InArgs._OnIsAssetValidForCustomToolTip)
 		.OnGetCustomAssetToolTip(InArgs._OnGetCustomAssetToolTip)
@@ -1984,6 +2050,13 @@ void SAssetListItem::OnAssetDataChanged()
 			AssetThumbnail->SetAsset(FAssetData());
 		}
 	}
+
+	UpdateAssetAccessSpecifierIcon();
+
+	if (AssetAccessSpecifierText.IsValid())
+	{
+		AssetAccessSpecifierText->SetText(GetAssetAccessSpecifierText());
+	}
 }
 
 void SAssetListItem::OnMouseEnter(const FGeometry& MyGeometry, const FPointerEvent& MouseEvent)
@@ -2097,7 +2170,7 @@ TSharedRef<SWidget> SAssetListItem::GenerateWidgetForColumn(const FName& ColumnN
 	if (ColumnName == "Name")
 	{
 		FName ItemShadowBorderName;
-		TSharedRef<SWidget> ItemContents = FAssetViewItemHelper::CreateListItemContents(this, ThumbnailWidget.ToSharedRef(), ItemShadowBorderName);
+		TSharedRef<SWidget> ItemContents = FAssetViewItemHelper::CreateListItemContents(this, ThumbnailWidget.ToSharedRef(), ItemShadowBorderName, 0.f, bShowAssetAccessSpecifier);
 
 		Content = SNew(SBorder)
 		.BorderImage(this, &SAssetViewItem::GetBorderImage)
@@ -2219,7 +2292,7 @@ TSharedRef<SWidget> SAssetListItem::GenerateWidgetForColumnNewStyle(const FName&
 	if (ColumnName == "Name")
 	{
 		FName ItemShadowBorderName;
-		TSharedRef<SWidget> ItemContents = FAssetViewItemHelper::CreateListItemContents(this, ThumbnailWidget.ToSharedRef(), ItemShadowBorderName);
+		TSharedRef<SWidget> ItemContents = FAssetViewItemHelper::CreateListItemContents(this, ThumbnailWidget.ToSharedRef(), ItemShadowBorderName, 0.f, bShowAssetAccessSpecifier);
 
 		const FSlateBrush* IconBrush;
 		if (IsFolder())
@@ -2370,6 +2443,21 @@ TSharedRef<SWidget> SAssetListItem::GenerateWidgetForColumnNewStyle(const FName&
 				.HighlightText(HighlightText)
 			];
 	}
+	else if (ColumnName == "AssetAccessSpecifier")
+	{
+		if(bShowAssetAccessSpecifier)
+		{
+			Content = SNew(SBox)
+				.Padding(DefaultColumnPadding)
+				.VAlign(VAlign_Center)
+				[
+					SAssignNew(AssetAccessSpecifierText, STextBlock)
+					.Text(GetAssetAccessSpecifierText())
+					.OverflowPolicy(ETextOverflowPolicy::MiddleEllipsis)
+					.HighlightText(HighlightText)
+				];
+		}
+	}
 	else
 	{
 		Content = SNew(SBox)
@@ -2434,6 +2522,7 @@ void SAssetTileItem::Construct( const FArguments& InArgs )
 		.OnVisualizeAssetToolTip(InArgs._OnVisualizeAssetToolTip)
 		.OnAssetToolTipClosing( InArgs._OnAssetToolTipClosing )
 		.IsSelected(InArgs._IsSelected)
+		.ShowAssetAccessSpecifier(InArgs._ShowAssetAccessSpecifier)
 		);
 
 	bShowType = InArgs._ShowType;
@@ -2509,7 +2598,7 @@ void SAssetTileItem::Construct( const FArguments& InArgs )
 	}
 
 	FName ItemShadowBorderName;
-	TSharedRef<SWidget> ItemContents = FAssetViewItemHelper::CreateTileItemContents(this, Thumbnail.ToSharedRef(), ItemShadowBorderName, FMargin(5.f));
+	TSharedRef<SWidget> ItemContents = FAssetViewItemHelper::CreateTileItemContents(this, Thumbnail.ToSharedRef(), ItemShadowBorderName, FMargin(5.f), InArgs._ShowAssetAccessSpecifier);
 
 	if (bShowType)
 	{
@@ -2794,6 +2883,8 @@ void SAssetTileItem::OnAssetDataChanged()
 			AssetThumbnail->SetAsset(FAssetData());
 		}
 	}
+
+	UpdateAssetAccessSpecifierIcon();
 }
 
 void SAssetTileItem::OnMouseEnter(const FGeometry& MyGeometry, const FPointerEvent& MouseEvent)
@@ -3263,6 +3354,7 @@ void SAssetColumnItem::Construct( const FArguments& InArgs )
 		.OnGetCustomAssetToolTip(InArgs._OnGetCustomAssetToolTip)
 		.OnVisualizeAssetToolTip(InArgs._OnVisualizeAssetToolTip)
 		.OnAssetToolTipClosing(InArgs._OnAssetToolTipClosing)
+		.ShowAssetAccessSpecifier(InArgs._ShowAssetAccessSpecifier)
 		);
 	
 	HighlightText = InArgs._HighlightText;
@@ -3409,6 +3501,16 @@ TSharedRef<SWidget> SAssetColumnItem::GenerateWidgetForColumn( const FName& Colu
 			.Text( GetAssetPathText() )
 			.HighlightText( HighlightText );
 	}
+	else if ( ColumnName == "AssetAccessSpecifier")
+	{
+		if (bShowAssetAccessSpecifier)
+		{
+			Content = SAssignNew(AssetAccessSpecifierText, STextBlock)
+				.ToolTipText( this, &SAssetColumnItem::GetAssetAccessSpecifierText )
+				.Text( GetAssetAccessSpecifierText() )
+				.HighlightText( HighlightText );
+		}
+	}
 	else if ( ColumnName == "RevisionControl" )
 	{
 		Content = SNew(SBox)
@@ -3450,6 +3552,11 @@ void SAssetColumnItem::OnAssetDataChanged()
 	{
 		PathText->SetText( GetAssetPathText() );
 	}
+
+	if ( AssetAccessSpecifierText.IsValid() )
+	{
+		AssetAccessSpecifierText->SetText( GetAssetAccessSpecifierText() );
+	}
 }
 
 FText SAssetColumnItem::GetAssetPathText() const
diff --git a/Engine/Source/Editor/ContentBrowser/Private/ContentBrowserConfig.h b/Engine/Source/Editor/ContentBrowser/Private/ContentBrowserConfig.h
--- a/Engine/Source/Editor/ContentBrowser/Private/ContentBrowserConfig.h
+++ b/Engine/Source/Editor/ContentBrowser/Private/ContentBrowserConfig.h
@@ -71,6 +71,9 @@ struct FContentBrowserInstanceConfig
 
 	UPROPERTY()
 	bool bSearchCollections = false;
+
+	UPROPERTY()
+	bool bShowAssetAccessSpecifier = false;
 };
 
 UCLASS(EditorConfig="ContentBrowser")
diff --git a/Engine/Source/Editor/ContentBrowser/Private/SAssetView.cpp b/Engine/Source/Editor/ContentBrowser/Private/SAssetView.cpp
--- a/Engine/Source/Editor/ContentBrowser/Private/SAssetView.cpp
+++ b/Engine/Source/Editor/ContentBrowser/Private/SAssetView.cpp
@@ -1207,6 +1207,7 @@ void SAssetView::Construct( const FArguments& InArgs )
 
 	bShowPathInColumnView = InArgs._ShowPathInColumnView;
 	bShowTypeInColumnView = InArgs._ShowTypeInColumnView;
+	bShowAssetAccessSpecifier = InArgs._ShowAssetAccessSpecifier;
 	bSortByPathInColumnView = bShowPathInColumnView && InArgs._SortByPathInColumnView;
 	bShowTypeInTileView = InArgs._ShowTypeInTileView;
 	bForceShowEngineContent = InArgs._ForceShowEngineContent;
@@ -2934,6 +2935,22 @@ TSharedRef<SAssetListView> SAssetView::CreateListView()
 		
 			NewListView->GetHeaderRow()->SetShowGeneratedColumn(SortManager->PathColumnId, bIsColumnVisible);
 		}
+
+		if (bShowAssetAccessSpecifier)
+		{
+			NewListView->GetHeaderRow()->AddColumn(
+				SHeaderRow::Column(SortManager->AssetAccessSpecifierColumnId)
+				.FillWidth(160)
+				.SortMode(TAttribute<EColumnSortMode::Type>::Create(TAttribute< EColumnSortMode::Type >::FGetter::CreateSP(this, &SAssetView::GetColumnSortMode, SortManager->AssetAccessSpecifierColumnId)))
+				.SortPriority(TAttribute<EColumnSortPriority::Type>::Create(TAttribute<EColumnSortPriority::Type>::FGetter::CreateSP(this, &SAssetView::GetColumnSortPriority, SortManager->AssetAccessSpecifierColumnId)))
+				.OnSort(FOnSortModeChanged::CreateSP(this, &SAssetView::OnSortColumnHeader))
+				.DefaultLabel(LOCTEXT("Column_AssetAccessSpecifier", "Asset Access Specifier"))
+			);
+
+			const bool bIsColumnVisible = !HiddenColumnNames.Contains(SortManager->AssetAccessSpecifierColumnId.ToString());
+
+			NewListView->GetHeaderRow()->SetShowGeneratedColumn(SortManager->AssetAccessSpecifierColumnId, bIsColumnVisible);
+		}
 	}
 	return NewListView.ToSharedRef();
 }
@@ -3024,6 +3041,22 @@ TSharedRef<SAssetColumnView> SAssetView::CreateColumnView()
 		NewColumnView->GetHeaderRow()->SetShowGeneratedColumn(SortManager->PathColumnId, bIsColumnVisible);
 	}
 
+	if (bShowAssetAccessSpecifier)
+	{
+		NewColumnView->GetHeaderRow()->AddColumn(
+			SHeaderRow::Column(SortManager->AssetAccessSpecifierColumnId)
+			.FillWidth(160)
+			.SortMode(TAttribute<EColumnSortMode::Type>::Create(TAttribute< EColumnSortMode::Type >::FGetter::CreateSP(this, &SAssetView::GetColumnSortMode, SortManager->AssetAccessSpecifierColumnId)))
+			.SortPriority(TAttribute<EColumnSortPriority::Type>::Create(TAttribute<EColumnSortPriority::Type>::FGetter::CreateSP(this, &SAssetView::GetColumnSortPriority, SortManager->AssetAccessSpecifierColumnId)))
+			.OnSort(FOnSortModeChanged::CreateSP(this, &SAssetView::OnSortColumnHeader))
+			.DefaultLabel(LOCTEXT("Column_AssetAccessSpecifier", "Asset Access Specifier"))
+		);
+
+		const bool bIsColumnVisible = !HiddenColumnNames.Contains(SortManager->AssetAccessSpecifierColumnId.ToString());
+
+		NewColumnView->GetHeaderRow()->SetShowGeneratedColumn(SortManager->AssetAccessSpecifierColumnId, bIsColumnVisible);
+	}
+
 	return NewColumnView.ToSharedRef();
 }
 
@@ -3862,7 +3895,8 @@ void SAssetView::SetMajorityAssetType(FName NewMajorityAssetType)
 		const bool bIsFixedRevisionControlColumn = InColumnId == SortManager->RevisionControlColumnId;
 		const bool bIsFixedClassColumn = bShowTypeInColumnView && InColumnId == SortManager->ClassColumnId;
 		const bool bIsFixedPathColumn = bShowPathInColumnView && InColumnId == SortManager->PathColumnId;
-		return bIsFixedNameColumn || bIsFixedRevisionControlColumn || bIsFixedClassColumn || bIsFixedPathColumn;
+		const bool bIsFixedAssetAccessSpecifierColumn = bShowAssetAccessSpecifier && InColumnId == SortManager->AssetAccessSpecifierColumnId;
+		return bIsFixedNameColumn || bIsFixedRevisionControlColumn || bIsFixedClassColumn || bIsFixedPathColumn || bIsFixedAssetAccessSpecifierColumn;
 	};
 
 
@@ -3908,7 +3942,8 @@ void SAssetView::SetMajorityAssetType(FName NewMajorityAssetType)
 			{
 				const bool bSortRelevant = SortColumn == FAssetViewSortManager::NameColumnId
 					|| SortColumn == FAssetViewSortManager::ClassColumnId
-					|| SortColumn == FAssetViewSortManager::PathColumnId;
+					|| SortColumn == FAssetViewSortManager::PathColumnId
+					|| SortColumn == FAssetViewSortManager::AssetAccessSpecifierColumnId;
 				CurrentSortOrder.Add(FSortOrder(bSortRelevant, SortColumn));
 			}
 		}
@@ -5090,6 +5125,7 @@ TSharedRef<ITableRow> SAssetView::MakeListViewWidget(TSharedPtr<FAssetViewItem>
 					.OnItemDestroyed(this, &SAssetView::AssetItemWidgetDestroyed)
 					.ShouldAllowToolTip(this, &SAssetView::ShouldAllowToolTips)
 					.HighlightText(HighlightedText)
+					.ShowAssetAccessSpecifier(bShowAssetAccessSpecifier)
 			);
 	}
 	else
@@ -5129,6 +5165,7 @@ TSharedRef<ITableRow> SAssetView::MakeListViewWidget(TSharedPtr<FAssetViewItem>
 					.OnGetCustomAssetToolTip(OnGetCustomAssetToolTip)
 					.OnVisualizeAssetToolTip(OnVisualizeAssetToolTip)
 					.OnAssetToolTipClosing(OnAssetToolTipClosing)
+					.ShowAssetAccessSpecifier(bShowAssetAccessSpecifier)
 			);
 	}
 }
@@ -5166,7 +5203,8 @@ TSharedRef<ITableRow> SAssetView::MakeTileViewWidget(TSharedPtr<FAssetViewItem>
 			.HighlightText( HighlightedText )
 			.IsSelected(FIsSelected::CreateSP(TableRowWidget.Get(), &STableRow<TSharedPtr<FAssetViewItem>>::IsSelected))
 			.IsSelectedExclusively(FIsSelected::CreateSP(TableRowWidget.Get(), &STableRow<TSharedPtr<FAssetViewItem>>::IsSelectedExclusively))
-			.AddMetaData<FTagMetaData>(AssetItem->GetItem().GetItemName());
+			.AddMetaData<FTagMetaData>(AssetItem->GetItem().GetItemName())
+			.ShowAssetAccessSpecifier(bShowAssetAccessSpecifier);
 
 		TableRowWidget->SetContent(Item);
 
@@ -5213,6 +5251,7 @@ TSharedRef<ITableRow> SAssetView::MakeTileViewWidget(TSharedPtr<FAssetViewItem>
 			.OnVisualizeAssetToolTip( OnVisualizeAssetToolTip )
 			.OnAssetToolTipClosing( OnAssetToolTipClosing )
 			.ShowType(bShowTypeInTileView)
+			.ShowAssetAccessSpecifier(bShowAssetAccessSpecifier)
 			.AddMetaData<FTagMetaData>(AssetItem->GetItem().GetItemName());
 		
 		TableRowWidget->SetContent(Item);
@@ -5248,6 +5287,7 @@ TSharedRef<ITableRow> SAssetView::MakeColumnViewWidget(TSharedPtr<FAssetViewItem
 				.OnGetCustomAssetToolTip(OnGetCustomAssetToolTip)
 				.OnVisualizeAssetToolTip( OnVisualizeAssetToolTip )
 				.OnAssetToolTipClosing( OnAssetToolTipClosing )
+				.ShowAssetAccessSpecifier(bShowAssetAccessSpecifier)
 		);
 }
 
diff --git a/Engine/Source/Editor/ContentBrowser/Private/SContentBrowser.cpp b/Engine/Source/Editor/ContentBrowser/Private/SContentBrowser.cpp
--- a/Engine/Source/Editor/ContentBrowser/Private/SContentBrowser.cpp
+++ b/Engine/Source/Editor/ContentBrowser/Private/SContentBrowser.cpp
@@ -3384,6 +3384,10 @@ TSharedRef<SWidget> SContentBrowser::CreateLegacyAssetViewWidgets()
 	// Currently this controls the asset count
 	const bool bShowBottomToolbar = InitConfig.bShowBottomToolbar;
 
+	const FContentBrowserInstanceConfig* Config = GetConstInstanceConfig();
+
+	const bool bShowAssetAccessSpecifier = Config ? Config->bShowAssetAccessSpecifier : false;
+	
 	LegacyContentSourceWidgets->AssetViewPtr = SNew(SAssetView)
 			.ThumbnailLabel(InitConfig.ThumbnailLabel)
 			//.ThumbnailScale(Config != nullptr ? Config->ThumbnailScale : 0.18f)
@@ -3412,7 +3416,8 @@ TSharedRef<SWidget> SContentBrowser::CreateLegacyAssetViewWidgets()
 			.bShowPathViewFilters(true)
 			.FillEmptySpaceInTileView(true)
 			.ShowDisallowedAssetClassAsUnsupportedItems(true)
-			.AllowCustomView(true);
+			.AllowCustomView(true)
+			.ShowAssetAccessSpecifier(bShowAssetAccessSpecifier);
 	
 	TSharedRef<SWidget> ViewOptions = SNullWidget::NullWidget;
 
@@ -3874,12 +3879,7 @@ TSharedRef<SWidget> SContentBrowser::CreateLegacyAssetViewWidgets()
 	// else
 	{
 		// in case we do not have a config, see what the global default settings are for the Sources Panel
-		bool bSourcesExpanded = true;
-		if (const FContentBrowserInstanceConfig* EditorConfig = GetConstInstanceConfig())
-		{
-			bSourcesExpanded = EditorConfig->bSourcesExpanded;
-		}
-	
+		const bool bSourcesExpanded = Config ? Config->bSourcesExpanded : true;
 		SetSourcesViewExpanded(bSourcesExpanded);
 	}
 	
diff --git a/Engine/Source/Editor/ContentBrowser/Public/AssetViewSortManager.h b/Engine/Source/Editor/ContentBrowser/Public/AssetViewSortManager.h
--- a/Engine/Source/Editor/ContentBrowser/Public/AssetViewSortManager.h
+++ b/Engine/Source/Editor/ContentBrowser/Public/AssetViewSortManager.h
@@ -55,6 +55,7 @@ public:
 	static const FName ClassColumnId;
 	static const FName PathColumnId;
 	static const FName DiskSizeColumnId;
+	static const FName AssetAccessSpecifierColumnId;
 
 	// The revision control column. NOTE: This column currently doesn't support sorting, but is wired through the Sort Manager so the feature can be added in the future
 	static const FName RevisionControlColumnId;
diff --git a/Engine/Source/Editor/ContentBrowser/Public/AssetViewTypes.h b/Engine/Source/Editor/ContentBrowser/Public/AssetViewTypes.h
--- a/Engine/Source/Editor/ContentBrowser/Public/AssetViewTypes.h
+++ b/Engine/Source/Editor/ContentBrowser/Public/AssetViewTypes.h
@@ -13,6 +13,7 @@
 
 #include <atomic>
 
+enum class EAssetAccessSpecifier : uint8;
 class FContentBrowserItemData;
 class FString;
 class FText;
@@ -72,6 +73,29 @@ public:
 	/** Get the value (and optionally also the type) of a named tag on this item */
 	bool GetTagValue(const FName Tag, FString& OutString, UObject::FAssetRegistryTag::ETagType* OutType = nullptr) const;
 
+	/**
+	 * Query the EAssetAccessSpecifier of the item.
+	 *
+	 * @param OutAssetAccessSpecifier The EAssetAccessSpecifier reference that will be filled.
+	 *
+	 * @return True if OutAssetAccessSpecifier was filled.
+	 */
+	bool GetItemAssetAccessSpecifier(EAssetAccessSpecifier& OutAssetAccessSpecifier) const;
+	
+	/**
+	 * Query a readable FText of the EAssetAccessSpecifier for the item.
+	 
+	 * @return If the item has a modifiable EAssetAccessSpecifier, then a readable text version. Otherwise blank.
+	 */
+	FText GetItemAssetAccessSpecifierText() const;
+
+	/**
+	 * Query the item's ability to modify its EAssetAccessSpecifier.
+	 *
+	 * @return True if the item has an EAssetAccessSpecifier that can be modified.
+	 */
+	bool CanModifyItemAssetAccessSpecifier() const;
+
 	/** Get the underlying Content Browser item */
 	CONTENTBROWSER_API const FContentBrowserItem& GetItem() const;
 
diff --git a/Engine/Source/Editor/ContentBrowser/Public/AssetViewWidgets.h b/Engine/Source/Editor/ContentBrowser/Public/AssetViewWidgets.h
--- a/Engine/Source/Editor/ContentBrowser/Public/AssetViewWidgets.h
+++ b/Engine/Source/Editor/ContentBrowser/Public/AssetViewWidgets.h
@@ -72,12 +72,12 @@ class SAssetTileItem;
 struct FAssetViewItemHelper
 {
 public:
-	static TSharedRef<SWidget> CreateListItemContents(SAssetListItem* const InListItem, const TSharedRef<SWidget>& InThumbnail, FName& OutItemShadowBorder, const FMargin& InFolderOverlayPadding = FMargin(0.f));
-	static TSharedRef<SWidget> CreateTileItemContents(SAssetTileItem* const InTileItem, const TSharedRef<SWidget>& InThumbnail, FName& OutItemShadowBorder, const FMargin& InFolderOverlayPadding = FMargin(0.f));
+	static TSharedRef<SWidget> CreateListItemContents(SAssetListItem* const InListItem, const TSharedRef<SWidget>& InThumbnail, FName& OutItemShadowBorder, const FMargin& InFolderOverlayPadding = FMargin(0.f), bool bShowAssetAccessSpecifier = false);
+	static TSharedRef<SWidget> CreateTileItemContents(SAssetTileItem* const InTileItem, const TSharedRef<SWidget>& InThumbnail, FName& OutItemShadowBorder, const FMargin& InFolderOverlayPadding = FMargin(0.f), bool bShowAssetAccessSpecifier = false);
 
 private:
 	template <typename T>
-	static TSharedRef<SWidget> CreateListTileItemContents(T* const InTileOrListItem, const TSharedRef<SWidget>& InThumbnail, FName& OutItemShadowBorder, const FMargin& InFolderOverlayPadding);
+	static TSharedRef<SWidget> CreateListTileItemContents(T* const InTileOrListItem, const TSharedRef<SWidget>& InThumbnail, FName& OutItemShadowBorder, const FMargin& InFolderOverlayPadding, bool bShowAssetAccessSpecifier);
 };
 
 /** The tile view mode of the asset view */
@@ -114,6 +114,7 @@ public:
 	SLATE_BEGIN_ARGS( SAssetViewItem )
 		: _ShouldAllowToolTip(true)
 		, _ThumbnailEditMode(false)
+		, _ShowAssetAccessSpecifier(false)
 		{}
 
 		/** Data for the asset this item represents */
@@ -137,6 +138,9 @@ public:
 		/** If true, display the thumbnail edit mode UI */
 		SLATE_ATTRIBUTE(bool, ThumbnailEditMode)
 
+		/** If true, display the item's EAssetAccessSpecifier */
+		SLATE_ARGUMENT(bool, ShowAssetAccessSpecifier)
+
 		/** The string in the title to highlight (used when searching by string) */
 		SLATE_ATTRIBUTE(FText, HighlightText)
 
@@ -271,6 +275,15 @@ protected:
 	/** If true it will expand the tooltip by default */
 	bool ShouldExpandTooltip() const;
 
+	/** Returns the SBox that will hold the AssetAccessSpecifierIcon */
+	TSharedRef<SWidget> GetAssetAccessSpecifierIconBox() const;
+
+	/** Renders the appropriate AssetAccessSpecifierIcon based on the EAssetAccessSpecifier of the asset */
+	void UpdateAssetAccessSpecifierIcon();
+
+	/** Gets the AssetAccessSpecifier as FText. */
+	FText GetAssetAccessSpecifierText() const;
+
 protected:
 	/** Data for a cached display tag for this item (used in the tooltip, and also as the display string in column views) */
 	struct FTagDisplayItem
@@ -366,6 +379,15 @@ protected:
 
 	/** Delegate handle for the HandleSourceControlStateChanged function callback */
 	FDelegateHandle SourceControlStateChangedDelegateHandle;
+
+	/** Whether to show the icon / text for the item's EAssetAccessSpecifier */
+	bool bShowAssetAccessSpecifier = false;
+
+	/** Widget holding the EAssetAccessSpecifier as an icon */
+	TSharedPtr<SBox> AssetAccessSpecifierIconBox;
+
+	/** Widget rendering the EAssetAccessSpecifier as text */
+	TSharedPtr<STextBlock> AssetAccessSpecifierText;
 };
 
 
@@ -384,6 +406,7 @@ public:
 		, _ShouldAllowToolTip(true)
 		, _ThumbnailEditMode(false)
 		, _AllowThumbnailHintLabel(false)
+		, _ShowAssetAccessSpecifier(false)
 		{}
 
 		/** The handle to the thumbnail this item should render */
@@ -431,6 +454,9 @@ public:
 		/** Whether the thumbnail should ever show it's hint label */
 		SLATE_ARGUMENT( bool, AllowThumbnailHintLabel )
 
+		/** If true, display the item's EAssetAccessSpecifier */
+		SLATE_ARGUMENT( bool, ShowAssetAccessSpecifier)
+
 		/** Whether the item is selected in the view */
 		SLATE_ARGUMENT(FIsSelected, IsSelected)
 
@@ -659,6 +685,7 @@ public:
 		, _ShouldAllowToolTip(true)
 		, _ShowType(true)
 		, _ThumbnailEditMode(false)
+		, _ShowAssetAccessSpecifier(false)
 	{}
 
 		/** The handle to the thumbnail this item should render */
@@ -712,6 +739,9 @@ public:
 		/** If true, the thumbnail in this item can be edited */
 		SLATE_ATTRIBUTE( bool, ThumbnailEditMode )
 
+		/** If true, display the item's EAssetAccessSpecifier */
+		SLATE_ARGUMENT(bool, ShowAssetAccessSpecifier)
+
 		/** Whether the item is selected in the view */
 		SLATE_ARGUMENT(FIsSelected, IsSelected)
 
@@ -849,11 +879,16 @@ private:
 class SAssetColumnItem : public SAssetViewItem
 {
 public:
-	SLATE_BEGIN_ARGS( SAssetColumnItem ) {}
+	SLATE_BEGIN_ARGS( SAssetColumnItem ) 
+		: _ShowAssetAccessSpecifier(false)
+		{}
 
 		/** Data for the asset this item represents */
 		SLATE_ARGUMENT( TSharedPtr<FAssetViewItem>, AssetItem )
 
+		/** If true, display the item's EAssetAccessSpecifier */
+		SLATE_ARGUMENT(bool, ShowAssetAccessSpecifier)
+
 		/** Delegate for when an asset name has entered a rename state */
 		SLATE_EVENT( FOnRenameBegin, OnRenameBegin )
 
diff --git a/Engine/Source/Editor/ContentBrowser/Public/SAssetView.h b/Engine/Source/Editor/ContentBrowser/Public/SAssetView.h
--- a/Engine/Source/Editor/ContentBrowser/Public/SAssetView.h
+++ b/Engine/Source/Editor/ContentBrowser/Public/SAssetView.h
@@ -136,6 +136,7 @@ public:
 		, _FillEmptySpaceInTileView(true)
 		, _ShowPathInColumnView(false)
 		, _ShowTypeInColumnView(true)
+		, _ShowAssetAccessSpecifier(false)
 		, _SortByPathInColumnView(false)
 		, _ShowTypeInTileView(true)
 		, _ForceShowEngineContent(false)
@@ -277,10 +278,13 @@ public:
 		/** Should show Type in column view if true */
 		SLATE_ARGUMENT(bool, ShowTypeInColumnView)
 
+		/** Should show EAssetAccessSpecifier if true */
+		SLATE_ARGUMENT(bool, ShowAssetAccessSpecifier)
+
 		/** Sort by path in the column view. Only works if the initial view type is Column */
 		SLATE_ARGUMENT(bool, SortByPathInColumnView)
 		
-		/** Should show Type in column view if true */
+		/** Should show Type in tile view if true */
 		SLATE_ARGUMENT(bool, ShowTypeInTileView)
 
 		/** Should always show engine content */
@@ -1240,6 +1244,9 @@ private:
 	/** If true, it will show type in the asset view */
 	bool bShowTypeInColumnView : 1;
 
+	/** If true, it will show EAssetAccessSpecifier */
+	bool bShowAssetAccessSpecifier : 1;
+
 	/** If true, it sorts by path and then name */
 	bool bSortByPathInColumnView : 1;
 
