Fixed asset dialog open button state getting out of sync with what is selected in the asset view.

diff --git a/Engine/Source/Editor/ContentBrowser/Private/SAssetDialog.cpp b/Engine/Source/Editor/ContentBrowser/Private/SAssetDialog.cpp
--- a/Engine/Source/Editor/ContentBrowser/Private/SAssetDialog.cpp
+++ b/Engine/Source/Editor/ContentBrowser/Private/SAssetDialog.cpp
@@ -1208,7 +1208,7 @@ bool SAssetDialog::IsConfirmButtonEnabled() const
 	switch (DialogType)
 	{
 	case EAssetDialogType::Open:
-		return CurrentlySelectedAssets.Num() > 0;
+		return GetCurrentSelectionDelegate.Execute().Num() > 0;
 	case EAssetDialogType::Save:
 		return bLastInputValidityCheckSuccessful;
 	default:
@@ -1252,8 +1252,6 @@ FReply SAssetDialog::OnCancelClicked()
 
 void SAssetDialog::OnAssetSelected(const FAssetData& AssetData)
 {
-	CurrentlySelectedAssets = GetCurrentSelectionDelegate.Execute();
-	
 	if ( AssetData.IsValid() )
 	{
 		SetCurrentlySelectedPath(AssetData.PackagePath.ToString(), EContentBrowserPathType::Internal);
diff --git a/Engine/Source/Editor/ContentBrowser/Private/SAssetDialog.h b/Engine/Source/Editor/ContentBrowser/Private/SAssetDialog.h
--- a/Engine/Source/Editor/ContentBrowser/Private/SAssetDialog.h
+++ b/Engine/Source/Editor/ContentBrowser/Private/SAssetDialog.h
@@ -234,9 +234,6 @@ private:
 	/** Fired when the asset dialog is cancelled or closed */
 	FOnAssetDialogCancelled OnAssetDialogCancelled;
 
-	/** The assets that are currently selected in the asset picker */
-	TArray<FAssetData> CurrentlySelectedAssets;
-
 	/** The name box. Only used in save dialogs. */
 	TSharedPtr<SEditableTextBox> NameEditableText;
 
