Fix for `SAssetDialog::SetCurrentlySelectedPath`

In the AssetBrowser, when going back in the breadcrumbs, the const-ref parameter `NewPath` becomes invalidated in `ContentBrowserUtils::UpdateNavigationBar`, but we are trying to use it just after the call to this function, crashing the engine.
This simple fix uses the local copy of the path which was made at the start of the function.

diff --git a/Engine/Source/Editor/ContentBrowser/Private/SAssetDialog.cpp b/Engine/Source/Editor/ContentBrowser/Private/SAssetDialog.cpp
--- a/Engine/Source/Editor/ContentBrowser/Private/SAssetDialog.cpp
+++ b/Engine/Source/Editor/ContentBrowser/Private/SAssetDialog.cpp
@@ -1182,13 +1182,14 @@ void SAssetDialog::SetCurrentlySelectedPath(const FString& NewPath, const EConte
 	// Update Navigation Bar
 	if (NavigationBar && AssetPicker && PathPicker)
 	{
+		// Calling this function will invalidate the const-ref parameter NewPath
 		ContentBrowserUtils::UpdateNavigationBar(NavigationBar, AssetPicker->GetAssetView(), PathPicker->GetPathView());
 	}
 	
 	UpdateInputValidity();
 
 	HistoryManager.AddHistoryData();
-	RecentDirectories.AddUnique(InPathType == EContentBrowserPathType::Virtual ? NewPath : NewVirtualPath.ToString());
+	RecentDirectories.AddUnique(InPathType == EContentBrowserPathType::Virtual ? CurrentlySelectedPath : NewVirtualPath.ToString());
 
 	OnPathSelected.ExecuteIfBound(NewPath);
 }
