Updated the placeholder iconography for AssetAccessSpecifiers to be a little more visually consistent
Disabled rendering the icon in List View. (Not needed since a column is displayed)

#rnx

diff --git a/Engine/Source/Editor/ContentBrowser/Private/AssetViewWidgets.cpp b/Engine/Source/Editor/ContentBrowser/Private/AssetViewWidgets.cpp
--- a/Engine/Source/Editor/ContentBrowser/Private/AssetViewWidgets.cpp
+++ b/Engine/Source/Editor/ContentBrowser/Private/AssetViewWidgets.cpp
@@ -136,9 +136,9 @@ bool IsTopLevelFolder(const FName InFolderPath)
 // FAssetViewModeUtils
 ///////////////////////////////
 
-TSharedRef<SWidget> FAssetViewItemHelper::CreateListItemContents(SAssetListItem* const InListItem, const TSharedRef<SWidget>& InThumbnail, FName& OutItemShadowBorder, const FMargin& InFolderOverlayPadding, bool bShowAssetAccessSpecifier)
+TSharedRef<SWidget> FAssetViewItemHelper::CreateListItemContents(SAssetListItem* const InListItem, const TSharedRef<SWidget>& InThumbnail, FName& OutItemShadowBorder, const FMargin& InFolderOverlayPadding)
 {
-	return CreateListTileItemContents(InListItem, InThumbnail, OutItemShadowBorder, InFolderOverlayPadding, bShowAssetAccessSpecifier);
+	return CreateListTileItemContents(InListItem, InThumbnail, OutItemShadowBorder, InFolderOverlayPadding, false);
 }
 
 TSharedRef<SWidget> FAssetViewItemHelper::CreateTileItemContents(SAssetTileItem* const InTileItem, const TSharedRef<SWidget>& InThumbnail, FName& OutItemShadowBorder, const FMargin& InFolderOverlayPadding, bool bShowAssetAccessSpecifier)
@@ -1890,23 +1890,25 @@ void SAssetViewItem::UpdateAssetAccessSpecifierIcon()
 	EAssetAccessSpecifier AssetAccessSpecifier;
 	if (AssetItem.IsValid() && AssetItem->CanModifyItemAssetAccessSpecifier() && AssetItem->GetItemAssetAccessSpecifier(AssetAccessSpecifier))
 	{
+		FName IconName;
+		
 		switch (AssetAccessSpecifier)
 		{
 			case EAssetAccessSpecifier::Public:
 			{
-				IconWidget = SNew(SImage).Image(FAppStyle::GetBrush(TEXT("ContentBrowser.Migrate"))).ColorAndOpacity(FLinearColor::Red);
+				IconName = TEXT("ContentBrowser.Shared");
 				break;
 			}
 
 			case EAssetAccessSpecifier::Private:
 			{
-				IconWidget = SNew(SImage).Image(FAppStyle::GetBrush(TEXT("ContentBrowser.Private"))).ColorAndOpacity(FLinearColor::Red);
+				IconName = TEXT("ContentBrowser.Private");
 				break;
 			}
 
 			case EAssetAccessSpecifier::EpicInternal:
 			{
-				IconWidget = SNew(SImage).Image(FAppStyle::GetBrush(TEXT("ContentBrowser.ContentDirty"))).ColorAndOpacity(FLinearColor::Red);
+				IconName = TEXT("ContentBrowser.Local");
 				break;
 			}
 
@@ -1916,6 +1918,30 @@ void SAssetViewItem::UpdateAssetAccessSpecifierIcon()
 				break;
 			}
 		}
+
+		if (IconName != NAME_None)
+		{
+			// Wrap the icon in an outer background layer to improve visibility on light thumbnails
+			const float Padding = 6.f;
+			const float OuterBoxSize = 16.f;
+			IconWidget = SNew(SBox)
+			.Padding(0, 0, Padding, Padding)
+			[
+				SNew(SBorder)
+				.BorderImage(FAppStyle::GetBrush("AssetThumbnail", ".AssetThumbnailStatusBar"))
+				[
+					SNew(SBox)
+					.WidthOverride(OuterBoxSize)
+					.HeightOverride(OuterBoxSize)
+					.HAlign(HAlign_Center)
+					.VAlign(VAlign_Center)
+					[
+						SNew(SImage)
+						.Image(FAppStyle::GetBrush(IconName))
+					]
+				]
+			];
+		}
 	}
 
 	if (ensure(AssetAccessSpecifierIconBox.IsValid()))
@@ -2171,7 +2197,7 @@ TSharedRef<SWidget> SAssetListItem::GenerateWidgetForColumn(const FName& ColumnN
 	if (ColumnName == "Name")
 	{
 		FName ItemShadowBorderName;
-		TSharedRef<SWidget> ItemContents = FAssetViewItemHelper::CreateListItemContents(this, ThumbnailWidget.ToSharedRef(), ItemShadowBorderName, 0.f, bShowAssetAccessSpecifier);
+		TSharedRef<SWidget> ItemContents = FAssetViewItemHelper::CreateListItemContents(this, ThumbnailWidget.ToSharedRef(), ItemShadowBorderName, 0.f);
 
 		Content = SNew(SBorder)
 		.BorderImage(this, &SAssetViewItem::GetBorderImage)
@@ -2293,7 +2319,7 @@ TSharedRef<SWidget> SAssetListItem::GenerateWidgetForColumnNewStyle(const FName&
 	if (ColumnName == "Name")
 	{
 		FName ItemShadowBorderName;
-		TSharedRef<SWidget> ItemContents = FAssetViewItemHelper::CreateListItemContents(this, ThumbnailWidget.ToSharedRef(), ItemShadowBorderName, 0.f, bShowAssetAccessSpecifier);
+		TSharedRef<SWidget> ItemContents = FAssetViewItemHelper::CreateListItemContents(this, ThumbnailWidget.ToSharedRef(), ItemShadowBorderName, 0.f);
 
 		const FSlateBrush* IconBrush;
 		if (IsFolder())
diff --git a/Engine/Source/Editor/ContentBrowser/Public/AssetViewWidgets.h b/Engine/Source/Editor/ContentBrowser/Public/AssetViewWidgets.h
--- a/Engine/Source/Editor/ContentBrowser/Public/AssetViewWidgets.h
+++ b/Engine/Source/Editor/ContentBrowser/Public/AssetViewWidgets.h
@@ -72,7 +72,7 @@ class SAssetTileItem;
 struct FAssetViewItemHelper
 {
 public:
-	static TSharedRef<SWidget> CreateListItemContents(SAssetListItem* const InListItem, const TSharedRef<SWidget>& InThumbnail, FName& OutItemShadowBorder, const FMargin& InFolderOverlayPadding = FMargin(0.f), bool bShowAssetAccessSpecifier = false);
+	static TSharedRef<SWidget> CreateListItemContents(SAssetListItem* const InListItem, const TSharedRef<SWidget>& InThumbnail, FName& OutItemShadowBorder, const FMargin& InFolderOverlayPadding = FMargin(0.f));
 	static TSharedRef<SWidget> CreateTileItemContents(SAssetTileItem* const InTileItem, const TSharedRef<SWidget>& InThumbnail, FName& OutItemShadowBorder, const FMargin& InFolderOverlayPadding = FMargin(0.f), bool bShowAssetAccessSpecifier = false);
 
 private:
