[ContentBrowser] PublicState in the CB tooltip shouldn't be hidden if the Icon is not set

diff --git a/Engine/Source/Editor/ContentBrowser/Private/AssetSystemContentBrowserInfoProvider.cpp b/Engine/Source/Editor/ContentBrowser/Private/AssetSystemContentBrowserInfoProvider.cpp
--- a/Engine/Source/Editor/ContentBrowser/Private/AssetSystemContentBrowserInfoProvider.cpp
+++ b/Engine/Source/Editor/ContentBrowser/Private/AssetSystemContentBrowserInfoProvider.cpp
@@ -179,7 +179,7 @@ void FAssetSystemContentBrowserInfoProvider::PopulateAssetInfo(TArray<FAssetDisp
 
 		// Restriction Info
 		FAssetDisplayInfo RestrictionInfo = FAssetDisplayInfo();
-		RestrictionInfo.IsVisible = PublicStateIcon && bIsPublicAssetUIEnabled && !PublicStateText.IsEmpty() ? EVisibility::Visible : EVisibility::Collapsed;
+		RestrictionInfo.IsVisible = bIsPublicAssetUIEnabled && !PublicStateText.IsEmpty() ? EVisibility::Visible : EVisibility::Collapsed;
 		RestrictionInfo.StatusIcon = PublicStateIcon;
 		RestrictionInfo.StatusTitle = LOCTEXT("Restriction", "Restriction");
 		RestrictionInfo.StatusDescription = PublicStateText;
diff --git a/Engine/Source/Editor/UnrealEd/Private/AssetThumbnail.cpp b/Engine/Source/Editor/UnrealEd/Private/AssetThumbnail.cpp
--- a/Engine/Source/Editor/UnrealEd/Private/AssetThumbnail.cpp
+++ b/Engine/Source/Editor/UnrealEd/Private/AssetThumbnail.cpp
@@ -1007,7 +1007,7 @@ private:
 			.Padding(0, 0, 4, 0)
 			[
 				SNew(SBox)
-				.Visibility(Icon.IsSet()? EVisibility::Visible : EVisibility::Collapsed)
+				.Visibility(Icon.IsSet() && Icon.Get() != nullptr? EVisibility::Visible : EVisibility::Collapsed)
 				.WidthOverride(16.f)
 				.HeightOverride(16.f)
 				[
