Updated headers using UnrealCodeFixup to make sure dllstorage is on methods/staticvars instead of types. Part 5/n

diff --git a/Engine/Source/Editor/CommonMenuExtensions/Public/ShowFlagFilter.h b/Engine/Source/Editor/CommonMenuExtensions/Public/ShowFlagFilter.h
--- a/Engine/Source/Editor/CommonMenuExtensions/Public/ShowFlagFilter.h
+++ b/Engine/Source/Editor/CommonMenuExtensions/Public/ShowFlagFilter.h
@@ -8,9 +8,11 @@
 #include "HAL/Platform.h"
 #include "ShowFlags.h"
 
+#define UE_API COMMONMENUEXTENSIONS_API
+
 class FShowFlagData;
 
-class COMMONMENUEXTENSIONS_API FShowFlagFilter
+class FShowFlagFilter
 {
 public:
 	enum EDefaultMode
@@ -60,7 +62,7 @@ public:
 		TArray<uint32> GroupedIndices[SFG_Max];
 	};
 
-	explicit FShowFlagFilter(EDefaultMode InDefaultMode);
+	UE_API explicit FShowFlagFilter(EDefaultMode InDefaultMode);
 
 	// Operations (earlier = higher priority):
 	// - Include all explicitly included flags.
@@ -68,16 +70,16 @@ public:
 	// - Include all flags in the include groups.
 	// - Exclude all flags in the exclude groups.
 	// - Include or exclude remaining flags depending on the default mode.
-	FShowFlagFilter& IncludeFlag(FEngineShowFlags::EShowFlag Flag);
-	FShowFlagFilter& ExcludeFlag(FEngineShowFlags::EShowFlag Flag);
-	FShowFlagFilter& IncludeGroup(EShowFlagGroup Group);
-	FShowFlagFilter& ExcludeGroup(EShowFlagGroup Group);
+	UE_API FShowFlagFilter& IncludeFlag(FEngineShowFlags::EShowFlag Flag);
+	UE_API FShowFlagFilter& ExcludeFlag(FEngineShowFlags::EShowFlag Flag);
+	UE_API FShowFlagFilter& IncludeGroup(EShowFlagGroup Group);
+	UE_API FShowFlagFilter& ExcludeGroup(EShowFlagGroup Group);
 
-	const FGroupedShowFlagIndices& GetFilteredIndices(bool bForceRebuild = false) const;
+	UE_API const FGroupedShowFlagIndices& GetFilteredIndices(bool bForceRebuild = false) const;
 
 private:
-	bool ShouldIncludeFlag(const FShowFlagData& Flag) const;
-	void RebuildIndexArray() const;
+	UE_API bool ShouldIncludeFlag(const FShowFlagData& Flag) const;
+	UE_API void RebuildIndexArray() const;
 
 	EDefaultMode DefaultMode;
 
@@ -93,4 +95,6 @@ private:
 	// These are indices into the GetShowFlagMenuItems() array.
 	mutable FGroupedShowFlagIndices AllowedFlagIndices;
 	mutable bool bHaveFiltered;
-};
\ No newline at end of file
+};
+
+#undef UE_API
diff --git a/Engine/Source/Editor/ComponentVisualizers/Public/AudioComponentVisualizer.h b/Engine/Source/Editor/ComponentVisualizers/Public/AudioComponentVisualizer.h
--- a/Engine/Source/Editor/ComponentVisualizers/Public/AudioComponentVisualizer.h
+++ b/Engine/Source/Editor/ComponentVisualizers/Public/AudioComponentVisualizer.h
@@ -5,7 +5,7 @@
 #include "AttenuatedComponentVisualizer.h"
 #include "Components/AudioComponent.h"
 
-class COMPONENTVISUALIZERS_API FAudioComponentVisualizer : public TAttenuatedComponentVisualizer<UAudioComponent>
+class FAudioComponentVisualizer : public TAttenuatedComponentVisualizer<UAudioComponent>
 {
 private:
 	virtual bool IsVisualizerEnabled(const FEngineShowFlags& ShowFlags) const override
diff --git a/Engine/Source/Editor/ComponentVisualizers/Public/ConstraintComponentVisualizer.h b/Engine/Source/Editor/ComponentVisualizers/Public/ConstraintComponentVisualizer.h
--- a/Engine/Source/Editor/ComponentVisualizers/Public/ConstraintComponentVisualizer.h
+++ b/Engine/Source/Editor/ComponentVisualizers/Public/ConstraintComponentVisualizer.h
@@ -5,14 +5,18 @@
 #include "ComponentVisualizer.h"
 #include "CoreMinimal.h"
 
+#define UE_API COMPONENTVISUALIZERS_API
+
 class FPrimitiveDrawInterface;
 class FSceneView;
 class UActorComponent;
 
-class COMPONENTVISUALIZERS_API FConstraintComponentVisualizer : public FComponentVisualizer
+class FConstraintComponentVisualizer : public FComponentVisualizer
 {
 public:
 	//~ Begin FComponentVisualizer Interface
-	virtual void DrawVisualization(const UActorComponent* Component, const FSceneView* View, FPrimitiveDrawInterface* PDI) override;
+	UE_API virtual void DrawVisualization(const UActorComponent* Component, const FSceneView* View, FPrimitiveDrawInterface* PDI) override;
 	//~ End FComponentVisualizer Interface
 };
+
+#undef UE_API
diff --git a/Engine/Source/Editor/ComponentVisualizers/Public/DecalComponentVisualizer.h b/Engine/Source/Editor/ComponentVisualizers/Public/DecalComponentVisualizer.h
--- a/Engine/Source/Editor/ComponentVisualizers/Public/DecalComponentVisualizer.h
+++ b/Engine/Source/Editor/ComponentVisualizers/Public/DecalComponentVisualizer.h
@@ -5,14 +5,18 @@
 #include "ComponentVisualizer.h"
 #include "CoreMinimal.h"
 
+#define UE_API COMPONENTVISUALIZERS_API
+
 class FPrimitiveDrawInterface;
 class FSceneView;
 class UActorComponent;
 
-class COMPONENTVISUALIZERS_API FDecalComponentVisualizer : public FComponentVisualizer
+class FDecalComponentVisualizer : public FComponentVisualizer
 {
 public:
 	//~ Begin FComponentVisualizer Interface
-	virtual void DrawVisualization(const UActorComponent* Component, const FSceneView* View, FPrimitiveDrawInterface* PDI) override;
+	UE_API virtual void DrawVisualization(const UActorComponent* Component, const FSceneView* View, FPrimitiveDrawInterface* PDI) override;
 	//~ End FComponentVisualizer Interface
 };
+
+#undef UE_API
diff --git a/Engine/Source/Editor/ComponentVisualizers/Public/ForceFeedbackComponentVisualizer.h b/Engine/Source/Editor/ComponentVisualizers/Public/ForceFeedbackComponentVisualizer.h
--- a/Engine/Source/Editor/ComponentVisualizers/Public/ForceFeedbackComponentVisualizer.h
+++ b/Engine/Source/Editor/ComponentVisualizers/Public/ForceFeedbackComponentVisualizer.h
@@ -5,7 +5,7 @@
 #include "AttenuatedComponentVisualizer.h"
 #include "Components/ForceFeedbackComponent.h"
 
-class COMPONENTVISUALIZERS_API FForceFeedbackComponentVisualizer : public TAttenuatedComponentVisualizer<UForceFeedbackComponent>
+class FForceFeedbackComponentVisualizer : public TAttenuatedComponentVisualizer<UForceFeedbackComponent>
 {
 private:
 	virtual bool IsVisualizerEnabled(const FEngineShowFlags& ShowFlags) const override
diff --git a/Engine/Source/Editor/ComponentVisualizers/Public/LocalFogVolumeComponentVisualizer.h b/Engine/Source/Editor/ComponentVisualizers/Public/LocalFogVolumeComponentVisualizer.h
--- a/Engine/Source/Editor/ComponentVisualizers/Public/LocalFogVolumeComponentVisualizer.h
+++ b/Engine/Source/Editor/ComponentVisualizers/Public/LocalFogVolumeComponentVisualizer.h
@@ -5,14 +5,18 @@
 #include "CoreMinimal.h"
 #include "ComponentVisualizer.h"
 
+#define UE_API COMPONENTVISUALIZERS_API
+
 class FPrimitiveDrawInterface;
 class FSceneView;
 class UActorComponent;
 
-class COMPONENTVISUALIZERS_API FLocalFogVolumeComponentVisualizer : public FComponentVisualizer
+class FLocalFogVolumeComponentVisualizer : public FComponentVisualizer
 {
 public:
 	//~ Begin FComponentVisualizer Interface
-	virtual void DrawVisualization(const UActorComponent* Component, const FSceneView* View, FPrimitiveDrawInterface* PDI) override;
+	UE_API virtual void DrawVisualization(const UActorComponent* Component, const FSceneView* View, FPrimitiveDrawInterface* PDI) override;
 	//~ End FComponentVisualizer Interface
 };
+
+#undef UE_API
diff --git a/Engine/Source/Editor/ComponentVisualizers/Public/Manipulator.h b/Engine/Source/Editor/ComponentVisualizers/Public/Manipulator.h
--- a/Engine/Source/Editor/ComponentVisualizers/Public/Manipulator.h
+++ b/Engine/Source/Editor/ComponentVisualizers/Public/Manipulator.h
@@ -11,38 +11,40 @@
 
 #include "Manipulator.generated.h"
 
+#define UE_API COMPONENTVISUALIZERS_API
+
 class UActorComponent;
 class UObject;
 class USceneComponent;
 class UViewportInteractor;
 struct FHitResult;
 
-UCLASS()
-class COMPONENTVISUALIZERS_API AManipulator : public AActor, public IViewportInteractableInterface
+UCLASS(MinimalAPI)
+class AManipulator : public AActor, public IViewportInteractableInterface
 {
 	GENERATED_BODY()
 
 public:
 
-	AManipulator();
+	UE_API AManipulator();
 
 	// Begin AActor
-	virtual void PostEditMove(bool bFinished) override;
-	virtual bool IsEditorOnly() const final;
+	UE_API virtual void PostEditMove(bool bFinished) override;
+	UE_API virtual bool IsEditorOnly() const final;
 	// End AActor
 
 	// Begin IViewportInteractableInterface
-	virtual void OnPressed(UViewportInteractor* Interactor, const FHitResult& InHitResult, bool& bOutResultedInDrag) override;
-	virtual void OnHover(UViewportInteractor* Interactor) override;
-	virtual void OnHoverEnter(UViewportInteractor* Interactor, const FHitResult& InHitResult) override;
-	virtual void OnHoverLeave(UViewportInteractor* Interactor, const UActorComponent* NewComponent) override;
-	virtual void OnDragRelease(UViewportInteractor* Interactor) override;
+	UE_API virtual void OnPressed(UViewportInteractor* Interactor, const FHitResult& InHitResult, bool& bOutResultedInDrag) override;
+	UE_API virtual void OnHover(UViewportInteractor* Interactor) override;
+	UE_API virtual void OnHoverEnter(UViewportInteractor* Interactor, const FHitResult& InHitResult) override;
+	UE_API virtual void OnHoverLeave(UViewportInteractor* Interactor, const UActorComponent* NewComponent) override;
+	UE_API virtual void OnDragRelease(UViewportInteractor* Interactor) override;
 	virtual class UViewportDragOperationComponent* GetDragOperationComponent() override { return nullptr; }
 	virtual bool CanBeSelected() override { return true; };
 	// End IViewportInteractableInterface
 
 	/** Set the component that should be moved when the manipulator was moved. */
-	void SetAssociatedComponent(USceneComponent* SceneComponent);
+	UE_API void SetAssociatedComponent(USceneComponent* SceneComponent);
 
 private:
 
@@ -54,4 +56,6 @@ private:
 	UPROPERTY()
 	TObjectPtr<class UStaticMeshComponent> StaticMeshComponent;
 
-};
\ No newline at end of file
+};
+
+#undef UE_API
diff --git a/Engine/Source/Editor/ComponentVisualizers/Public/PhysicalAnimationComponentVisualizer.h b/Engine/Source/Editor/ComponentVisualizers/Public/PhysicalAnimationComponentVisualizer.h
--- a/Engine/Source/Editor/ComponentVisualizers/Public/PhysicalAnimationComponentVisualizer.h
+++ b/Engine/Source/Editor/ComponentVisualizers/Public/PhysicalAnimationComponentVisualizer.h
@@ -5,14 +5,18 @@
 #include "ComponentVisualizer.h"
 #include "CoreMinimal.h"
 
+#define UE_API COMPONENTVISUALIZERS_API
+
 class FPrimitiveDrawInterface;
 class FSceneView;
 class UActorComponent;
 
-class COMPONENTVISUALIZERS_API FPhysicsAnimationComponentVisualizer : public FComponentVisualizer
+class FPhysicsAnimationComponentVisualizer : public FComponentVisualizer
 {
 public:
 	//~ Begin FComponentVisualizer Interface
-	virtual void DrawVisualization(const UActorComponent* Component, const FSceneView* View, FPrimitiveDrawInterface* PDI) override;
+	UE_API virtual void DrawVisualization(const UActorComponent* Component, const FSceneView* View, FPrimitiveDrawInterface* PDI) override;
 	//~ End FComponentVisualizer Interface
 };
+
+#undef UE_API
diff --git a/Engine/Source/Editor/ComponentVisualizers/Public/PointLightComponentVisualizer.h b/Engine/Source/Editor/ComponentVisualizers/Public/PointLightComponentVisualizer.h
--- a/Engine/Source/Editor/ComponentVisualizers/Public/PointLightComponentVisualizer.h
+++ b/Engine/Source/Editor/ComponentVisualizers/Public/PointLightComponentVisualizer.h
@@ -7,6 +7,8 @@
 #include "CoreMinimal.h"
 #include "Math/Transform.h"
 
+#define UE_API COMPONENTVISUALIZERS_API
+
 class FPrimitiveDrawInterface;
 class FSceneView;
 class UActorComponent;
@@ -26,13 +28,15 @@ private:
 	TArray< float > IntensitiesCache;
 };
 
-class COMPONENTVISUALIZERS_API FPointLightComponentVisualizer : public FComponentVisualizer
+class FPointLightComponentVisualizer : public FComponentVisualizer
 {
 public:
 	//~ Begin FComponentVisualizer Interface
-	virtual void DrawVisualization(const UActorComponent* Component, const FSceneView* View, FPrimitiveDrawInterface* PDI) override;
+	UE_API virtual void DrawVisualization(const UActorComponent* Component, const FSceneView* View, FPrimitiveDrawInterface* PDI) override;
 	//~ End FComponentVisualizer Interface
 
 private:
 	FTextureLightProfileVisualizer LightProfileVisualizer;
 };
+
+#undef UE_API
diff --git a/Engine/Source/Editor/ComponentVisualizers/Public/RadialForceComponentVisualizer.h b/Engine/Source/Editor/ComponentVisualizers/Public/RadialForceComponentVisualizer.h
--- a/Engine/Source/Editor/ComponentVisualizers/Public/RadialForceComponentVisualizer.h
+++ b/Engine/Source/Editor/ComponentVisualizers/Public/RadialForceComponentVisualizer.h
@@ -5,14 +5,18 @@
 #include "ComponentVisualizer.h"
 #include "CoreMinimal.h"
 
+#define UE_API COMPONENTVISUALIZERS_API
+
 class FPrimitiveDrawInterface;
 class FSceneView;
 class UActorComponent;
 
-class COMPONENTVISUALIZERS_API FRadialForceComponentVisualizer : public FComponentVisualizer
+class FRadialForceComponentVisualizer : public FComponentVisualizer
 {
 public:
 	//~ Begin FComponentVisualizer Interface
-	virtual void DrawVisualization( const UActorComponent* Component, const FSceneView* View, FPrimitiveDrawInterface* PDI ) override;
+	UE_API virtual void DrawVisualization( const UActorComponent* Component, const FSceneView* View, FPrimitiveDrawInterface* PDI ) override;
 	//~ End FComponentVisualizer Interface
 };
+
+#undef UE_API
diff --git a/Engine/Source/Editor/ComponentVisualizers/Public/RectLightComponentVisualizer.h b/Engine/Source/Editor/ComponentVisualizers/Public/RectLightComponentVisualizer.h
--- a/Engine/Source/Editor/ComponentVisualizers/Public/RectLightComponentVisualizer.h
+++ b/Engine/Source/Editor/ComponentVisualizers/Public/RectLightComponentVisualizer.h
@@ -5,14 +5,18 @@
 #include "ComponentVisualizer.h"
 #include "CoreMinimal.h"
 
+#define UE_API COMPONENTVISUALIZERS_API
+
 class FPrimitiveDrawInterface;
 class FSceneView;
 class UActorComponent;
 
-class COMPONENTVISUALIZERS_API FRectLightComponentVisualizer : public FComponentVisualizer
+class FRectLightComponentVisualizer : public FComponentVisualizer
 {
 public:
 	//~ Begin FComponentVisualizer Interface
-	virtual void DrawVisualization(const UActorComponent* Component, const FSceneView* View, FPrimitiveDrawInterface* PDI) override;
+	UE_API virtual void DrawVisualization(const UActorComponent* Component, const FSceneView* View, FPrimitiveDrawInterface* PDI) override;
 	//~ End FComponentVisualizer Interface
 };
+
+#undef UE_API
diff --git a/Engine/Source/Editor/ComponentVisualizers/Public/SensingComponentVisualizer.h b/Engine/Source/Editor/ComponentVisualizers/Public/SensingComponentVisualizer.h
--- a/Engine/Source/Editor/ComponentVisualizers/Public/SensingComponentVisualizer.h
+++ b/Engine/Source/Editor/ComponentVisualizers/Public/SensingComponentVisualizer.h
@@ -5,14 +5,18 @@
 #include "ComponentVisualizer.h"
 #include "CoreMinimal.h"
 
+#define UE_API COMPONENTVISUALIZERS_API
+
 class FPrimitiveDrawInterface;
 class FSceneView;
 class UActorComponent;
 
-class COMPONENTVISUALIZERS_API FSensingComponentVisualizer : public FComponentVisualizer
+class FSensingComponentVisualizer : public FComponentVisualizer
 {
 public:
 	//~ Begin FComponentVisualizer Interface
-	virtual void DrawVisualization(const UActorComponent* Component, const FSceneView* View, FPrimitiveDrawInterface* PDI) override;
+	UE_API virtual void DrawVisualization(const UActorComponent* Component, const FSceneView* View, FPrimitiveDrawInterface* PDI) override;
 	//~ End FComponentVisualizer Interface
 };
+
+#undef UE_API
diff --git a/Engine/Source/Editor/ComponentVisualizers/Public/SplineMeshComponentVisualizer.h b/Engine/Source/Editor/ComponentVisualizers/Public/SplineMeshComponentVisualizer.h
--- a/Engine/Source/Editor/ComponentVisualizers/Public/SplineMeshComponentVisualizer.h
+++ b/Engine/Source/Editor/ComponentVisualizers/Public/SplineMeshComponentVisualizer.h
@@ -13,6 +13,8 @@
 #include "Math/UnrealMathSSE.h"
 #include "Templates/SharedPointer.h"
 
+#define UE_API COMPONENTVISUALIZERS_API
+
 class AActor;
 class FEditorViewportClient;
 class FPrimitiveDrawInterface;
@@ -26,7 +28,7 @@ struct FViewportClick;
 /** Base class for clickable spline mesh component editing proxies */
 struct HSplineMeshVisProxy : public HComponentVisProxy
 {
-	DECLARE_HIT_PROXY();
+	DECLARE_HIT_PROXY( UE_API );
 
 	HSplineMeshVisProxy(const UActorComponent* InComponent)
 	: HComponentVisProxy(InComponent, HPP_Wireframe)
@@ -36,7 +38,7 @@ struct HSplineMeshVisProxy : public HComponentVisProxy
 /** Proxy for a spline mesh component key */
 struct HSplineMeshKeyProxy : public HSplineMeshVisProxy
 {
-	DECLARE_HIT_PROXY();
+	DECLARE_HIT_PROXY( UE_API );
 
 	HSplineMeshKeyProxy(const UActorComponent* InComponent, int32 InKeyIndex) 
 		: HSplineMeshVisProxy(InComponent)
@@ -49,7 +51,7 @@ struct HSplineMeshKeyProxy : public HSplineMeshVisProxy
 /** Proxy for a tangent handle */
 struct HSplineMeshTangentHandleProxy : public HSplineMeshVisProxy
 {
-	DECLARE_HIT_PROXY();
+	DECLARE_HIT_PROXY( UE_API );
 
 	HSplineMeshTangentHandleProxy(const UActorComponent* InComponent, int32 InKeyIndex, bool bInArriveTangent)
 		: HSplineMeshVisProxy(InComponent)
@@ -62,34 +64,34 @@ struct HSplineMeshTangentHandleProxy : public HSplineMeshVisProxy
 };
 
 /** SplineMeshComponent visualizer/edit functionality */
-class COMPONENTVISUALIZERS_API FSplineMeshComponentVisualizer : public FComponentVisualizer
+class FSplineMeshComponentVisualizer : public FComponentVisualizer
 {
 public:
-	FSplineMeshComponentVisualizer();
-	virtual ~FSplineMeshComponentVisualizer();
+	UE_API FSplineMeshComponentVisualizer();
+	UE_API virtual ~FSplineMeshComponentVisualizer();
 
 	//~ Begin FComponentVisualizer Interface
-	virtual void OnRegister() override;
-	virtual void DrawVisualization(const UActorComponent* Component, const FSceneView* View, FPrimitiveDrawInterface* PDI) override;
-	virtual bool VisProxyHandleClick(FEditorViewportClient* InViewportClient, HComponentVisProxy* VisProxy, const FViewportClick& Click) override;
-	virtual void EndEditing() override;
-	virtual bool GetWidgetLocation(const FEditorViewportClient* ViewportClient, FVector& OutLocation) const override;
-	virtual bool GetCustomInputCoordinateSystem(const FEditorViewportClient* ViewportClient, FMatrix& OutMatrix) const override;
-	virtual bool HandleInputDelta(FEditorViewportClient* ViewportClient, FViewport* Viewport, FVector& DeltaTranslate, FRotator& DeltaRotate, FVector& DeltaScale) override;
-	virtual bool HandleInputKey(FEditorViewportClient* ViewportClient, FViewport* Viewport, FKey Key, EInputEvent Event) override;
-	virtual TSharedPtr<SWidget> GenerateContextMenu() const override;
+	UE_API virtual void OnRegister() override;
+	UE_API virtual void DrawVisualization(const UActorComponent* Component, const FSceneView* View, FPrimitiveDrawInterface* PDI) override;
+	UE_API virtual bool VisProxyHandleClick(FEditorViewportClient* InViewportClient, HComponentVisProxy* VisProxy, const FViewportClick& Click) override;
+	UE_API virtual void EndEditing() override;
+	UE_API virtual bool GetWidgetLocation(const FEditorViewportClient* ViewportClient, FVector& OutLocation) const override;
+	UE_API virtual bool GetCustomInputCoordinateSystem(const FEditorViewportClient* ViewportClient, FMatrix& OutMatrix) const override;
+	UE_API virtual bool HandleInputDelta(FEditorViewportClient* ViewportClient, FViewport* Viewport, FVector& DeltaTranslate, FRotator& DeltaRotate, FVector& DeltaScale) override;
+	UE_API virtual bool HandleInputKey(FEditorViewportClient* ViewportClient, FViewport* Viewport, FKey Key, EInputEvent Event) override;
+	UE_API virtual TSharedPtr<SWidget> GenerateContextMenu() const override;
 	//~ End FComponentVisualizer Interface
 
 	/** Get the spline component we are currently editing */
-	USplineMeshComponent* GetEditedSplineMeshComponent() const;
+	UE_API USplineMeshComponent* GetEditedSplineMeshComponent() const;
 
 protected:
 
 	/** Get a spline object for the specified spline mesh component */
-	FInterpCurveVector GetSpline(const USplineMeshComponent* SplineMeshComp) const;
+	UE_API FInterpCurveVector GetSpline(const USplineMeshComponent* SplineMeshComp) const;
 
 	/** Syncs changes made by the visualizer in the actual component */
-	void NotifyComponentModified();
+	UE_API void NotifyComponentModified();
 
 	/** Property path from the parent actor to the component */
 	FComponentPropertyPath SplineMeshPropertyPath;
@@ -113,3 +115,5 @@ protected:
 	/** The type of the selected tangent handle */
 	ESelectedTangentHandle::Type SelectedTangentHandleType;
 };
+
+#undef UE_API
diff --git a/Engine/Source/Editor/ComponentVisualizers/Public/SpotLightComponentVisualizer.h b/Engine/Source/Editor/ComponentVisualizers/Public/SpotLightComponentVisualizer.h
--- a/Engine/Source/Editor/ComponentVisualizers/Public/SpotLightComponentVisualizer.h
+++ b/Engine/Source/Editor/ComponentVisualizers/Public/SpotLightComponentVisualizer.h
@@ -6,17 +6,21 @@
 #include "CoreMinimal.h"
 #include "PointLightComponentVisualizer.h"
 
+#define UE_API COMPONENTVISUALIZERS_API
+
 class FPrimitiveDrawInterface;
 class FSceneView;
 class UActorComponent;
 
-class COMPONENTVISUALIZERS_API FSpotLightComponentVisualizer : public FComponentVisualizer
+class FSpotLightComponentVisualizer : public FComponentVisualizer
 {
 public:
 	//~ Begin FComponentVisualizer Interface
-	virtual void DrawVisualization(const UActorComponent* Component, const FSceneView* View, FPrimitiveDrawInterface* PDI) override;
+	UE_API virtual void DrawVisualization(const UActorComponent* Component, const FSceneView* View, FPrimitiveDrawInterface* PDI) override;
 	//~ End FComponentVisualizer Interface
 
 private:
 	FTextureLightProfileVisualizer LightProfileVisualizer;
 };
+
+#undef UE_API
diff --git a/Engine/Source/Editor/ComponentVisualizers/Public/SpringArmComponentVisualizer.h b/Engine/Source/Editor/ComponentVisualizers/Public/SpringArmComponentVisualizer.h
--- a/Engine/Source/Editor/ComponentVisualizers/Public/SpringArmComponentVisualizer.h
+++ b/Engine/Source/Editor/ComponentVisualizers/Public/SpringArmComponentVisualizer.h
@@ -5,14 +5,18 @@
 #include "ComponentVisualizer.h"
 #include "CoreMinimal.h"
 
+#define UE_API COMPONENTVISUALIZERS_API
+
 class FPrimitiveDrawInterface;
 class FSceneView;
 class UActorComponent;
 
-class COMPONENTVISUALIZERS_API FSpringArmComponentVisualizer : public FComponentVisualizer
+class FSpringArmComponentVisualizer : public FComponentVisualizer
 {
 public:
 	//~ Begin FComponentVisualizer Interface
-	virtual void DrawVisualization(const UActorComponent* Component, const FSceneView* View, FPrimitiveDrawInterface* PDI) override;
+	UE_API virtual void DrawVisualization(const UActorComponent* Component, const FSceneView* View, FPrimitiveDrawInterface* PDI) override;
 	//~ End FComponentVisualizer Interface
 };
+
+#undef UE_API
diff --git a/Engine/Source/Editor/ComponentVisualizers/Public/SpringComponentVisualizer.h b/Engine/Source/Editor/ComponentVisualizers/Public/SpringComponentVisualizer.h
--- a/Engine/Source/Editor/ComponentVisualizers/Public/SpringComponentVisualizer.h
+++ b/Engine/Source/Editor/ComponentVisualizers/Public/SpringComponentVisualizer.h
@@ -5,14 +5,18 @@
 #include "ComponentVisualizer.h"
 #include "CoreMinimal.h"
 
+#define UE_API COMPONENTVISUALIZERS_API
+
 class FPrimitiveDrawInterface;
 class FSceneView;
 class UActorComponent;
 
-class COMPONENTVISUALIZERS_API FSpringComponentVisualizer : public FComponentVisualizer
+class FSpringComponentVisualizer : public FComponentVisualizer
 {
 public:
 	//~ Begin FComponentVisualizer Interface
-	virtual void DrawVisualization(const UActorComponent* Component, const FSceneView* View, FPrimitiveDrawInterface* PDI) override;
+	UE_API virtual void DrawVisualization(const UActorComponent* Component, const FSceneView* View, FPrimitiveDrawInterface* PDI) override;
 	//~ End FComponentVisualizer Interface
 };
+
+#undef UE_API
diff --git a/Engine/Source/Editor/ComponentVisualizers/Public/StereoLayerComponentVisualizer.h b/Engine/Source/Editor/ComponentVisualizers/Public/StereoLayerComponentVisualizer.h
--- a/Engine/Source/Editor/ComponentVisualizers/Public/StereoLayerComponentVisualizer.h
+++ b/Engine/Source/Editor/ComponentVisualizers/Public/StereoLayerComponentVisualizer.h
@@ -5,14 +5,18 @@
 #include "ComponentVisualizer.h"
 #include "CoreMinimal.h"
 
+#define UE_API COMPONENTVISUALIZERS_API
+
 class FPrimitiveDrawInterface;
 class FSceneView;
 class UActorComponent;
 
-class COMPONENTVISUALIZERS_API FStereoLayerComponentVisualizer : public FComponentVisualizer
+class FStereoLayerComponentVisualizer : public FComponentVisualizer
 {
 public:
 	//~ Begin FComponentVisualizer Interface
-	virtual void DrawVisualization(const UActorComponent* Component, const FSceneView* View, FPrimitiveDrawInterface* PDI) override;
+	UE_API virtual void DrawVisualization(const UActorComponent* Component, const FSceneView* View, FPrimitiveDrawInterface* PDI) override;
 	//~ End FComponentVisualizer Interface
 };
+
+#undef UE_API
diff --git a/Engine/Source/Editor/ComponentVisualizers/Public/WorldPartitionStreamingSourceComponentVisualizer.h b/Engine/Source/Editor/ComponentVisualizers/Public/WorldPartitionStreamingSourceComponentVisualizer.h
--- a/Engine/Source/Editor/ComponentVisualizers/Public/WorldPartitionStreamingSourceComponentVisualizer.h
+++ b/Engine/Source/Editor/ComponentVisualizers/Public/WorldPartitionStreamingSourceComponentVisualizer.h
@@ -5,14 +5,18 @@
 #include "ComponentVisualizer.h"
 #include "CoreMinimal.h"
 
+#define UE_API COMPONENTVISUALIZERS_API
+
 class FPrimitiveDrawInterface;
 class FSceneView;
 class UActorComponent;
 
-class COMPONENTVISUALIZERS_API FWorldPartitionStreamingSourceComponentVisualizer : public FComponentVisualizer
+class FWorldPartitionStreamingSourceComponentVisualizer : public FComponentVisualizer
 {
 public:
 	//~ Begin FComponentVisualizer Interface
-	virtual void DrawVisualization(const UActorComponent* Component, const FSceneView* View, FPrimitiveDrawInterface* PDI) override;
+	UE_API virtual void DrawVisualization(const UActorComponent* Component, const FSceneView* View, FPrimitiveDrawInterface* PDI) override;
 	//~ End FComponentVisualizer Interface
 };
+
+#undef UE_API
diff --git a/Engine/Source/Editor/ConfigEditor/Public/PropertyVisualization/ConfigPropertyColumn.h b/Engine/Source/Editor/ConfigEditor/Public/PropertyVisualization/ConfigPropertyColumn.h
--- a/Engine/Source/Editor/ConfigEditor/Public/PropertyVisualization/ConfigPropertyColumn.h
+++ b/Engine/Source/Editor/ConfigEditor/Public/PropertyVisualization/ConfigPropertyColumn.h
@@ -9,6 +9,8 @@
 #include "Templates/SharedPointer.h"
 #include "UObject/NameTypes.h"
 
+#define UE_API CONFIGEDITOR_API
+
 class FProperty;
 class IPropertyTableCell;
 class IPropertyTableCellPresenter;
@@ -19,17 +21,19 @@ class SWidget;
 /**
 * A property table custom column used to display source control condition of files.
 */
-class CONFIGEDITOR_API FConfigPropertyCustomColumn : public IPropertyTableCustomColumn
+class FConfigPropertyCustomColumn : public IPropertyTableCustomColumn
 {
 public:
 	FConfigPropertyCustomColumn(){}
 
 	/** Begin IPropertyTableCustomColumn interface */
-	virtual bool Supports(const TSharedRef< IPropertyTableColumn >& Column, const TSharedRef< IPropertyTableUtilities >& Utilities) const override;
-	virtual TSharedPtr< SWidget > CreateColumnLabel(const TSharedRef< IPropertyTableColumn >& Column, const TSharedRef< IPropertyTableUtilities >& Utilities, const FName& Style) const override;
-	virtual TSharedPtr< IPropertyTableCellPresenter > CreateCellPresenter(const TSharedRef< IPropertyTableCell >& Cell, const TSharedRef< IPropertyTableUtilities >& Utilities, const FName& Style) const override;
+	UE_API virtual bool Supports(const TSharedRef< IPropertyTableColumn >& Column, const TSharedRef< IPropertyTableUtilities >& Utilities) const override;
+	UE_API virtual TSharedPtr< SWidget > CreateColumnLabel(const TSharedRef< IPropertyTableColumn >& Column, const TSharedRef< IPropertyTableUtilities >& Utilities, const FName& Style) const override;
+	UE_API virtual TSharedPtr< IPropertyTableCellPresenter > CreateCellPresenter(const TSharedRef< IPropertyTableCell >& Cell, const TSharedRef< IPropertyTableUtilities >& Utilities, const FName& Style) const override;
 	/** End IPropertyTableCustomColumn interface */
 
 	/* The property type which can be displayed in this column */
 	FProperty* EditProperty;
 };
+
+#undef UE_API
diff --git a/Engine/Source/Editor/ConfigEditor/Public/PropertyVisualization/ConfigPropertyConfigFileStateColumn.h b/Engine/Source/Editor/ConfigEditor/Public/PropertyVisualization/ConfigPropertyConfigFileStateColumn.h
--- a/Engine/Source/Editor/ConfigEditor/Public/PropertyVisualization/ConfigPropertyConfigFileStateColumn.h
+++ b/Engine/Source/Editor/ConfigEditor/Public/PropertyVisualization/ConfigPropertyConfigFileStateColumn.h
@@ -12,6 +12,8 @@
 #include "Templates/SharedPointer.h"
 #include "UObject/NameTypes.h"
 
+#define UE_API CONFIGEDITOR_API
+
 class FProperty;
 class IPropertyHandle;
 class IPropertyTableCell;
@@ -53,17 +55,19 @@ protected:
 /**
 * A property table custom column used to display source control condition of files.
 */
-class CONFIGEDITOR_API FConfigPropertyConfigFileStateCustomColumn : public IPropertyTableCustomColumn
+class FConfigPropertyConfigFileStateCustomColumn : public IPropertyTableCustomColumn
 {
 public:
 	FConfigPropertyConfigFileStateCustomColumn(){}
 
 	/** Begin IPropertyTableCustomColumn interface */
-	virtual bool Supports(const TSharedRef< IPropertyTableColumn >& Column, const TSharedRef< IPropertyTableUtilities >& Utilities) const override;
-	virtual TSharedPtr< SWidget > CreateColumnLabel(const TSharedRef< IPropertyTableColumn >& Column, const TSharedRef< IPropertyTableUtilities >& Utilities, const FName& Style) const override;
-	virtual TSharedPtr< IPropertyTableCellPresenter > CreateCellPresenter(const TSharedRef< IPropertyTableCell >& Cell, const TSharedRef< IPropertyTableUtilities >& Utilities, const FName& Style) const override;
+	UE_API virtual bool Supports(const TSharedRef< IPropertyTableColumn >& Column, const TSharedRef< IPropertyTableUtilities >& Utilities) const override;
+	UE_API virtual TSharedPtr< SWidget > CreateColumnLabel(const TSharedRef< IPropertyTableColumn >& Column, const TSharedRef< IPropertyTableUtilities >& Utilities, const FName& Style) const override;
+	UE_API virtual TSharedPtr< IPropertyTableCellPresenter > CreateCellPresenter(const TSharedRef< IPropertyTableCell >& Cell, const TSharedRef< IPropertyTableUtilities >& Utilities, const FName& Style) const override;
 	/** End IPropertyTableCustomColumn interface */
 
 	/* The property type which can be displayed in this column */
 	FProperty* SupportedProperty;
 };
+
+#undef UE_API
diff --git a/Engine/Source/Editor/ContentBrowser/Private/AssetView/AssetViewConfig.h b/Engine/Source/Editor/ContentBrowser/Private/AssetView/AssetViewConfig.h
--- a/Engine/Source/Editor/ContentBrowser/Private/AssetView/AssetViewConfig.h
+++ b/Engine/Source/Editor/ContentBrowser/Private/AssetView/AssetViewConfig.h
@@ -10,6 +10,8 @@
 
 #include "AssetViewConfig.generated.h"
 
+#define UE_API CONTENTBROWSER_API
+
 USTRUCT()
 struct FAssetViewInstanceConfig
 {
@@ -34,20 +36,22 @@ struct FAssetViewInstanceConfig
 	TArray<FName> ListHiddenColumns;
 };
 
-UCLASS(EditorConfig="AssetView")
-class CONTENTBROWSER_API UAssetViewConfig : public UEditorConfigBase
+UCLASS(MinimalAPI, EditorConfig="AssetView")
+class UAssetViewConfig : public UEditorConfigBase
 {
 	GENERATED_BODY()
 
 public:
-	static void Initialize();
+	static UE_API void Initialize();
 	static UAssetViewConfig* Get() { return Instance; }
 
-	FAssetViewInstanceConfig& GetInstanceConfig(FName ViewName);
+	UE_API FAssetViewInstanceConfig& GetInstanceConfig(FName ViewName);
 	
 	UPROPERTY(meta=(EditorConfig))
 	TMap<FName, FAssetViewInstanceConfig> Instances;
 	
 private:
-	static TObjectPtr<UAssetViewConfig> Instance;
+	static UE_API TObjectPtr<UAssetViewConfig> Instance;
 };
+
+#undef UE_API
diff --git a/Engine/Source/Editor/ContentBrowser/Public/AssetTextFilter.h b/Engine/Source/Editor/ContentBrowser/Public/AssetTextFilter.h
--- a/Engine/Source/Editor/ContentBrowser/Public/AssetTextFilter.h
+++ b/Engine/Source/Editor/ContentBrowser/Public/AssetTextFilter.h
@@ -68,7 +68,7 @@ public:
 	CONTENTBROWSER_API bool GetIncludeCollectionNames() const;
 
 	/** Delegate to bind for when the effective text filter changes so filtering can be re-run. */
-	CONTENTBROWSER_API FSimpleMulticastDelegate& OnChanged() { return ChangedEvent; }
+	FSimpleMulticastDelegate& OnChanged() { return ChangedEvent; }
 
 private:
 	/** Handles an on collection container created event */
@@ -243,7 +243,7 @@ public:
 	 * @return true If this handler provided an answer for whether the item matches. bOutIsMatch may be true or false in this case.
 	 * @return false If this handler did not decide whether this item was filtered, and other handlers or the built-in implementation should continue.
 	 */
-	CONTENTBROWSER_API virtual bool HandleTextFilterValue(const FContentBrowserItem& InContentBrowserItem,
+	virtual bool HandleTextFilterValue(const FContentBrowserItem& InContentBrowserItem,
 		const FTextFilterString& InValue,
 		const ETextFilterTextComparisonMode InTextComparisonMode,
 		bool& bOutIsMatch) = 0;
@@ -260,7 +260,7 @@ public:
 	 * @return true If this handler provided an answer for whether the item matches. bOutIsMatch may be true or false in this case.
 	 * @return false If this handler did not decide whether this item was filtered, and other handlers or the built-in implementation should continue.
 	 */
-	CONTENTBROWSER_API virtual bool HandleTextFilterKeyValue(const FContentBrowserItem& InContentBrowserItem,
+	virtual bool HandleTextFilterKeyValue(const FContentBrowserItem& InContentBrowserItem,
 		const FName& InKey,
 		const FTextFilterString& InValue,
 		const ETextFilterComparisonOperation InComparisonOperation,
diff --git a/Engine/Source/Editor/ContentBrowser/Public/ContentBrowserFrontEndFilterExtension.h b/Engine/Source/Editor/ContentBrowser/Public/ContentBrowserFrontEndFilterExtension.h
--- a/Engine/Source/Editor/ContentBrowser/Public/ContentBrowserFrontEndFilterExtension.h
+++ b/Engine/Source/Editor/ContentBrowser/Public/ContentBrowserFrontEndFilterExtension.h
@@ -9,8 +9,8 @@
 #include "ContentBrowserFrontEndFilterExtension.generated.h"
 
 // Override this class in order to make an additional front-end filter available in the Content Browser
-UCLASS(Abstract)
-class CONTENTBROWSER_API UContentBrowserFrontEndFilterExtension : public UObject
+UCLASS(MinimalAPI, Abstract)
+class UContentBrowserFrontEndFilterExtension : public UObject
 {
 	GENERATED_BODY()
 
diff --git a/Engine/Source/Editor/ContentBrowser/Public/ContentBrowserMenuContexts.h b/Engine/Source/Editor/ContentBrowser/Public/ContentBrowserMenuContexts.h
--- a/Engine/Source/Editor/ContentBrowser/Public/ContentBrowserMenuContexts.h
+++ b/Engine/Source/Editor/ContentBrowser/Public/ContentBrowserMenuContexts.h
@@ -16,6 +16,8 @@
 
 #include "ContentBrowserMenuContexts.generated.h"
 
+#define UE_API CONTENTBROWSER_API
+
 class FAssetViewSortManager;
 class FAssetContextMenu;
 class IAssetTypeActions;
@@ -35,8 +37,8 @@ enum class EIncludeSubclasses : uint8
 	Yes
 };
 
-UCLASS()
-class CONTENTBROWSER_API UContentBrowserAssetContextMenuContext : public UObject
+UCLASS(MinimalAPI)
+class UContentBrowserAssetContextMenuContext : public UObject
 {
 	GENERATED_BODY()
 
@@ -105,7 +107,7 @@ public:
 	 * will likely be deprecated a few versions after GetSelectedObjects.
 	 */
 	UFUNCTION(BlueprintCallable, Category="Tool Menus")
-	TArray<UObject*> LoadSelectedObjectsIfNeeded() const;
+	UE_API TArray<UObject*> LoadSelectedObjectsIfNeeded() const;
 
 	/**
 	 * Loads all the selected assets and returns an array of the objects.
@@ -184,14 +186,14 @@ public:
 	/**
 	 * Returns a filtered array of assets that are of the desired class and potentially any subclasses.
 	 */
-	TArray<FAssetData> GetSelectedAssetsOfType(const UClass* AssetClass, EIncludeSubclasses IncludeSubclasses = EIncludeSubclasses::Yes) const;
+	UE_API TArray<FAssetData> GetSelectedAssetsOfType(const UClass* AssetClass, EIncludeSubclasses IncludeSubclasses = EIncludeSubclasses::Yes) const;
 	
 	/**
 	 * Sometimes you want to write actions that will only operate on a singular selected asset, in those cases you
 	 * can use the following function which will only return a live ptr if it's an instance of that asset type, and
 	 * only one thing is selected.
 	 */
-	const FAssetData* GetSingleSelectedAssetOfType(const UClass* AssetClass, EIncludeSubclasses IncludeSubclasses = EIncludeSubclasses::Yes) const;
+	UE_API const FAssetData* GetSingleSelectedAssetOfType(const UClass* AssetClass, EIncludeSubclasses IncludeSubclasses = EIncludeSubclasses::Yes) const;
 
 	/**
 	 * Get the selected assets as an array of SoftObjectPtr<T>'s.
@@ -220,7 +222,7 @@ public:
 	}
 
 	/** Read-only access to the list of currently selected items */
-	const TArray<FContentBrowserItem>& GetSelectedItems() const;
+	UE_API const TArray<FContentBrowserItem>& GetSelectedItems() const;
 
 	/**
 	 * Finds the Content Browser MenuContext from a Menu or Section, and returns the context provided there are some
@@ -262,8 +264,8 @@ public:
 	}
 };
 
-UCLASS()
-class CONTENTBROWSER_API UContentBrowserAssetViewContextMenuContext : public UObject
+UCLASS(MinimalAPI)
+class UContentBrowserAssetViewContextMenuContext : public UObject
 {
 	GENERATED_BODY()
 
@@ -272,8 +274,8 @@ public:
 	TWeakPtr<SAssetView> AssetView;
 };
 
-UCLASS()
-class CONTENTBROWSER_API UContentBrowserAssetSortingContextMenuContext : public UObject
+UCLASS(MinimalAPI)
+class UContentBrowserAssetSortingContextMenuContext : public UObject
 {
 	GENERATED_BODY()
 
@@ -283,8 +285,8 @@ public:
 	TWeakPtr<FAssetViewSortManager> AssetViewSortManager;
 };
 
-UCLASS()
-class CONTENTBROWSER_API UContentBrowserMenuContext : public UObject
+UCLASS(MinimalAPI)
+class UContentBrowserMenuContext : public UObject
 {
 	GENERATED_BODY()
 
@@ -293,8 +295,8 @@ public:
 	TWeakPtr<SContentBrowser> ContentBrowser;
 };
 
-UCLASS()
-class CONTENTBROWSER_API UContentBrowserFolderContext : public UContentBrowserMenuContext
+UCLASS(MinimalAPI)
+class UContentBrowserFolderContext : public UContentBrowserMenuContext
 {
 	GENERATED_BODY()
 
@@ -320,8 +322,8 @@ public:
 	const TArray<FString>& GetSelectedPackagePaths() const { return SelectedPackagePaths; }
 };
 
-UCLASS()
-class CONTENTBROWSER_API UContentBrowserAddNewContextMenuContext : public UObject
+UCLASS(MinimalAPI)
+class UContentBrowserAddNewContextMenuContext : public UObject
 {
 	GENERATED_BODY()
 
@@ -330,15 +332,15 @@ public:
 	TWeakPtr<SContentBrowser> ContentBrowser;
 };
 
-UCLASS()
-class CONTENTBROWSER_API UContentBrowserToolbarMenuContext : public UObject
+UCLASS(MinimalAPI)
+class UContentBrowserToolbarMenuContext : public UObject
 {
 	GENERATED_BODY()
 
 public:
-	FName GetCurrentPath() const;
+	UE_API FName GetCurrentPath() const;
 
-	bool CanWriteToCurrentPath() const;
+	UE_API bool CanWriteToCurrentPath() const;
 
 	TWeakPtr<SContentBrowser> ContentBrowser;
 	TWeakPtr<SAssetView> AssetView;
@@ -352,4 +354,6 @@ namespace UE::ContentBrowser
 	CONTENTBROWSER_API UToolMenu* ExtendToolMenu_AssetContextMenu(UClass* AssetClass);
 	CONTENTBROWSER_API UToolMenu* ExtendToolMenu_AssetContextMenu(TSoftClassPtr<UObject> AssetSoftClass);
 	CONTENTBROWSER_API TArray<UToolMenu*> ExtendToolMenu_AssetContextMenus(TConstArrayView<UClass*> AssetClasses);
-}
\ No newline at end of file
+}
+
+#undef UE_API
diff --git a/Engine/Source/Editor/ContentBrowser/Public/FrontendFilters.h b/Engine/Source/Editor/ContentBrowser/Public/FrontendFilters.h
--- a/Engine/Source/Editor/ContentBrowser/Public/FrontendFilters.h
+++ b/Engine/Source/Editor/ContentBrowser/Public/FrontendFilters.h
@@ -24,6 +24,8 @@
 #include "Types/SlateEnums.h"
 #include "UObject/NameTypes.h"
 
+#define UE_API CONTENTBROWSER_API
+
 class UPackage;
 struct FAssetRenameData;
 struct FCollectionNameType;
@@ -36,11 +38,11 @@ struct FAssetCompileData;
 
 /** A filter for text search */
 UE_DEPRECATED(5.5, "FFrontendFilter_Text has been deprecated in favor of FAssetTextFilter, used by SAssetView to perform text filtering in parallel.");
-class CONTENTBROWSER_API FFrontendFilter_Text : public FFrontendFilter
+class FFrontendFilter_Text : public FFrontendFilter
 {
 public:
-	FFrontendFilter_Text();
-	~FFrontendFilter_Text();
+	UE_API FFrontendFilter_Text();
+	UE_API ~FFrontendFilter_Text();
 
 	// FFrontendFilter implementation
 	virtual FString GetName() const override { return TEXT("TextFilter"); }
@@ -48,54 +50,54 @@ public:
 	virtual FText GetToolTipText() const override { return LOCTEXT("FrontendFilter_TextTooltip", "Show only assets that match the input text"); }
 
 	// IFilter implementation
-	virtual bool PassesFilter(FAssetFilterType InItem) const override;
+	UE_API virtual bool PassesFilter(FAssetFilterType InItem) const override;
 
 public:
 	/** Returns the unsanitized and unsplit filter terms */
-	FText GetRawFilterText() const;
+	UE_API FText GetRawFilterText() const;
 
 	/** Set the Text to be used as the Filter's restrictions */
-	void SetRawFilterText(const FText& InFilterText);
+	UE_API void SetRawFilterText(const FText& InFilterText);
 
 	/** Get the last error returned from lexing or compiling the current filter text */
-	FText GetFilterErrorText() const;
+	UE_API FText GetFilterErrorText() const;
 
 	/** If bIncludeClassName is true, the text filter will include an asset's class name in the search */
-	void SetIncludeClassName(const bool InIncludeClassName);
+	UE_API void SetIncludeClassName(const bool InIncludeClassName);
 
 	/** If bIncludeAssetPath is true, the text filter will match against full Asset path */
-	void SetIncludeAssetPath(const bool InIncludeAssetPath);
+	UE_API void SetIncludeAssetPath(const bool InIncludeAssetPath);
 
-	bool GetIncludeAssetPath() const;
+	UE_API bool GetIncludeAssetPath() const;
 
 	/** If bIncludeCollectionNames is true, the text filter will match against collection names as well */
-	void SetIncludeCollectionNames(const bool InIncludeCollectionNames);
+	UE_API void SetIncludeCollectionNames(const bool InIncludeCollectionNames);
 
-	bool GetIncludeCollectionNames() const;
+	UE_API bool GetIncludeCollectionNames() const;
 private:
 	/** Handles an on collection container created event */
-	void HandleCollectionContainerCreated(const TSharedRef<ICollectionContainer>& CollectionContainer);
+	UE_API void HandleCollectionContainerCreated(const TSharedRef<ICollectionContainer>& CollectionContainer);
 
 	/** Handles an on collection container destroyed event */
-	void HandleCollectionContainerDestroyed(const TSharedRef<ICollectionContainer>& CollectionContainer);
+	UE_API void HandleCollectionContainerDestroyed(const TSharedRef<ICollectionContainer>& CollectionContainer);
 
 	/** Handles an on collection container is hidden changed event */
-	void HandleIsHiddenChanged(ICollectionContainer& CollectionContainer, bool bIsHidden);
+	UE_API void HandleIsHiddenChanged(ICollectionContainer& CollectionContainer, bool bIsHidden);
 
 	/** Handles an on collection created event */
-	void HandleCollectionCreated(ICollectionContainer& CollectionContainer, const FCollectionNameType& Collection);
+	UE_API void HandleCollectionCreated(ICollectionContainer& CollectionContainer, const FCollectionNameType& Collection);
 
 	/** Handles an on collection destroyed event */
-	void HandleCollectionDestroyed(ICollectionContainer& CollectionContainer, const FCollectionNameType& Collection);
+	UE_API void HandleCollectionDestroyed(ICollectionContainer& CollectionContainer, const FCollectionNameType& Collection);
 
 	/** Handles an on collection renamed event */
-	void HandleCollectionRenamed(ICollectionContainer& CollectionContainer, const FCollectionNameType& OriginalCollection, const FCollectionNameType& NewCollection);
+	UE_API void HandleCollectionRenamed(ICollectionContainer& CollectionContainer, const FCollectionNameType& OriginalCollection, const FCollectionNameType& NewCollection);
 
 	/** Handles an on collection updated event */
-	void HandleCollectionUpdated(ICollectionContainer& CollectionContainer, const FCollectionNameType& Collection);
+	UE_API void HandleCollectionUpdated(ICollectionContainer& CollectionContainer, const FCollectionNameType& Collection);
 
 	/** Rebuild the array of dynamic collections that are being referenced by the current query */
-	void RebuildReferencedDynamicCollections();
+	UE_API void RebuildReferencedDynamicCollections();
 
 	/** An array of collection containers that are being monitored */
 	TArray<TSharedPtr<ICollectionContainer>> CollectionContainers;
@@ -125,55 +127,55 @@ private:
 };
 
 /** A filter that displays only checked out assets */
-class CONTENTBROWSER_API FFrontendFilter_CheckedOut : public FFrontendFilter, public TSharedFromThis<FFrontendFilter_CheckedOut>
+class FFrontendFilter_CheckedOut : public FFrontendFilter, public TSharedFromThis<FFrontendFilter_CheckedOut>
 {
 public:
-	FFrontendFilter_CheckedOut(TSharedPtr<FFrontendFilterCategory> InCategory);
+	UE_API FFrontendFilter_CheckedOut(TSharedPtr<FFrontendFilterCategory> InCategory);
 
 	// FFrontendFilter implementation
 	virtual FString GetName() const override { return TEXT("CheckedOut"); }
 	virtual FText GetDisplayName() const override { return LOCTEXT("FrontendFilter_CheckedOut", "Checked Out"); }
 	virtual FText GetToolTipText() const override { return LOCTEXT("FrontendFilter_CheckedOutTooltip", "Show only assets that you have checked out or pending for add."); }
-	virtual void ActiveStateChanged(bool bActive) override;
-	virtual void SetCurrentFilter(TArrayView<const FName> InSourcePaths, const FContentBrowserDataFilter& InBaseFilter) override;
+	UE_API virtual void ActiveStateChanged(bool bActive) override;
+	UE_API virtual void SetCurrentFilter(TArrayView<const FName> InSourcePaths, const FContentBrowserDataFilter& InBaseFilter) override;
 
 	// IFilter implementation
-	virtual bool PassesFilter(FAssetFilterType InItem) const override;
+	UE_API virtual bool PassesFilter(FAssetFilterType InItem) const override;
 
 private:
 	
 	/** Request the source control status for this filter */
-	void RequestStatus();
+	UE_API void RequestStatus();
 
 	/** Callback when source control operation has completed */
-	void SourceControlOperationComplete(const FSourceControlOperationRef& InOperation, ECommandResult::Type InResult);
+	UE_API void SourceControlOperationComplete(const FSourceControlOperationRef& InOperation, ECommandResult::Type InResult);
 
 	bool bSourceControlEnabled;
 };
 
 /** A filter that displays assets not tracked by source control */
-class CONTENTBROWSER_API FFrontendFilter_NotSourceControlled : public FFrontendFilter, public TSharedFromThis<FFrontendFilter_NotSourceControlled>
+class FFrontendFilter_NotSourceControlled : public FFrontendFilter, public TSharedFromThis<FFrontendFilter_NotSourceControlled>
 {
 public:
-	FFrontendFilter_NotSourceControlled(TSharedPtr<FFrontendFilterCategory> InCategory);
+	UE_API FFrontendFilter_NotSourceControlled(TSharedPtr<FFrontendFilterCategory> InCategory);
 
 	// FFrontendFilter implementation
 	virtual FString GetName() const override { return TEXT("NotSourceControlled"); }
 	virtual FText GetDisplayName() const override { return LOCTEXT("FrontendFilter_NotSourceControlled", "Not Revision Controlled"); }
 	virtual FText GetToolTipText() const override { return LOCTEXT("FrontendFilter_NotSourceControlledTooltip", "Show only assets that are not tracked by revision control."); }
-	virtual void ActiveStateChanged(bool bActive) override;
-	virtual void SetCurrentFilter(TArrayView<const FName> InSourcePaths, const FContentBrowserDataFilter& InBaseFilter) override;
+	UE_API virtual void ActiveStateChanged(bool bActive) override;
+	UE_API virtual void SetCurrentFilter(TArrayView<const FName> InSourcePaths, const FContentBrowserDataFilter& InBaseFilter) override;
 
 	// IFilter implementation
-	virtual bool PassesFilter(FAssetFilterType InItem) const override;
+	UE_API virtual bool PassesFilter(FAssetFilterType InItem) const override;
 
 private:
 
 	/** Request the source control status for this filter */
-	void RequestStatus();
+	UE_API void RequestStatus();
 
 	/** Callback when source control operation has completed */
-	void SourceControlOperationComplete(const FSourceControlOperationRef& InOperation, ECommandResult::Type InResult);
+	UE_API void SourceControlOperationComplete(const FSourceControlOperationRef& InOperation, ECommandResult::Type InResult);
 
 	bool bSourceControlEnabled;
 	bool bIsRequestStatusRunning;
@@ -181,31 +183,31 @@ private:
 };
 
 /** A filter that displays only modified assets */
-class CONTENTBROWSER_API FFrontendFilter_Modified : public FFrontendFilter
+class FFrontendFilter_Modified : public FFrontendFilter
 {
 public:
-	FFrontendFilter_Modified(TSharedPtr<FFrontendFilterCategory> InCategory);
-	~FFrontendFilter_Modified();
+	UE_API FFrontendFilter_Modified(TSharedPtr<FFrontendFilterCategory> InCategory);
+	UE_API ~FFrontendFilter_Modified();
 
 	// FFrontendFilter implementation
 	virtual FString GetName() const override { return TEXT("Modified"); }
 	virtual FText GetDisplayName() const override { return LOCTEXT("FrontendFilter_Modified", "Modified"); }
 	virtual FText GetToolTipText() const override { return LOCTEXT("FrontendFilter_ModifiedTooltip", "Show only assets that have been modified and not yet saved."); }
-	virtual void ActiveStateChanged(bool bActive) override;
+	UE_API virtual void ActiveStateChanged(bool bActive) override;
 
 	// IFilter implementation
-	virtual bool PassesFilter(FAssetFilterType InItem) const override;
+	UE_API virtual bool PassesFilter(FAssetFilterType InItem) const override;
 
 private:
 
 	/** Handler for when a package's dirty state has changed */
-	void OnPackageDirtyStateUpdated(UPackage* Package);
+	UE_API void OnPackageDirtyStateUpdated(UPackage* Package);
 
 	bool bIsCurrentlyActive;
 };
 
 /** A filter that displays blueprints that have replicated properties */
-class CONTENTBROWSER_API FFrontendFilter_ReplicatedBlueprint : public FFrontendFilter
+class FFrontendFilter_ReplicatedBlueprint : public FFrontendFilter
 {
 public:
 	FFrontendFilter_ReplicatedBlueprint(TSharedPtr<FFrontendFilterCategory> InCategory) : FFrontendFilter(InCategory) {}
@@ -216,36 +218,36 @@ public:
 	virtual FText GetToolTipText() const override { return LOCTEXT("FFrontendFilter_ReplicatedBlueprintToolTip", "Show only blueprints with replicated properties."); }
 
 	// IFilter implementation
-	virtual bool PassesFilter(FAssetFilterType InItem) const override;
+	UE_API virtual bool PassesFilter(FAssetFilterType InItem) const override;
 };
 
 /** A filter that compares the value of an asset registry tag to a target value */
-class CONTENTBROWSER_API FFrontendFilter_ArbitraryComparisonOperation : public FFrontendFilter
+class FFrontendFilter_ArbitraryComparisonOperation : public FFrontendFilter
 {
 public:
-	FFrontendFilter_ArbitraryComparisonOperation(TSharedPtr<FFrontendFilterCategory> InCategory);
+	UE_API FFrontendFilter_ArbitraryComparisonOperation(TSharedPtr<FFrontendFilterCategory> InCategory);
 
 	// FFrontendFilter implementation
-	virtual FString GetName() const override;
-	virtual FText GetDisplayName() const override;
-	virtual FText GetToolTipText() const override;
-	virtual void ModifyContextMenu(FMenuBuilder& MenuBuilder) override;
-	virtual void SaveSettings(const FString& IniFilename, const FString& IniSection, const FString& SettingsString) const override;
-	virtual void LoadSettings(const FString& IniFilename, const FString& IniSection, const FString& SettingsString) override;
+	UE_API virtual FString GetName() const override;
+	UE_API virtual FText GetDisplayName() const override;
+	UE_API virtual FText GetToolTipText() const override;
+	UE_API virtual void ModifyContextMenu(FMenuBuilder& MenuBuilder) override;
+	UE_API virtual void SaveSettings(const FString& IniFilename, const FString& IniSection, const FString& SettingsString) const override;
+	UE_API virtual void LoadSettings(const FString& IniFilename, const FString& IniSection, const FString& SettingsString) override;
 
 	// IFilter implementation
-	virtual bool PassesFilter(FAssetFilterType InItem) const override;
+	UE_API virtual bool PassesFilter(FAssetFilterType InItem) const override;
 
 protected:
-	static FString ConvertOperationToString(ETextFilterComparisonOperation Op);
+	static UE_API FString ConvertOperationToString(ETextFilterComparisonOperation Op);
 	
-	void SetComparisonOperation(ETextFilterComparisonOperation NewOp);
-	bool IsComparisonOperationEqualTo(ETextFilterComparisonOperation TestOp) const;
+	UE_API void SetComparisonOperation(ETextFilterComparisonOperation NewOp);
+	UE_API bool IsComparisonOperationEqualTo(ETextFilterComparisonOperation TestOp) const;
 
-	FText GetKeyValueAsText() const;
-	FText GetTargetValueAsText() const;
-	void OnKeyValueTextCommitted(const FText& InText, ETextCommit::Type InCommitType);
-	void OnTargetValueTextCommitted(const FText& InText, ETextCommit::Type InCommitType);
+	UE_API FText GetKeyValueAsText() const;
+	UE_API FText GetTargetValueAsText() const;
+	UE_API void OnKeyValueTextCommitted(const FText& InText, ETextCommit::Type InCommitType);
+	UE_API void OnTargetValueTextCommitted(const FText& InText, ETextCommit::Type InCommitType);
 
 public:
 	FName TagName;
@@ -310,55 +312,55 @@ private:
 };
 
 /** A filter that only displays assets used by loaded levels */
-class CONTENTBROWSER_API FFrontendFilter_InUseByLoadedLevels : public FFrontendFilter
+class FFrontendFilter_InUseByLoadedLevels : public FFrontendFilter
 {
 public:
 	/** Constructor/Destructor */
-	FFrontendFilter_InUseByLoadedLevels(TSharedPtr<FFrontendFilterCategory> InCategory);
-	~FFrontendFilter_InUseByLoadedLevels() override;
+	UE_API FFrontendFilter_InUseByLoadedLevels(TSharedPtr<FFrontendFilterCategory> InCategory);
+	UE_API ~FFrontendFilter_InUseByLoadedLevels() override;
 
 	// FFrontendFilter implementation
 	virtual FString GetName() const override { return TEXT("InUseByLoadedLevels"); }
 	virtual FText GetDisplayName() const override { return LOCTEXT("FrontendFilter_InUseByLoadedLevels", "In Use By Level"); }
 	virtual FText GetToolTipText() const override { return LOCTEXT("FrontendFilter_InUseByLoadedLevelsToolTip", "Show only assets that are currently in use by any loaded level."); }
-	virtual void ActiveStateChanged(bool bActive) override;
+	UE_API virtual void ActiveStateChanged(bool bActive) override;
 
 	// IFilter implementation
-	virtual bool PassesFilter(FAssetFilterType InItem) const override;
+	UE_API virtual bool PassesFilter(FAssetFilterType InItem) const override;
 
 	/** Handler for when maps change in the editor */
-	void OnEditorMapChange( uint32 MapChangeFlags );
+	UE_API void OnEditorMapChange( uint32 MapChangeFlags );
 
 	/** Handler for when an asset is renamed */
-	void OnAssetPostRename(const TArray<FAssetRenameData>& AssetsAndNames);
+	UE_API void OnAssetPostRename(const TArray<FAssetRenameData>& AssetsAndNames);
 
 	/** Handler for when assets are finished compiling */
-	void OnAssetPostCompile(const TArray<FAssetCompileData>& CompiledAssets);
+	UE_API void OnAssetPostCompile(const TArray<FAssetCompileData>& CompiledAssets);
 
 private:
-	void Refresh();
-	void RegisterDelayedRefresh(float DelayInSeconds);
-	void UnregisterDelayedRefresh();
+	UE_API void Refresh();
+	UE_API void RegisterDelayedRefresh(float DelayInSeconds);
+	UE_API void UnregisterDelayedRefresh();
 	FTSTicker::FDelegateHandle DelayedRefreshHandle;
 	bool bIsDirty = false;
 	bool bIsCurrentlyActive = false;
 };
 
 /** A filter that only displays assets used by any level */
-class CONTENTBROWSER_API FFrontendFilter_UsedInAnyLevel: public FFrontendFilter
+class FFrontendFilter_UsedInAnyLevel: public FFrontendFilter
 {
 public:
 	/** Constructor/Destructor */
-	FFrontendFilter_UsedInAnyLevel(TSharedPtr<FFrontendFilterCategory> InCategory);
+	UE_API FFrontendFilter_UsedInAnyLevel(TSharedPtr<FFrontendFilterCategory> InCategory);
 
 	// FFrontendFilter implementation
 	virtual FString GetName() const override { return TEXT("UsedInAnyLevel"); }
 	virtual FText GetDisplayName() const override { return LOCTEXT("FFrontendFilter_UsedInAnyLevel", "Used In Any Level"); }
 	virtual FText GetToolTipText() const override { return LOCTEXT("FFrontendFilter_UsedInAnyLevelTooltip", "Show only assets that are used in any level."); }
-	virtual void ActiveStateChanged(bool bActive) override;
+	UE_API virtual void ActiveStateChanged(bool bActive) override;
 
 	// IFilter implementation
-	virtual bool PassesFilter(FAssetFilterType InItem) const override;
+	UE_API virtual bool PassesFilter(FAssetFilterType InItem) const override;
 
 private:
 	class IAssetRegistry* AssetRegistry;
@@ -366,20 +368,20 @@ private:
 };
 
 /** A filter that only displays assets not used by any level */
-class CONTENTBROWSER_API FFrontendFilter_NotUsedInAnyLevel : public FFrontendFilter
+class FFrontendFilter_NotUsedInAnyLevel : public FFrontendFilter
 {
 public:
 	/** Constructor/Destructor */
-	FFrontendFilter_NotUsedInAnyLevel(TSharedPtr<FFrontendFilterCategory> InCategory);
+	UE_API FFrontendFilter_NotUsedInAnyLevel(TSharedPtr<FFrontendFilterCategory> InCategory);
 
 	// FFrontendFilter implementation
 	virtual FString GetName() const override { return TEXT("NotUsedInAnyLevel"); }
 	virtual FText GetDisplayName() const override { return LOCTEXT("FFrontendFilter_NotUsedInAnyLevel", "Not Used In Any Level"); }
 	virtual FText GetToolTipText() const override { return LOCTEXT("FFrontendFilter_NotUsedInAnyLevelTooltip", "Show only assets that are not used in any level."); }
-	virtual void ActiveStateChanged(bool bActive) override;
+	UE_API virtual void ActiveStateChanged(bool bActive) override;
 
 	// IFilter implementation
-	virtual bool PassesFilter(FAssetFilterType InItem) const override;
+	UE_API virtual bool PassesFilter(FAssetFilterType InItem) const override;
 
 private:
 	class IAssetRegistry* AssetRegistry;
@@ -387,11 +389,11 @@ private:
 };
 
 /** A filter that only displays assets not used in another asset (Note It does not update itself automatically) */
-class CONTENTBROWSER_API FFrontendFilter_NotUsedInAnyAsset : public FFrontendFilter
+class FFrontendFilter_NotUsedInAnyAsset : public FFrontendFilter
 {
 public:
 	/** Constructor/Destructor */
-	FFrontendFilter_NotUsedInAnyAsset(TSharedPtr<FFrontendFilterCategory> InCategory);
+	UE_API FFrontendFilter_NotUsedInAnyAsset(TSharedPtr<FFrontendFilterCategory> InCategory);
 
 	// FFrontendFilter implementation
 	virtual FString GetName() const override { return TEXT("NotUsedInAnyAsset"); }
@@ -399,44 +401,44 @@ public:
 	virtual FText GetToolTipText() const override { return LOCTEXT("FFrontendFilter_NotUsedInAnyAssetTooltip", "Show only the assets that aren't used by another asset."); }
 
 	// IFilter implementation
-	virtual bool PassesFilter(FAssetFilterType InItem) const override;
+	UE_API virtual bool PassesFilter(FAssetFilterType InItem) const override;
 
 private:
 	class IAssetRegistry* AssetRegistry = nullptr;
 };
 
 /** A filter that displays recently opened assets */
-class CONTENTBROWSER_API FFrontendFilter_Recent : public FFrontendFilter
+class FFrontendFilter_Recent : public FFrontendFilter
 {
 public:
-	FFrontendFilter_Recent(TSharedPtr<FFrontendFilterCategory> InCategory);
-	~FFrontendFilter_Recent();
+	UE_API FFrontendFilter_Recent(TSharedPtr<FFrontendFilterCategory> InCategory);
+	UE_API ~FFrontendFilter_Recent();
 
 	// FFrontendFilter implementation
 	virtual FString GetName() const override { return TEXT("RecentlyOpened"); }
 	virtual FText GetDisplayName() const override { return LOCTEXT("FrontendFilter_Recent", "Recently Opened"); }
 	virtual FText GetToolTipText() const override { return LOCTEXT("FrontendFilter_RecentTooltip", "Show only recently opened assets."); }
-	virtual void ActiveStateChanged(bool bActive) override;
-	virtual void SetCurrentFilter(TArrayView<const FName> InSourcePaths, const FContentBrowserDataFilter& InBaseFilter) override;
+	UE_API virtual void ActiveStateChanged(bool bActive) override;
+	UE_API virtual void SetCurrentFilter(TArrayView<const FName> InSourcePaths, const FContentBrowserDataFilter& InBaseFilter) override;
 
 	// IFilter implementation
-	virtual bool PassesFilter(FAssetFilterType InItem) const override;
+	UE_API virtual bool PassesFilter(FAssetFilterType InItem) const override;
 
-	void ResetFilter(FName InName);
+	UE_API void ResetFilter(FName InName);
 
 private:
-	void RefreshRecentPackagePaths();
+	UE_API void RefreshRecentPackagePaths();
 
 	TSet<FName> RecentPackagePaths;
 	bool bIsCurrentlyActive;
 };
 
 /** A filter that displays only assets that are not read only */
-class CONTENTBROWSER_API FFrontendFilter_Writable : public FFrontendFilter
+class FFrontendFilter_Writable : public FFrontendFilter
 {
 public:
-	FFrontendFilter_Writable(TSharedPtr<FFrontendFilterCategory> InCategory);
-	~FFrontendFilter_Writable();
+	UE_API FFrontendFilter_Writable(TSharedPtr<FFrontendFilterCategory> InCategory);
+	UE_API ~FFrontendFilter_Writable();
 
 	// FFrontendFilter implementation
 	virtual FString GetName() const override { return TEXT("Writable"); }
@@ -444,16 +446,16 @@ public:
 	virtual FText GetToolTipText() const override { return LOCTEXT("FrontendFilter_WritableTooltip", "Show only assets that are not read only."); }
 
 	// IFilter implementation
-	virtual bool PassesFilter(FAssetFilterType InItem) const override;
+	UE_API virtual bool PassesFilter(FAssetFilterType InItem) const override;
 
 private:
 };
 
 /** A filter that displays only packages that contain virtualized data  */
-class CONTENTBROWSER_API FFrontendFilter_VirtualizedData : public FFrontendFilter
+class FFrontendFilter_VirtualizedData : public FFrontendFilter
 {
 public:
-	FFrontendFilter_VirtualizedData(TSharedPtr<FFrontendFilterCategory> InCategory);
+	UE_API FFrontendFilter_VirtualizedData(TSharedPtr<FFrontendFilterCategory> InCategory);
 	~FFrontendFilter_VirtualizedData() = default;
 
 private:
@@ -463,14 +465,14 @@ private:
 	virtual FText GetToolTipText() const override { return LOCTEXT("FrontendFilter_VirtualizedDataTooltip", "Show only package that contain virtualized data."); }
 
 	// IFilter implementation
-	virtual bool PassesFilter(FAssetFilterType InItem) const override;
+	UE_API virtual bool PassesFilter(FAssetFilterType InItem) const override;
 };
 
 /** A filter that displays only assets the are mark as unsupported */
-class CONTENTBROWSER_API FFrontendFilter_Unsupported : public FFrontendFilter
+class FFrontendFilter_Unsupported : public FFrontendFilter
 {
 public:
-	FFrontendFilter_Unsupported(TSharedPtr<FFrontendFilterCategory> InCategory);
+	UE_API FFrontendFilter_Unsupported(TSharedPtr<FFrontendFilterCategory> InCategory);
 	~FFrontendFilter_Unsupported() = default;
 
 	// FFrontendFilter implementation
@@ -479,7 +481,9 @@ public:
 	virtual FText GetToolTipText() const override { return LOCTEXT("FrontendFilter_UnsupportedTooltip", "Show only assets that are not supported by this project."); }
 
 	// IFilter implementation
-	virtual bool PassesFilter(FAssetFilterType InItem) const override;
+	UE_API virtual bool PassesFilter(FAssetFilterType InItem) const override;
 };
 
 #undef LOCTEXT_NAMESPACE
+
+#undef UE_API
