Updated headers using UnrealCodeFixup to make sure dllstorage is on methods/staticvars instead of types. Part 5/n

diff --git a/Engine/Source/Editor/ContentBrowser/Public/SAssetView.h b/Engine/Source/Editor/ContentBrowser/Public/SAssetView.h
--- a/Engine/Source/Editor/ContentBrowser/Public/SAssetView.h
+++ b/Engine/Source/Editor/ContentBrowser/Public/SAssetView.h
@@ -48,6 +48,8 @@
 #include "Widgets/Views/STableRow.h"
 #include "Widgets/Views/STableViewBase.h"
 
+#define UE_API CONTENTBROWSER_API
+
 class FAssetContextMenu;
 class FAssetTextFilter;
 class FAssetThumbnail;
@@ -108,7 +110,7 @@ enum class EAssetViewCopyType
 /**
  * A widget to display a list of filtered assets
  */
-class CONTENTBROWSER_API SAssetView
+class SAssetView
 	: public SCompoundWidget
 {
 public:
@@ -330,26 +332,26 @@ public:
 	friend FAssetContextMenu;
 	friend class UContentBrowserAssetViewContextMenuContext;
 
-	SAssetView();
-	virtual ~SAssetView() override;
+	UE_API SAssetView();
+	UE_API virtual ~SAssetView() override;
 
 	/** Constructs this widget with InArgs */
-	void Construct( const FArguments& InArgs );
+	UE_API void Construct( const FArguments& InArgs );
 
 	/** Changes the base sources for this view */
-	void SetContentSources(const FAssetViewContentSources& InContentSources);
+	UE_API void SetContentSources(const FAssetViewContentSources& InContentSources);
 
 	UE_DEPRECATED(5.6, "Use SetContentSources() instead.")
-	void SetSourcesData(const FSourcesData& InSourcesData);
+	UE_API void SetSourcesData(const FSourcesData& InSourcesData);
 
 	/** Returns the sources filter applied to this asset view */
-	const FAssetViewContentSources& GetContentSources() const;
+	UE_API const FAssetViewContentSources& GetContentSources() const;
 
 	UE_DEPRECATED(5.6, "Use GetContentSources() instead.")
-	FSourcesData GetSourcesData() const;
+	UE_API FSourcesData GetSourcesData() const;
 
 	/** Returns true if a real asset path is selected (i.e \Engine\* or \Game\*) */
-	bool IsAssetPathSelected() const;
+	UE_API bool IsAssetPathSelected() const;
 
 	/**
 	 * @brief Provide a backend filter for the asset view and invalidate current source items.
@@ -358,583 +360,583 @@ public:
 	 * @param InCustomPermissionLists Optional permission lists to allow/deny specific folders.
 	 * 	All folders will be allowed by default, so only denials will have any effect.
 	 */
-	void SetBackendFilter(const FARFilter& InBackendFilter, TArray<TSharedRef<const FPathPermissionList>>* InCustomPermissionLists = nullptr);
+	UE_API void SetBackendFilter(const FARFilter& InBackendFilter, TArray<TSharedRef<const FPathPermissionList>>* InCustomPermissionLists = nullptr);
 
 	/** Get the current backend filter */
 	const FARFilter& GetBackendFilter() const { return BackendFilter; }
 	
 	/** Handler for when a data source requests folder item creation */
-	void NewFolderItemRequested(const FContentBrowserItemTemporaryContext& NewItemContext);
+	UE_API void NewFolderItemRequested(const FContentBrowserItemTemporaryContext& NewItemContext);
 
 	/** Handler for when a data source requests file item creation */
-	void NewFileItemRequested(const FContentBrowserItemDataTemporaryContext& NewItemContext);
+	UE_API void NewFileItemRequested(const FContentBrowserItemDataTemporaryContext& NewItemContext);
 
 	/** Creates a new asset item designed to allocate a new object once it is named. Uses the supplied factory to create the asset */
-	void CreateNewAsset(const FString& DefaultAssetName, const FString& PackagePath, UClass* AssetClass, UFactory* Factory);
+	UE_API void CreateNewAsset(const FString& DefaultAssetName, const FString& PackagePath, UClass* AssetClass, UFactory* Factory);
 
 	/** Sets up an inline rename for the specified item */
-	void RenameItem(const FContentBrowserItem& ItemToRename);
+	UE_API void RenameItem(const FContentBrowserItem& ItemToRename);
 
 	/** Selects the specified items. */
-	void SyncToItems( TArrayView<const FContentBrowserItem> ItemsToSync, const bool bFocusOnSync = true );
+	UE_API void SyncToItems( TArrayView<const FContentBrowserItem> ItemsToSync, const bool bFocusOnSync = true );
 
 	/** Selects the specified virtual paths. */
-	void SyncToVirtualPaths( TArrayView<const FName> VirtualPathsToSync, const bool bFocusOnSync = true );
+	UE_API void SyncToVirtualPaths( TArrayView<const FName> VirtualPathsToSync, const bool bFocusOnSync = true );
 
 	/** Selects the specified assets and paths. */
-	void SyncToLegacy( TArrayView<const FAssetData> AssetDataList, TArrayView<const FString> FolderList, const bool bFocusOnSync = true );
+	UE_API void SyncToLegacy( TArrayView<const FAssetData> AssetDataList, TArrayView<const FString> FolderList, const bool bFocusOnSync = true );
 
 	/** Setup Deferred Pending Sync for recently added PendingSyncItems. */
-	void InitDeferredPendingSyncItems();
+	UE_API void InitDeferredPendingSyncItems();
 
 	/** Sets the state of the asset view to the one described by the history data */
-	void ApplyHistoryData( const FHistoryData& History );
+	UE_API void ApplyHistoryData( const FHistoryData& History );
 
 	/** Returns all the items currently selected in the view */
-	TArray<TSharedPtr<FAssetViewItem>> GetSelectedViewItems() const;
+	UE_API TArray<TSharedPtr<FAssetViewItem>> GetSelectedViewItems() const;
 
 	/** Returns all the items currently selected in the view */
-	TArray<FContentBrowserItem> GetSelectedItems() const;
+	UE_API TArray<FContentBrowserItem> GetSelectedItems() const;
 
 	/** Returns all the folder items currently selected in the view */
-	TArray<FContentBrowserItem> GetSelectedFolderItems() const;
+	UE_API TArray<FContentBrowserItem> GetSelectedFolderItems() const;
 
 	/** Returns all the file items currently selected in the view */
-	TArray<FContentBrowserItem> GetSelectedFileItems() const;
+	UE_API TArray<FContentBrowserItem> GetSelectedFileItems() const;
 
 	/** Returns all the asset data objects in items currently selected in the view */
-	TArray<FAssetData> GetSelectedAssets() const;
+	UE_API TArray<FAssetData> GetSelectedAssets() const;
 
 	/** Returns all the folders currently selected in the view */
-	TArray<FString> GetSelectedFolders() const;
+	UE_API TArray<FString> GetSelectedFolders() const;
 
 	/** Requests that the asset view refreshes all it's source items. This is slow and should only be used if the source items change. */
-	void RequestSlowFullListRefresh();
+	UE_API void RequestSlowFullListRefresh();
 
 	/** Requests that the asset view refreshes only items that are filtered through frontend sources. This should be used when possible. */
-	void RequestQuickFrontendListRefresh();
+	UE_API void RequestQuickFrontendListRefresh();
 
 	/** Saves any settings to config that should be persistent between editor sessions */
-	void SaveSettings(const FString& IniFilename, const FString& IniSection, const FString& SettingsString) const;
+	UE_API void SaveSettings(const FString& IniFilename, const FString& IniSection, const FString& SettingsString) const;
 
 	/** Loads any settings to config that should be persistent between editor sessions */
-	void LoadSettings(const FString& IniFilename, const FString& IniSection, const FString& SettingsString);
+	UE_API void LoadSettings(const FString& IniFilename, const FString& IniSection, const FString& SettingsString);
 
 	/** Adjusts the selected asset by the selection delta, which should be +1 or -1) */
-	void AdjustActiveSelection(int32 SelectionDelta);
+	UE_API void AdjustActiveSelection(int32 SelectionDelta);
 
 	/** Returns true if an asset is currently in the process of being renamed */
-	bool IsRenamingAsset() const;
+	UE_API bool IsRenamingAsset() const;
 
 	// SWidget inherited
-	virtual void Tick( const FGeometry& AllottedGeometry, const double InCurrentTime, const float InDeltaTime ) override;
-	virtual void OnDragLeave( const FDragDropEvent& DragDropEvent ) override;
-	virtual FReply OnDragOver( const FGeometry& MyGeometry, const FDragDropEvent& DragDropEvent ) override;
-	virtual FReply OnDrop( const FGeometry& MyGeometry, const FDragDropEvent& DragDropEvent ) override;
-	virtual FReply OnKeyChar( const FGeometry& MyGeometry,const FCharacterEvent& InCharacterEvent ) override;
-	virtual FReply OnKeyDown( const FGeometry& MyGeometry, const FKeyEvent& InKeyEvent ) override;
-	virtual FReply OnMouseWheel( const FGeometry& MyGeometry, const FPointerEvent& MouseEvent ) override;
-	virtual void OnFocusChanging( const FWeakWidgetPath& PreviousFocusPath, const FWidgetPath& NewWidgetPath, const FFocusEvent& InFocusEvent ) override;
+	UE_API virtual void Tick( const FGeometry& AllottedGeometry, const double InCurrentTime, const float InDeltaTime ) override;
+	UE_API virtual void OnDragLeave( const FDragDropEvent& DragDropEvent ) override;
+	UE_API virtual FReply OnDragOver( const FGeometry& MyGeometry, const FDragDropEvent& DragDropEvent ) override;
+	UE_API virtual FReply OnDrop( const FGeometry& MyGeometry, const FDragDropEvent& DragDropEvent ) override;
+	UE_API virtual FReply OnKeyChar( const FGeometry& MyGeometry,const FCharacterEvent& InCharacterEvent ) override;
+	UE_API virtual FReply OnKeyDown( const FGeometry& MyGeometry, const FKeyEvent& InKeyEvent ) override;
+	UE_API virtual FReply OnMouseWheel( const FGeometry& MyGeometry, const FPointerEvent& MouseEvent ) override;
+	UE_API virtual void OnFocusChanging( const FWeakWidgetPath& PreviousFocusPath, const FWidgetPath& NewWidgetPath, const FFocusEvent& InFocusEvent ) override;
 
 	/** Opens the selected assets or folders, depending on the selection */
-	void OnOpenAssetsOrFolders();
+	UE_API void OnOpenAssetsOrFolders();
 
 	/** Loads the selected assets and previews them if possible */
-	void OnPreviewAssets();
+	UE_API void OnPreviewAssets();
 
 	/** Clears the selection of all the lists in the view */
-	void ClearSelection(bool bForceSilent = false);
+	UE_API void ClearSelection(bool bForceSilent = false);
 
 	/** Returns true if the asset view is in thumbnail editing mode */
-	bool IsThumbnailEditMode() const;
+	UE_API bool IsThumbnailEditMode() const;
 
 	/** Delegate called when an editor setting is changed */
-	void HandleSettingChanged(FName PropertyName);
+	UE_API void HandleSettingChanged(FName PropertyName);
 
 	/** Set whether the user is currently searching or not */
-	void SetUserSearching(bool bInSearching);
+	UE_API void SetUserSearching(bool bInSearching);
 
 	/**
 	 * Forces the plugin content folder to be shown.
 	 *
 	 * @param bEnginePlugin		If true, also forces the engine folder to be shown.
 	 */
-	void ForceShowPluginFolder( bool bEnginePlugin );
+	UE_API void ForceShowPluginFolder( bool bEnginePlugin );
 
 	/** Enables the Show Engine Content setting for the active Content Browser. The user can still toggle the setting manually. */
-	void OverrideShowEngineContent();
+	UE_API void OverrideShowEngineContent();
 
 	/** Enables the Show Developer Content setting for the active Content Browser. The user can still toggle the setting manually. */
-	void OverrideShowDeveloperContent();
+	UE_API void OverrideShowDeveloperContent();
 
 	/** Enables the Show Plugin Content setting for the active Content Browser. The user can still toggle the setting manually. */
-	void OverrideShowPluginContent();
+	UE_API void OverrideShowPluginContent();
 
 	/** Enables the Show Localized Content setting for the active Content Browser. The user can still toggle the setting manually. */
-	void OverrideShowLocalizedContent();
+	UE_API void OverrideShowLocalizedContent();
 
 	/** @return true when we are including class names in search criteria */
-	bool IsIncludingClassNames() const;
+	UE_API bool IsIncludingClassNames() const;
 
 	/** @return true when we are including the entire asset path in search criteria */
-	bool IsIncludingAssetPaths() const;
+	UE_API bool IsIncludingAssetPaths() const;
 
 	/** @return true when we are including collection names in search criteria */
-	bool IsIncludingCollectionNames() const;
+	UE_API bool IsIncludingCollectionNames() const;
 
 	/** Handler for when the view combo button is clicked */
-	TSharedRef<SWidget> GetViewButtonContent();
+	UE_API TSharedRef<SWidget> GetViewButtonContent();
 
 	/** Sets the view type and updates lists accordingly */
-	void SetCurrentViewType(EAssetViewType::Type NewType);
+	UE_API void SetCurrentViewType(EAssetViewType::Type NewType);
 
 	/** Sets the thumbnail size and updates lists accordingly */
-	void SetCurrentThumbnailSize(EThumbnailSize NewThumbnailSize);
+	UE_API void SetCurrentThumbnailSize(EThumbnailSize NewThumbnailSize);
 
 	/** Gets the text for the asset count label */
-	FText GetAssetCountText() const;
+	UE_API FText GetAssetCountText() const;
 
 	/** Gets text name for given thumbnail */
-	static FText ThumbnailSizeToDisplayName(EThumbnailSize InSize);
+	static UE_API FText ThumbnailSizeToDisplayName(EThumbnailSize InSize);
 
 	/** Set the filter list attached to this asset view - allows toggling of the the filter bar layout from the view options */
-	void SetFilterBar(TSharedPtr<SFilterList> InFilterBar);
+	UE_API void SetFilterBar(TSharedPtr<SFilterList> InFilterBar);
 
 	/** Change the delegate bound via the widget argument OnShouldFilterAsset after construction. */
-	void SetShouldFilterItem(FOnShouldFilterItem InCallback);
+	UE_API void SetShouldFilterItem(FOnShouldFilterItem InCallback);
 
 	/** True if we have items pending that ProcessItemsPendingFilter still needs to process */
-	bool HasItemsPendingFilter() const;
+	UE_API bool HasItemsPendingFilter() const;
 
 	/** True if we have thumbnails that are pending update */
-	bool HasThumbnailsPendingUpdate() const;
+	UE_API bool HasThumbnailsPendingUpdate() const;
 
 	/** True if we have deferred item to create */
-	bool HasDeferredItemToCreate() const;
+	UE_API bool HasDeferredItemToCreate() const;
 
 	/** Get the underlying sort manager */
-	TWeakPtr<FAssetViewSortManager> GetSortManager() const;
+	UE_API TWeakPtr<FAssetViewSortManager> GetSortManager() const;
 
 	/** Set's sorting based on the given parameters, using the current value if not specified */
-	void SetSortParameters(
+	UE_API void SetSortParameters(
 		const TOptional<const EColumnSortPriority::Type>& InSortPriority,
 		const TOptional<const FName>& InColumnId,
 		const TOptional<const EColumnSortMode::Type>& InNewSortMode);
 
 private:
 	/** Sets the pending selection to the current selection (used when changing views or refreshing the view). */
-	void SyncToSelection( const bool bFocusOnSync = true );
+	UE_API void SyncToSelection( const bool bFocusOnSync = true );
 
 	/** @return the thumbnail scale setting path to use when looking up the setting in an ini. */
-	FString GetThumbnailScaleSettingPath(const FString& SettingsString, const FString& InViewTypeString) const;
+	UE_API FString GetThumbnailScaleSettingPath(const FString& SettingsString, const FString& InViewTypeString) const;
 
 	/** Load the config for each AssetViewType size into their own container */
-	void LoadScaleSetting(const FString& IniFilename, const FString& IniSection, const FString& SettingsString, const FString& InViewTypeString, EThumbnailSize& OutThumbnailSize);
+	UE_API void LoadScaleSetting(const FString& IniFilename, const FString& IniSection, const FString& SettingsString, const FString& InViewTypeString, EThumbnailSize& OutThumbnailSize);
 
 	/** @return the view type setting path to use when looking up the setting in an ini. */
-	FString GetCurrentViewTypeSettingPath(const FString& SettingsString) const;
+	UE_API FString GetCurrentViewTypeSettingPath(const FString& SettingsString) const;
 
 	/** Calculates a new filler scale used to adjust the thumbnails to fill empty space. */
-	void CalculateFillScale( const FGeometry& AllottedGeometry );
+	UE_API void CalculateFillScale( const FGeometry& AllottedGeometry );
 
 	/** Calculates the latest color and opacity for the hint on thumbnails */
-	void CalculateThumbnailHintColorAndOpacity();
+	UE_API void CalculateThumbnailHintColorAndOpacity();
 
 	/** Handles amortizing the additional filters */
-	void ProcessItemsPendingFilter(const double TickStartTime);
+	UE_API void ProcessItemsPendingFilter(const double TickStartTime);
 
 	/** Creates a new tile view */
-	TSharedRef<SAssetTileView> CreateTileView();
+	UE_API TSharedRef<SAssetTileView> CreateTileView();
 
 	/** Creates a new list view */
-	TSharedRef<SAssetListView> CreateListView();
+	UE_API TSharedRef<SAssetListView> CreateListView();
 
 	/** Creates a new column view */
-	TSharedRef<SAssetColumnView> CreateColumnView();
+	UE_API TSharedRef<SAssetColumnView> CreateColumnView();
 	
 	/** Creates a custom view (if specified to the content browser module) */
-	TSharedRef<SWidget> CreateCustomView();
+	UE_API TSharedRef<SWidget> CreateCustomView();
 
-	const FSlateBrush* GetRevisionControlColumnIconBadge() const;
+	UE_API const FSlateBrush* GetRevisionControlColumnIconBadge() const;
 
 	/** Returns true if the specified search token is allowed */
-	bool IsValidSearchToken(const FString& Token) const;
+	UE_API bool IsValidSearchToken(const FString& Token) const;
 
 	/** Regenerates the AssetItems list from the AssetRegistry */
-	void RefreshSourceItems();
+	UE_API void RefreshSourceItems();
 
 	/** Regenerates the FilteredAssetItems list from the AssetItems list */
-	void RefreshFilteredItems();
+	UE_API void RefreshFilteredItems();
 
 	/** Sets the asset type that represents the majority of the assets in view */
-	void SetMajorityAssetType(FName NewMajorityAssetType);
+	UE_API void SetMajorityAssetType(FName NewMajorityAssetType);
 
 	/** Handler for when an asset is added to a collection */
-	void OnAssetsAddedToCollection( ICollectionContainer& CollectionContainer, const FCollectionNameType& Collection, TConstArrayView<FSoftObjectPath> ObjectPaths );
+	UE_API void OnAssetsAddedToCollection( ICollectionContainer& CollectionContainer, const FCollectionNameType& Collection, TConstArrayView<FSoftObjectPath> ObjectPaths );
 
 	/** Handler for when an asset is removed from a collection */
-	void OnAssetsRemovedFromCollection(ICollectionContainer& CollectionContainer, const FCollectionNameType& Collection, TConstArrayView<FSoftObjectPath> ObjectPaths );
+	UE_API void OnAssetsRemovedFromCollection(ICollectionContainer& CollectionContainer, const FCollectionNameType& Collection, TConstArrayView<FSoftObjectPath> ObjectPaths );
 
 	/** Handler for when a collection is renamed */
-	void OnCollectionRenamed(ICollectionContainer& CollectionContainer, const FCollectionNameType& OriginalCollection, const FCollectionNameType& NewCollection );
+	UE_API void OnCollectionRenamed(ICollectionContainer& CollectionContainer, const FCollectionNameType& OriginalCollection, const FCollectionNameType& NewCollection );
 
 	/** Handler for when a collection is updated */
-	void OnCollectionUpdated(ICollectionContainer& CollectionContainer, const FCollectionNameType& Collection );
+	UE_API void OnCollectionUpdated(ICollectionContainer& CollectionContainer, const FCollectionNameType& Collection );
 
 	/** Handler for when any frontend filters have been changed */
-	void OnFrontendFiltersChanged();
+	UE_API void OnFrontendFiltersChanged();
 
 	/** Returns true if there is any frontend filter active */
-	bool IsFrontendFilterActive() const;
+	UE_API bool IsFrontendFilterActive() const;
 
 	/** Returns true if the specified Content Browser item passes all applied frontend (non asset registry) filters */
-	bool PassesCurrentFrontendFilter(const FContentBrowserItem& Item) const;
+	UE_API bool PassesCurrentFrontendFilter(const FContentBrowserItem& Item) const;
 
 	/** Returns true if the current filters deem that the asset view should be filtered recursively (overriding folder view) */
-	bool ShouldFilterRecursively() const;
+	UE_API bool ShouldFilterRecursively() const;
 
 	/** Sorts the contents of the asset view alphabetically */
-	void SortList(bool bSyncToSelection = true);
+	UE_API void SortList(bool bSyncToSelection = true);
 
 	/** Returns the thumbnails hint color and opacity */
-	FLinearColor GetThumbnailHintColorAndOpacity() const;
+	UE_API FLinearColor GetThumbnailHintColorAndOpacity() const;
 
 	/** Populate the given params for this AssetView */
-	void PopulateFilterAdditionalParams(FFiltersAdditionalParams& OutParams);
+	UE_API void PopulateFilterAdditionalParams(FFiltersAdditionalParams& OutParams);
 
 	/** Registers the menu used by the settings button */
-	static FDelayedAutoRegisterHelper AssetViewOptionsMenuRegistration;
+	static UE_API FDelayedAutoRegisterHelper AssetViewOptionsMenuRegistration;
 
 	/** Extend the top toolbar */
-	static FDelayedAutoRegisterHelper ToolBarMenuExtensionRegistration;
+	static UE_API FDelayedAutoRegisterHelper ToolBarMenuExtensionRegistration;
 
 	/** Extend the bottom navigation bar */
-	static FDelayedAutoRegisterHelper NavigationBarMenuExtensionRegistration;
+	static UE_API FDelayedAutoRegisterHelper NavigationBarMenuExtensionRegistration;
 
-	void OnSetSortParameters(const FToolMenuContext& InMenuContext,
+	UE_API void OnSetSortParameters(const FToolMenuContext& InMenuContext,
 		const TOptional<const EColumnSortPriority::Type> InSortPriority,
 		const TOptional<const FName> InColumnId,
 		const TOptional<const EColumnSortMode::Type> InNewSortMode);
 
 	/** Fill in menu content for when the sorting combo button is clicked */
-	void PopulateSortingButtonMenu(UToolMenu* Menu);
+	UE_API void PopulateSortingButtonMenu(UToolMenu* Menu);
 
 	/** Toggle whether folders should be shown or not */
-	void ToggleShowFolders();
+	UE_API void ToggleShowFolders();
 
 	/** Whether or not it's possible to show folders */
-	bool IsToggleShowFoldersAllowed() const;
+	UE_API bool IsToggleShowFoldersAllowed() const;
 
 	/** @return true when we are showing folders */
-	bool IsShowingFolders() const;
+	UE_API bool IsShowingFolders() const;
 
 	/** @return true when we are showing read-only folders */
-	bool IsShowingReadOnlyFolders() const;
+	UE_API bool IsShowingReadOnlyFolders() const;
 
 	/** Toggle whether empty folders should be shown or not */
-	void ToggleShowEmptyFolders();
+	UE_API void ToggleShowEmptyFolders();
 
 	/** Whether or not it's possible to show empty folders */
-	bool IsToggleShowEmptyFoldersAllowed() const;
+	UE_API bool IsToggleShowEmptyFoldersAllowed() const;
 
 	/** @return true when we are showing empty folders */
-	bool IsShowingEmptyFolders() const;
+	UE_API bool IsShowingEmptyFolders() const;
 
 	/** @return true when the asset view is showing object redirectors */
-	bool IsShowingRedirectors() const;
+	UE_API bool IsShowingRedirectors() const;
 
 	/** Toggle whether localized content should be shown or not */
-	void ToggleShowLocalizedContent();
+	UE_API void ToggleShowLocalizedContent();
 
 	/** Whether or not it's possible to show localized content */
-	bool IsToggleShowLocalizedContentAllowed() const;
+	UE_API bool IsToggleShowLocalizedContentAllowed() const;
 
 	/** @return true when we are showing folders */
-	bool IsShowingLocalizedContent() const;
+	UE_API bool IsShowingLocalizedContent() const;
 
 	/** Toggle whether to show real-time thumbnails */
-	void ToggleRealTimeThumbnails();
+	UE_API void ToggleRealTimeThumbnails();
 
 	/** Whether it is possible to show real-time thumbnails */
-	bool CanShowRealTimeThumbnails() const;
+	UE_API bool CanShowRealTimeThumbnails() const;
 
 	/** @return true if we are showing real-time thumbnails */
-	bool IsShowingRealTimeThumbnails() const;
+	UE_API bool IsShowingRealTimeThumbnails() const;
 
 	/** Toggle whether plugin content should be shown or not */
-	void ToggleShowPluginContent();
+	UE_API void ToggleShowPluginContent();
 
 	/** @return true when we are showing plugin content */
-	bool IsShowingPluginContent() const;
+	UE_API bool IsShowingPluginContent() const;
 
 	/** Toggle whether engine content should be shown or not */
-	void ToggleShowEngineContent();
+	UE_API void ToggleShowEngineContent();
 
 	/** @return true when we are showing engine content */
-	bool IsShowingEngineContent() const;
+	UE_API bool IsShowingEngineContent() const;
 
 	/** Toggle whether developers content should be shown or not */
-	void ToggleShowDevelopersContent();
+	UE_API void ToggleShowDevelopersContent();
 
 	/** Whether or not it's possible to toggle developers content */
-	bool IsToggleShowDevelopersContentAllowed() const;
+	UE_API bool IsToggleShowDevelopersContentAllowed() const;
 
 	/** Whether or not it's possible to toggle engine content */
-	bool IsToggleShowEngineContentAllowed() const;
+	UE_API bool IsToggleShowEngineContentAllowed() const;
 
 	/** Whether or not it's possible to toggle plugin content */
-	bool IsToggleShowPluginContentAllowed() const;
+	UE_API bool IsToggleShowPluginContentAllowed() const;
 	
 	/** @return true when we are showing the developers content */
-	bool IsShowingDevelopersContent() const;
+	UE_API bool IsShowingDevelopersContent() const;
 
 	/** Toggle whether favorites should be shown or not */
-	void ToggleShowFavorites();
+	UE_API void ToggleShowFavorites();
 
 	/** Whether or not it's possible to toggle favorites */
-	bool IsToggleShowFavoritesAllowed() const;
+	UE_API bool IsToggleShowFavoritesAllowed() const;
 
 	/** @return true when we are showing favorites */
-	bool IsShowingFavorites() const;
+	UE_API bool IsShowingFavorites() const;
 
 	/** Whether or not it's possible to show C++ content */
-	bool IsToggleShowCppContentAllowed() const;
+	UE_API bool IsToggleShowCppContentAllowed() const;
 
 	/** @return true when we are showing C++ content */
-	bool IsShowingCppContent() const;
+	UE_API bool IsShowingCppContent() const;
 
 	/** Toggle whether class names should be included in search criteria */
-	void ToggleIncludeClassNames();
+	UE_API void ToggleIncludeClassNames();
 
 	/** Whether or not it's possible to include class names in search criteria */
-	bool IsToggleIncludeClassNamesAllowed() const;
+	UE_API bool IsToggleIncludeClassNamesAllowed() const;
 
 	/** Toggle whether the entire asset path should be included in search criteria */
-	void ToggleIncludeAssetPaths();
+	UE_API void ToggleIncludeAssetPaths();
 
 	/** Whether or not it's possible to include the entire asset path in search criteria */
-	bool IsToggleIncludeAssetPathsAllowed() const;
+	UE_API bool IsToggleIncludeAssetPathsAllowed() const;
 
 	/** Toggle whether collection names should be included in search criteria */
-	void ToggleIncludeCollectionNames();
+	UE_API void ToggleIncludeCollectionNames();
 
 	/** Whether or not it's possible to include collection names in search criteria */
-	bool IsToggleIncludeCollectionNamesAllowed() const;
+	UE_API bool IsToggleIncludeCollectionNamesAllowed() const;
 
 	/** @return true when we are filtering recursively when we have an asset path */
-	bool IsFilteringRecursively() const;
+	UE_API bool IsFilteringRecursively() const;
 
 	/** Whether or not it's possible to toggle how to filtering recursively */
-	bool IsToggleFilteringRecursivelyAllowed() const;
+	UE_API bool IsToggleFilteringRecursivelyAllowed() const;
 
 	/** Toggle whether we're filtering recursively */
-	void ToggleFilteringRecursively();
+	UE_API void ToggleFilteringRecursively();
 
 	/** Toggle whether we are showing the all virtual folder */
-	void ToggleShowAllFolder();
+	UE_API void ToggleShowAllFolder();
 
 	/** @return true when we are showing the all virtual folder */
-	bool IsShowingAllFolder() const;
+	UE_API bool IsShowingAllFolder() const;
 
 	/** Toggle whether we are re-organizing folders into virtual locations */
-	void ToggleOrganizeFolders();
+	UE_API void ToggleOrganizeFolders();
 
 	/** @return true when we are organizing folders into virtual locations */
-	bool IsOrganizingFolders() const;
+	UE_API bool IsOrganizingFolders() const;
 
 	/** Sets the view type and forcibly dismisses all currently open context menus */
-	void SetCurrentViewTypeFromMenu(EAssetViewType::Type NewType);
+	UE_API void SetCurrentViewTypeFromMenu(EAssetViewType::Type NewType);
 
 	/** Clears the reference to the current view and creates a new one, based on CurrentViewType */
-	void CreateCurrentView();
+	UE_API void CreateCurrentView();
 
 	/** Gets the current view type (list or tile) */
-	EAssetViewType::Type GetCurrentViewType() const;
+	UE_API EAssetViewType::Type GetCurrentViewType() const;
 
-	TSharedRef<SWidget> CreateShadowOverlay( TSharedRef<STableViewBase> Table );
+	UE_API TSharedRef<SWidget> CreateShadowOverlay( TSharedRef<STableViewBase> Table );
 
 	/** Returns true if ViewType is the current view type */
-	bool IsCurrentViewType(EAssetViewType::Type ViewType) const;
+	UE_API bool IsCurrentViewType(EAssetViewType::Type ViewType) const;
 
 	/** Set the keyboard focus to the correct list view that should be active */
-	void FocusList() const;
+	UE_API void FocusList() const;
 
 	/** Refreshes the list view to display any changes made to the non-filtered assets */
-	void RefreshList();
+	UE_API void RefreshList();
 
 	/** Sets the sole selection for all lists in the view */
-	void SetSelection(const TSharedPtr<FAssetViewItem>& Item);
+	UE_API void SetSelection(const TSharedPtr<FAssetViewItem>& Item);
 
 	/** Sets selection for an item in all lists in the view */
-	void SetItemSelection(const TSharedPtr<FAssetViewItem>& Item, bool bSelected, const ESelectInfo::Type SelectInfo = ESelectInfo::Direct);
+	UE_API void SetItemSelection(const TSharedPtr<FAssetViewItem>& Item, bool bSelected, const ESelectInfo::Type SelectInfo = ESelectInfo::Direct);
 
 	/** Scrolls the selected item into view for all lists in the view */
-	void RequestScrollIntoView(const TSharedPtr<FAssetViewItem>& Item);
+	UE_API void RequestScrollIntoView(const TSharedPtr<FAssetViewItem>& Item);
 
 	/** Handler for list view widget creation */
-	TSharedRef<ITableRow> MakeListViewWidget(TSharedPtr<FAssetViewItem> AssetItem, const TSharedRef<STableViewBase>& OwnerTable);
+	UE_API TSharedRef<ITableRow> MakeListViewWidget(TSharedPtr<FAssetViewItem> AssetItem, const TSharedRef<STableViewBase>& OwnerTable);
 
 	/** Handler for tile view widget creation */
-	TSharedRef<ITableRow> MakeTileViewWidget(TSharedPtr<FAssetViewItem> AssetItem, const TSharedRef<STableViewBase>& OwnerTable);
+	UE_API TSharedRef<ITableRow> MakeTileViewWidget(TSharedPtr<FAssetViewItem> AssetItem, const TSharedRef<STableViewBase>& OwnerTable);
 
 	/** Handler for column view widget creation */
-	TSharedRef<ITableRow> MakeColumnViewWidget(TSharedPtr<FAssetViewItem> AssetItem, const TSharedRef<STableViewBase>& OwnerTable);
+	UE_API TSharedRef<ITableRow> MakeColumnViewWidget(TSharedPtr<FAssetViewItem> AssetItem, const TSharedRef<STableViewBase>& OwnerTable);
 
 	/** Handler for when any asset item widget gets destroyed */
-	void AssetItemWidgetDestroyed(const TSharedPtr<FAssetViewItem>& Item);
+	UE_API void AssetItemWidgetDestroyed(const TSharedPtr<FAssetViewItem>& Item);
 	
 	/** Creates new thumbnails that are near the view area and deletes old thumbnails that are no longer relevant. */
-	void UpdateThumbnails();
+	UE_API void UpdateThumbnails();
 
 	/**  Helper function for UpdateThumbnails. Adds the specified item to the new thumbnail relevancy map and creates any thumbnails for new items. Returns the thumbnail. */
-	TSharedPtr<FAssetThumbnail> AddItemToNewThumbnailRelevancyMap(const TSharedPtr<FAssetViewItem>& Item, TMap< TSharedPtr<FAssetViewItem>, TSharedPtr<FAssetThumbnail> >& NewRelevantThumbnails);
+	UE_API TSharedPtr<FAssetThumbnail> AddItemToNewThumbnailRelevancyMap(const TSharedPtr<FAssetViewItem>& Item, TMap< TSharedPtr<FAssetViewItem>, TSharedPtr<FAssetThumbnail> >& NewRelevantThumbnails);
 
 	/** Handler for tree view selection changes */
-	void AssetSelectionChanged( TSharedPtr<FAssetViewItem > AssetItem, ESelectInfo::Type SelectInfo );
+	UE_API void AssetSelectionChanged( TSharedPtr<FAssetViewItem > AssetItem, ESelectInfo::Type SelectInfo );
 
 	/** Handler for when an item has scrolled into view after having been requested to do so */
-	void ItemScrolledIntoView(TSharedPtr<FAssetViewItem> AssetItem, const TSharedPtr<ITableRow>& Widget);
+	UE_API void ItemScrolledIntoView(TSharedPtr<FAssetViewItem> AssetItem, const TSharedPtr<ITableRow>& Widget);
 
 	/** Handler for context menus */
-	TSharedPtr<SWidget> OnGetContextMenuContent();
+	UE_API TSharedPtr<SWidget> OnGetContextMenuContent();
 
 	/** Handler called when an asset context menu is about to open */
-	bool CanOpenContextMenu() const;
+	UE_API bool CanOpenContextMenu() const;
 
 	/** Handler for double clicking an item */
-	void OnListMouseButtonDoubleClick(TSharedPtr<FAssetViewItem> AssetItem);
+	UE_API void OnListMouseButtonDoubleClick(TSharedPtr<FAssetViewItem> AssetItem);
 
 	/** Handle dragging an asset */
-	FReply OnDraggingAssetItem(const FGeometry& MyGeometry, const FPointerEvent& MouseEvent);
+	UE_API FReply OnDraggingAssetItem(const FGeometry& MyGeometry, const FPointerEvent& MouseEvent);
 
 	/** Checks if the asset name being committed is valid */
-	bool AssetVerifyRenameCommit(const TSharedPtr<FAssetViewItem>& Item, const FText& NewName, const FSlateRect& MessageAnchor, FText& OutErrorMessage);
+	UE_API bool AssetVerifyRenameCommit(const TSharedPtr<FAssetViewItem>& Item, const FText& NewName, const FSlateRect& MessageAnchor, FText& OutErrorMessage);
 
 	/** An asset item has started to be renamed */
-	void AssetRenameBegin(const TSharedPtr<FAssetViewItem>& Item, const FString& NewName, const FSlateRect& MessageAnchor);
+	UE_API void AssetRenameBegin(const TSharedPtr<FAssetViewItem>& Item, const FString& NewName, const FSlateRect& MessageAnchor);
 
 	/** An asset item that was prompting the user for a new name was committed. Return false to indicate that the user should enter a new name */
-	void AssetRenameCommit(const TSharedPtr<FAssetViewItem>& Item, const FString& NewName, const FSlateRect& MessageAnchor, const ETextCommit::Type CommitType);
+	UE_API void AssetRenameCommit(const TSharedPtr<FAssetViewItem>& Item, const FString& NewName, const FSlateRect& MessageAnchor, const ETextCommit::Type CommitType);
 
 	/** Gets the color and opacity for all names of assets in the asset view */
-	FLinearColor GetAssetNameColorAndOpacity() const;
+	UE_API FLinearColor GetAssetNameColorAndOpacity() const;
 
 	/** Returns true if tooltips should be allowed right now. Tooltips are typically disabled while right click scrolling. */
-	bool ShouldAllowToolTips() const;
+	UE_API bool ShouldAllowToolTips() const;
 
 	/** Returns true if the asset view is currently allowing the user to edit thumbnails */
-	bool IsThumbnailEditModeAllowed() const;
+	UE_API bool IsThumbnailEditModeAllowed() const;
 
 	/** The "Done Editing" button was pressed in the thumbnail edit mode strip */
-	FReply EndThumbnailEditModeClicked();
+	UE_API FReply EndThumbnailEditModeClicked();
 
 	/** Gets the visibility of the Thumbnail Edit Mode label */
-	EVisibility GetEditModeLabelVisibility() const;
+	UE_API EVisibility GetEditModeLabelVisibility() const;
 
 	/** Gets the visibility of the list view */
-	EVisibility GetListViewVisibility() const;
+	UE_API EVisibility GetListViewVisibility() const;
 
 	/** Gets the visibility of the tile view */
-	EVisibility GetTileViewVisibility() const;
+	UE_API EVisibility GetTileViewVisibility() const;
 
 	/** Gets the visibility of the column view */
-	EVisibility GetColumnViewVisibility() const;
+	UE_API EVisibility GetColumnViewVisibility() const;
 
 	/** Toggles tooltip expanded default state */
-	void ToggleTooltipExpandedState();
+	UE_API void ToggleTooltipExpandedState();
 
 	/** True if the default state of tooltip is expanded */
-	bool IsTooltipExpandedByDefault();
+	UE_API bool IsTooltipExpandedByDefault();
 
 	/** Toggles thumbnail editing mode */
-	void ToggleThumbnailEditMode();
+	UE_API void ToggleThumbnailEditMode();
 
 	/** Called when  thumbnail size is changed */
-	void OnThumbnailSizeChanged(EThumbnailSize NewThumbnailSize);
+	UE_API void OnThumbnailSizeChanged(EThumbnailSize NewThumbnailSize);
 
 	/** Is the current thumbnail size menu option checked */
-	bool IsThumbnailSizeChecked(EThumbnailSize InThumbnailSize) const;
+	UE_API bool IsThumbnailSizeChecked(EThumbnailSize InThumbnailSize) const;
 
 	/** Are thumbnails allowed to be scaled by the user */
-	bool IsThumbnailScalingAllowed() const;
+	UE_API bool IsThumbnailScalingAllowed() const;
 
 	/** Gets the current thumbnail scale */
-	float GetThumbnailScale() const;
+	UE_API float GetThumbnailScale() const;
 
 	/** Gets the current thumbnail size */
-	float GetThumbnailSizeValue() const;
+	UE_API float GetThumbnailSizeValue() const;
 
 	/** Set the Min/Max Thumbnail size based on the EThumbnailSize chosen */
-	void UpdateThumbnailSizeValue();
+	UE_API void UpdateThumbnailSizeValue();
 
 	/** Gets the current thumbnail size enum */
 	EThumbnailSize GetThumbnailSize() const { return ThumbnailSizes.Contains(CurrentViewType) ? ThumbnailSizes[CurrentViewType] : EThumbnailSize::Medium; }
 
 	/** Gets the spacing dedicated to support type names */
-	float GetTileViewTypeNameHeight() const;
+	UE_API float GetTileViewTypeNameHeight() const;
 
 	/** Gets the spacing needed to support source control icons. No scaling is required for this since the icon itself is not scaled by the thumbnail size. */
-	float GetSourceControlIconHeight() const;
+	UE_API float GetSourceControlIconHeight() const;
 
 	/** Gets the scaled item height for the list view */
-	float GetListViewItemHeight() const;
+	UE_API float GetListViewItemHeight() const;
 
 	/** Gets the padding for the list view rows */
-	FMargin GetListViewItemPadding() const;
+	UE_API FMargin GetListViewItemPadding() const;
 
 	/** Gets the final scaled item height for the tile view */
-	float GetTileViewItemHeight() const;
+	UE_API float GetTileViewItemHeight() const;
 
 	/** Get the TileView Thumbnail dimension height and width are the same for the thumbnail itself */
-	float GetTileViewThumbnailDimension() const;
+	UE_API float GetTileViewThumbnailDimension() const;
 
 	/** Gets the scaled item height for the tile view before the filler scale is applied */
-	float GetTileViewItemBaseHeight() const;
+	UE_API float GetTileViewItemBaseHeight() const;
 
 	/** Gets the final scaled item width for the tile view */
-	float GetTileViewItemWidth() const;
+	UE_API float GetTileViewItemWidth() const;
 
 	/** Gets the scaled item width for the tile view before the filler scale is applied */
-	float GetTileViewItemBaseWidth() const;
+	UE_API float GetTileViewItemBaseWidth() const;
 
 	/** Gets the sort mode for the supplied ColumnId */
-	EColumnSortMode::Type GetColumnSortMode(const FName ColumnId) const;
+	UE_API EColumnSortMode::Type GetColumnSortMode(const FName ColumnId) const;
 
 	/** Gets the sort order for the supplied ColumnId */
-	EColumnSortPriority::Type GetColumnSortPriority(const FName ColumnId) const;
+	UE_API EColumnSortPriority::Type GetColumnSortPriority(const FName ColumnId) const;
 
 	/** Handler for when a column header is clicked */
-	void OnSortColumnHeader(const EColumnSortPriority::Type SortPriority, const FName& ColumnId, const EColumnSortMode::Type NewSortMode);
+	UE_API void OnSortColumnHeader(const EColumnSortPriority::Type SortPriority, const FName& ColumnId, const EColumnSortMode::Type NewSortMode);
 
 	/** @return The state of the is working progress bar */
-	TOptional< float > GetIsWorkingProgressBarState() const;
+	UE_API TOptional< float > GetIsWorkingProgressBarState() const;
 
 	/** Is the no assets to show warning visible? */
-	EVisibility IsAssetShowWarningTextVisible() const;
+	UE_API EVisibility IsAssetShowWarningTextVisible() const;
 
 	/** Gets the text for displaying no assets to show warning */
-	FText GetAssetShowWarningText() const;
+	UE_API FText GetAssetShowWarningText() const;
 
 	/** Whether we have a single source collection selected */
-	bool HasSingleCollectionSource() const;
+	UE_API bool HasSingleCollectionSource() const;
 
 	/** Starts the async flow of creating a file or folder item from deferred data */
-	void BeginCreateDeferredItem();
+	UE_API void BeginCreateDeferredItem();
 
 	/** Finishes the async flow of creating a file or folder item from deferred data */
-	FContentBrowserItem EndCreateDeferredItem(const TSharedPtr<FAssetViewItem>& InItem, const FString& InName, const bool bFinalize, FText& OutErrorText);
+	UE_API FContentBrowserItem EndCreateDeferredItem(const TSharedPtr<FAssetViewItem>& InItem, const FString& InName, const bool bFinalize, FText& OutErrorText);
 
 	/** @return The current quick-jump term */
-	FText GetQuickJumpTerm() const;
+	UE_API FText GetQuickJumpTerm() const;
 
 	/** @return Whether the quick-jump term is currently visible? */
-	EVisibility IsQuickJumpVisible() const;
+	UE_API EVisibility IsQuickJumpVisible() const;
 
 	/** @return The color that should be used for the quick-jump term */
-	FSlateColor GetQuickJumpColor() const;
+	UE_API FSlateColor GetQuickJumpColor() const;
 
 	/** Reset the quick-jump to its empty state */
-	void ResetQuickJump();
+	UE_API void ResetQuickJump();
 
 	/**
 	 * Called from OnKeyChar and OnKeyDown to handle quick-jump key presses
@@ -944,62 +946,62 @@ private:
 	 * @param bTestOnly			True if we only want to test whether the key press would be handled, but not actually update the quick-jump term
 	 * @return FReply::Handled or FReply::Unhandled
 	 */
-	FReply HandleQuickJumpKeyDown(const TCHAR InCharacter, const bool bIsControlDown, const bool bIsAltDown, const bool bTestOnly);
+	UE_API FReply HandleQuickJumpKeyDown(const TCHAR InCharacter, const bool bIsControlDown, const bool bIsAltDown, const bool bTestOnly);
 
 	/**
 	 * Perform a quick-jump to the next available asset in FilteredAssetItems that matches the current term
 	 * @param bWasJumping		True if we were performing an ongoing quick-jump last Tick
 	 * @return True if the quick-jump found a valid match, false otherwise
 	 */
-	bool PerformQuickJump(const bool bWasJumping);
+	UE_API bool PerformQuickJump(const bool bWasJumping);
 	
 	/** Resets the column filtering state to make them all visible */
-	void ResetColumns();
+	UE_API void ResetColumns();
 
 	/** Export columns to CSV */
-	void ExportColumns();
+	UE_API void ExportColumns();
 
 	/** Called when a column is shown/hidden in the column view */
-	void OnHiddenColumnsChanged();
+	UE_API void OnHiddenColumnsChanged();
 
 	/** Will compute the max row size from all its children for the specified column id*/
-	FVector2D GetMaxRowSizeForColumn(const FName& ColumnId);
+	UE_API FVector2D GetMaxRowSizeForColumn(const FName& ColumnId);
 
 	/** Append the current effective backend filter (intersection of BackendFilter and SupportedFilter) to the given filter. */
-	void AppendBackendFilter(FARFilter& FilterToAppendTo) const;
+	UE_API void AppendBackendFilter(FARFilter& FilterToAppendTo) const;
 
-	EContentBrowserItemCategoryFilter DetermineItemCategoryFilter() const;
-	FContentBrowserDataFilter CreateBackendDataFilter(bool bInvalidateCache) const;
+	UE_API EContentBrowserItemCategoryFilter DetermineItemCategoryFilter() const;
+	UE_API FContentBrowserDataFilter CreateBackendDataFilter(bool bInvalidateCache) const;
 
 	/** Handles updating the view when content items are changed */
-	void HandleItemDataUpdated(TArrayView<const FContentBrowserItemDataUpdate> InUpdatedItems);
+	UE_API void HandleItemDataUpdated(TArrayView<const FContentBrowserItemDataUpdate> InUpdatedItems);
 
 	/** Notification for when the content browser has completed it's initial search */
-	void HandleItemDataDiscoveryComplete();
+	UE_API void HandleItemDataDiscoveryComplete();
 
 	/** Get the config struct for the owning Content Browser, if one exists. */
-	FContentBrowserInstanceConfig* GetContentBrowserConfig() const;
+	UE_API FContentBrowserInstanceConfig* GetContentBrowserConfig() const;
 
 	/** Get the config struct for this asset view, if one exists. */
-	FAssetViewInstanceConfig* GetAssetViewConfig() const;
+	UE_API FAssetViewInstanceConfig* GetAssetViewConfig() const;
 
 	/** Bind our UI commands */
-	void BindCommands();
+	UE_API void BindCommands();
 
 	/** Populate the given parameters based on the current selection */
-	void PopulateSelectedFilesAndFolders(TArray<FContentBrowserItem>& OutSelectedFolders, TArray<FContentBrowserItem>& OutSelectedFiles) const;
+	UE_API void PopulateSelectedFilesAndFolders(TArray<FContentBrowserItem>& OutSelectedFolders, TArray<FContentBrowserItem>& OutSelectedFiles) const;
 
 	/** Handler for the CopyReference CopyObjectPath and CopyPackageName */
-	void ExecuteCopy(EAssetViewCopyType InCopyType) const;
+	UE_API void ExecuteCopy(EAssetViewCopyType InCopyType) const;
 
 	/** Append folders path to the given ClipboardText */
-	void ExecuteCopyFolders(const TArray<FContentBrowserItem>& InSelectedFolders, FString& OutClipboardText) const;
+	UE_API void ExecuteCopyFolders(const TArray<FContentBrowserItem>& InSelectedFolders, FString& OutClipboardText) const;
 
 	/** Handler for Paste */
-	void ExecutePaste();
+	UE_API void ExecutePaste();
 
 	/** Check if the custom view view is available */
-	bool IsCustomViewSet() const;
+	UE_API bool IsCustomViewSet() const;
 
 private:
 	friend class FAssetViewFrontendFilterHelper;
@@ -1384,12 +1386,14 @@ private:
 	UE::Telemetry::ContentBrowser::FFrontendFilterTelemetry CurrentFrontendFilterTelemetry;
 	
 	// If there is in-progress filter telemetry, mark it as complete and send pending data 
-	void OnCompleteFiltering(double InAmortizeDuration);
+	UE_API void OnCompleteFiltering(double InAmortizeDuration);
 	// If there is in-progress filter telemetry, mark it as incomplete and send pending data 
-	void OnInterruptFiltering();
+	UE_API void OnInterruptFiltering();
 	// Update telemetry for user interaction during an incomplete filtering
-	void OnInteractDuringFiltering();
+	UE_API void OnInteractDuringFiltering();
 
 public:
-	bool ShouldColumnGenerateWidget(const FString ColumnName) const;
+	UE_API bool ShouldColumnGenerateWidget(const FString ColumnName) const;
 };
+
+#undef UE_API
diff --git a/Engine/Source/Editor/ContentBrowser/Public/Widgets/SAssetMenuIcon.h b/Engine/Source/Editor/ContentBrowser/Public/Widgets/SAssetMenuIcon.h
--- a/Engine/Source/Editor/ContentBrowser/Public/Widgets/SAssetMenuIcon.h
+++ b/Engine/Source/Editor/ContentBrowser/Public/Widgets/SAssetMenuIcon.h
@@ -9,8 +9,10 @@
 #include "Widgets/SCompoundWidget.h"
 #include "Widgets/Text/STextBlock.h"
 
+#define UE_API CONTENTBROWSER_API
+
 /** Represents an Asset Type within a menu. */
-class CONTENTBROWSER_API SAssetMenuIcon : public SCompoundWidget
+class SAssetMenuIcon : public SCompoundWidget
 {
 public:
 	SLATE_BEGIN_ARGS(SAssetMenuIcon)
@@ -22,5 +24,7 @@ public:
 	SLATE_END_ARGS()
 
 	/** Constructs this widget with InArgs */
-	void Construct(const FArguments& InArgs, const UClass* InAssetClass, const FName InIconOverride = NAME_None);
+	UE_API void Construct(const FArguments& InArgs, const UClass* InAssetClass, const FName InIconOverride = NAME_None);
 };
+
+#undef UE_API
diff --git a/Engine/Source/Editor/ContentBrowserData/Private/ContentBrowserVirtualFolderDataPayload.h b/Engine/Source/Editor/ContentBrowserData/Private/ContentBrowserVirtualFolderDataPayload.h
--- a/Engine/Source/Editor/ContentBrowserData/Private/ContentBrowserVirtualFolderDataPayload.h
+++ b/Engine/Source/Editor/ContentBrowserData/Private/ContentBrowserVirtualFolderDataPayload.h
@@ -4,7 +4,7 @@
 
 #include "ContentBrowserItemData.h"
 
-class CONTENTBROWSERDATA_API FContentBrowserVirtualFolderItemDataPayload : public IContentBrowserItemDataPayload
+class FContentBrowserVirtualFolderItemDataPayload : public IContentBrowserItemDataPayload
 {
 public:
 	explicit FContentBrowserVirtualFolderItemDataPayload (bool InIsCustomVirtual)
diff --git a/Engine/Source/Editor/ContentBrowserData/Public/ContentBrowserDataDragDropOp.h b/Engine/Source/Editor/ContentBrowserData/Public/ContentBrowserDataDragDropOp.h
--- a/Engine/Source/Editor/ContentBrowserData/Public/ContentBrowserDataDragDropOp.h
+++ b/Engine/Source/Editor/ContentBrowserData/Public/ContentBrowserDataDragDropOp.h
@@ -13,13 +13,15 @@
 #include "Internationalization/Text.h"
 #include "Templates/SharedPointer.h"
 
+#define UE_API CONTENTBROWSERDATA_API
+
 class UActorFactory;
 struct FAssetData;
 
 /**
  * Additional params for FContentBrowserDataDragDropOp to override the thumbnail used if dragging only folders
  */
-struct CONTENTBROWSERDATA_API FThumbnailOverrideParams
+struct FThumbnailOverrideParams
 {
 	FThumbnailOverrideParams() {}
 
@@ -33,14 +35,14 @@ struct CONTENTBROWSERDATA_API FThumbnailOverrideParams
 	FLinearColor FolderColorOverride = FLinearColor::Gray;
 };
 
-class CONTENTBROWSERDATA_API FContentBrowserDataDragDropOp : public FAssetDragDropOp
+class FContentBrowserDataDragDropOp : public FAssetDragDropOp
 {
 public:
 	DRAG_DROP_OPERATOR_TYPE(FContentBrowserDataDragDropOp, FAssetDragDropOp)
 
-	static TSharedRef<FContentBrowserDataDragDropOp> New(TArrayView<const FContentBrowserItem> InDraggedItems, FThumbnailOverrideParams InThumbnailOverrideParams = FThumbnailOverrideParams());
+	static UE_API TSharedRef<FContentBrowserDataDragDropOp> New(TArrayView<const FContentBrowserItem> InDraggedItems, FThumbnailOverrideParams InThumbnailOverrideParams = FThumbnailOverrideParams());
 
-	static TSharedRef<FContentBrowserDataDragDropOp> Legacy_New(TArrayView<const FAssetData> InAssetData, TArrayView<const FString> InAssetPaths = TArrayView<const FString>(), UActorFactory* InActorFactory = nullptr);
+	static UE_API TSharedRef<FContentBrowserDataDragDropOp> Legacy_New(TArrayView<const FAssetData> InAssetData, TArrayView<const FString> InAssetPaths = TArrayView<const FString>(), UActorFactory* InActorFactory = nullptr);
 
 	const TArray<FContentBrowserItem>& GetDraggedItems() const
 	{
@@ -58,24 +60,24 @@ public:
 	}
 
 private:
-	void Init(TArrayView<const FContentBrowserItem> InDraggedItems);
-	void LegacyInit(TArrayView<const FAssetData> InAssetData, TArrayView<const FString> InAssetPaths, UActorFactory* ActorFactory);
-	virtual void InitThumbnail() override;
+	UE_API void Init(TArrayView<const FContentBrowserItem> InDraggedItems);
+	UE_API void LegacyInit(TArrayView<const FAssetData> InAssetData, TArrayView<const FString> InAssetPaths, UActorFactory* ActorFactory);
+	UE_API virtual void InitThumbnail() override;
 
-	virtual bool HasFiles() const override;
-	virtual bool HasFolders() const override;
+	UE_API virtual bool HasFiles() const override;
+	UE_API virtual bool HasFolders() const override;
 
-	virtual int32 GetTotalCount() const override;
-	virtual FText GetFirstItemText() const override;
+	UE_API virtual int32 GetTotalCount() const override;
+	UE_API virtual FText GetFirstItemText() const override;
 
 	/** Return whether or not it should override the thumbnail widget */
-	bool ShouldOverrideThumbnailWidget() const;
+	UE_API bool ShouldOverrideThumbnailWidget() const;
 
 	/** Return the folder widget or nullptr if the BrushName is not valid */
-	TSharedPtr<SWidget> GetFolderWidget() const;
+	UE_API TSharedPtr<SWidget> GetFolderWidget() const;
 
 	/** Return the folder widget for DragAndDrop or nullptr if there is no basic folder widget  */
-	TSharedPtr<SWidget> GetFolderWidgetDragAndDrop() const;
+	UE_API TSharedPtr<SWidget> GetFolderWidgetDragAndDrop() const;
 
 private:
 	FName FolderBrushName;
@@ -85,3 +87,5 @@ private:
 	TArray<FContentBrowserItem> DraggedFiles;
 	TArray<FContentBrowserItem> DraggedFolders;
 };
+
+#undef UE_API
diff --git a/Engine/Source/Editor/ContentBrowserData/Public/ContentBrowserDataFilter.h b/Engine/Source/Editor/ContentBrowserData/Public/ContentBrowserDataFilter.h
--- a/Engine/Source/Editor/ContentBrowserData/Public/ContentBrowserDataFilter.h
+++ b/Engine/Source/Editor/ContentBrowserData/Public/ContentBrowserDataFilter.h
@@ -20,6 +20,8 @@
 
 #include "ContentBrowserDataFilter.generated.h"
 
+#define UE_API CONTENTBROWSERDATA_API
+
 class FNamePermissionList;
 class FPathPermissionList;
 class UContentBrowserDataSource;
@@ -74,7 +76,7 @@ ENUM_CLASS_FLAGS(EContentBrowserItemAttributeFilter);
  * Structure used to filter folders by their path
  * @note All implementations of this filter should copy all dependent data on creation such that all calls to HideFolderIfEmpty for a given path return the same value for a instance of the implementation.
  */
-class CONTENTBROWSERDATA_API IContentBrowserHideFolderIfEmptyFilter
+class IContentBrowserHideFolderIfEmptyFilter
 {
 public:
 	IContentBrowserHideFolderIfEmptyFilter() = default;
@@ -87,15 +89,15 @@ public:
  * A list of typed filter structs and their associated data.
  * This allows systems to add new filter types that the core Content Browser data module doesn't know about.
  */
-struct CONTENTBROWSERDATA_API FContentBrowserDataFilterList
+struct FContentBrowserDataFilterList
 {
 public:
 	/** Constructor */
 	FContentBrowserDataFilterList() = default;
 
 	/** Copy support */
-	FContentBrowserDataFilterList(const FContentBrowserDataFilterList& InOther);
-	FContentBrowserDataFilterList& operator=(const FContentBrowserDataFilterList& InOther);
+	UE_API FContentBrowserDataFilterList(const FContentBrowserDataFilterList& InOther);
+	UE_API FContentBrowserDataFilterList& operator=(const FContentBrowserDataFilterList& InOther);
 
 	/** Move support */
 	FContentBrowserDataFilterList(FContentBrowserDataFilterList&&) = default;
@@ -107,7 +109,7 @@ public:
 	{
 		return *static_cast<T*>(FindOrAddFilter(TBaseStructure<T>::Get()));
 	}
-	void* FindOrAddFilter(const UScriptStruct* InFilterType);
+	UE_API void* FindOrAddFilter(const UScriptStruct* InFilterType);
 
 	/** Set the filter associated with the given type, replacing any instance of this type that may exist in the list */
 	template <typename T>
@@ -115,7 +117,7 @@ public:
 	{
 		SetFilter(TBaseStructure<T>::Get(), &InFilter);
 	}
-	void SetFilter(const UScriptStruct* InFilterType, const void* InFilterData);
+	UE_API void SetFilter(const UScriptStruct* InFilterType, const void* InFilterData);
 
 	/** Find the filter associated with the given type, if it exists in the list */
 	template <typename T>
@@ -123,7 +125,7 @@ public:
 	{
 		return static_cast<const T*>(FindFilter(TBaseStructure<T>::Get()));
 	}
-	const void* FindFilter(const UScriptStruct* InFilterType) const;
+	UE_API const void* FindFilter(const UScriptStruct* InFilterType) const;
 
 	/** Find the filter associated with the given type, if it exists in the list */
 	template <typename T>
@@ -154,22 +156,22 @@ public:
 	{
 		RemoveFilter(TBaseStructure<T>::Get());
 	}
-	void RemoveFilter(const UScriptStruct* InFilterType);
+	UE_API void RemoveFilter(const UScriptStruct* InFilterType);
 
 	/** Remove all filters in the list */
-	void ClearFilters();
+	UE_API void ClearFilters();
 	
-	TArray<const UScriptStruct*> GetFilterTypes() const;
+	UE_API TArray<const UScriptStruct*> GetFilterTypes() const;
 
 private:
 	/** Set the contents of this list to be a deep copy of the contents of the other list */
-	void SetTo(const FContentBrowserDataFilterList& InOther);
+	UE_API void SetTo(const FContentBrowserDataFilterList& InOther);
 
 	/** Array of typed filter structs */
 	TArray<FStructOnScope> TypedFilters;
 };
 
-struct CONTENTBROWSERDATA_API FContentBrowserDataFilterCacheID
+struct FContentBrowserDataFilterCacheID
 {
 public:
 	operator bool() const
@@ -198,7 +200,7 @@ private:
  * @note The compiled version of this, FContentBrowserDataCompiledFilter, is produced via UContentBrowserDataSubsystem::CompileFilter.
  */
 USTRUCT(BlueprintType)
-struct CONTENTBROWSERDATA_API FContentBrowserDataFilter
+struct FContentBrowserDataFilter
 {
 	GENERATED_BODY()
 
@@ -230,7 +232,7 @@ public:
  * A filter used to control what is returned from Content Browser data queries.
  * @note The source version of this, FContentBrowserDataFilter, is used with UContentBrowserDataSubsystem::CompileFilter to produce a compiled filter.
  */
-struct CONTENTBROWSERDATA_API FContentBrowserDataCompiledFilter
+struct FContentBrowserDataCompiledFilter
 {
 	/** Flags controlling which item types should be included in this query */
 	EContentBrowserItemTypeFilter ItemTypeFilter = EContentBrowserItemTypeFilter::IncludeAll;
@@ -250,7 +252,7 @@ struct CONTENTBROWSERDATA_API FContentBrowserDataCompiledFilter
  * @note This will typically limit your query to returning assets.
  */
 USTRUCT(BlueprintType)
-struct CONTENTBROWSERDATA_API FContentBrowserDataObjectFilter
+struct FContentBrowserDataObjectFilter
 {
 	GENERATED_BODY()
 
@@ -279,7 +281,7 @@ public:
  * @note This will typically limit your query to returning assets.
  */
 USTRUCT(BlueprintType)
-struct CONTENTBROWSERDATA_API FContentBrowserDataPackageFilter
+struct FContentBrowserDataPackageFilter
 {
 	GENERATED_BODY()
 
@@ -317,7 +319,7 @@ public:
  * @note This will typically limit your query to returning assets.
  */
 USTRUCT(BlueprintType)
-struct CONTENTBROWSERDATA_API FContentBrowserDataClassFilter
+struct FContentBrowserDataClassFilter
 {
 	GENERATED_BODY()
 
@@ -347,7 +349,7 @@ public:
  * @note This will typically limit your query to items that support being inside a collection.
  */
 USTRUCT(BlueprintType)
-struct CONTENTBROWSERDATA_API FContentBrowserDataCollectionFilter
+struct FContentBrowserDataCollectionFilter
 {
 	GENERATED_BODY()
 
@@ -378,7 +380,7 @@ public:
 /*
  * Structure used to optionally filter folders by their broad contents
  */
-struct CONTENTBROWSERDATA_API FContentBrowserFolderContentsFilter
+struct FContentBrowserFolderContentsFilter
 {
 public:
 	EContentBrowserItemTypeFilter ItemTypeFilter = EContentBrowserItemTypeFilter::IncludeAll;
@@ -399,7 +401,7 @@ public:
  * @note This will restrict user interaction with those asset in the content browser and only affect the asset in the folders specified in the permission list
  */
 USTRUCT()
-struct CONTENTBROWSERDATA_API FContentBrowserDataUnsupportedClassFilter
+struct FContentBrowserDataUnsupportedClassFilter
 {
 	GENERATED_BODY()
 
@@ -416,7 +418,7 @@ public:
  * 2) When compiling the filters pass the Cache ID Owner to the cacheID of the ContentBrowserDataFilter.
  * 3) When the filter settings change call RemoveUnusedCachedData to clean the cache and to remove the potential invalid data.
  */
-struct CONTENTBROWSERDATA_API FContentBrowserDataFilterCacheIDOwner
+struct FContentBrowserDataFilterCacheIDOwner
 {
 public:
 	operator FContentBrowserDataFilterCacheID() const
@@ -454,13 +456,13 @@ public:
 	FContentBrowserDataFilterCacheIDOwner(const FContentBrowserDataFilterCacheIDOwner&) = delete;
 	FContentBrowserDataFilterCacheIDOwner& operator=(const FContentBrowserDataFilterCacheIDOwner&) = delete;
 
-	void Initialaze(UContentBrowserDataSubsystem* InContentBrowserDataSubsystem);
+	UE_API void Initialaze(UContentBrowserDataSubsystem* InContentBrowserDataSubsystem);
 
-	void RemoveUnusedCachedData(TArrayView<const FName> InVirtualPathsInUse, const FContentBrowserDataFilter& DataFilter) const;
+	UE_API void RemoveUnusedCachedData(TArrayView<const FName> InVirtualPathsInUse, const FContentBrowserDataFilter& DataFilter) const;
 
-	void ClearCachedData() const;
+	UE_API void ClearCachedData() const;
 
-	void Reset();
+	UE_API void Reset();
 
 private:
 	friend UContentBrowserDataSubsystem;
@@ -468,3 +470,5 @@ private:
 	int64 ID = INDEX_NONE;
 	TWeakObjectPtr<UContentBrowserDataSubsystem> DataSource;
 };
+
+#undef UE_API
diff --git a/Engine/Source/Editor/ContentBrowserData/Public/ContentBrowserDataLegacyBridge.h b/Engine/Source/Editor/ContentBrowserData/Public/ContentBrowserDataLegacyBridge.h
--- a/Engine/Source/Editor/ContentBrowserData/Public/ContentBrowserDataLegacyBridge.h
+++ b/Engine/Source/Editor/ContentBrowserData/Public/ContentBrowserDataLegacyBridge.h
@@ -25,7 +25,7 @@ DECLARE_DELEGATE_FiveParams(FOnCreateNewAsset, const FName /*DefaultAssetName*/,
 
 /** Filter data used to provide legacy information to the asset data source */
 USTRUCT()
-struct CONTENTBROWSERDATA_API FContentBrowserDataLegacyFilter
+struct FContentBrowserDataLegacyFilter
 {
 	GENERATED_BODY()
 
diff --git a/Engine/Source/Editor/ContentBrowserData/Public/ContentBrowserDataMenuContexts.h b/Engine/Source/Editor/ContentBrowserData/Public/ContentBrowserDataMenuContexts.h
--- a/Engine/Source/Editor/ContentBrowserData/Public/ContentBrowserDataMenuContexts.h
+++ b/Engine/Source/Editor/ContentBrowserData/Public/ContentBrowserDataMenuContexts.h
@@ -18,8 +18,8 @@ enum class EContentBrowserDataMenuContext_AddNewMenuDomain : uint8
 	PathView,
 };
 
-UCLASS()
-class CONTENTBROWSERDATA_API UContentBrowserDataMenuContext_AddNewMenu : public UObject
+UCLASS(MinimalAPI)
+class UContentBrowserDataMenuContext_AddNewMenu : public UObject
 {
 	GENERATED_BODY()
 
@@ -44,8 +44,8 @@ public:
 	const FContentBrowserInstanceConfig* OwningInstanceConfig;
 };
 
-UCLASS()
-class CONTENTBROWSERDATA_API UContentBrowserDataMenuContext_FolderMenu : public UObject
+UCLASS(MinimalAPI)
+class UContentBrowserDataMenuContext_FolderMenu : public UObject
 {
 	GENERATED_BODY()
 
@@ -59,8 +59,8 @@ public:
 	TWeakPtr<SWidget> ParentWidget;
 };
 
-UCLASS()
-class CONTENTBROWSERDATA_API UContentBrowserDataMenuContext_FileMenu : public UObject
+UCLASS(MinimalAPI)
+class UContentBrowserDataMenuContext_FileMenu : public UObject
 {
 	GENERATED_BODY()
 
@@ -96,8 +96,8 @@ public:
 	FOnRefreshView OnRefreshView;
 };
 
-UCLASS()
-class CONTENTBROWSERDATA_API UContentBrowserDataMenuContext_DragDropMenu : public UObject
+UCLASS(MinimalAPI)
+class UContentBrowserDataMenuContext_DragDropMenu : public UObject
 {
 	GENERATED_BODY()
 
diff --git a/Engine/Source/Editor/ContentBrowserData/Public/ContentBrowserDataSource.h b/Engine/Source/Editor/ContentBrowserData/Public/ContentBrowserDataSource.h
--- a/Engine/Source/Editor/ContentBrowserData/Public/ContentBrowserDataSource.h
+++ b/Engine/Source/Editor/ContentBrowserData/Public/ContentBrowserDataSource.h
@@ -21,6 +21,8 @@
 
 #include "ContentBrowserDataSource.generated.h"
 
+#define UE_API CONTENTBROWSERDATA_API
+
 class FAssetThumbnail;
 class FDragDropEvent;
 class FDragDropOperation;
@@ -118,42 +120,42 @@ namespace ContentBrowserItemAttributes
  * In that sense you can think of this like a C API, where the data source returns an opaque object that is later passed back into the data source 
  * functions so that they can interpret the opaque object and provide functionality for it.
  */
-UCLASS(Abstract, Transient)
-class CONTENTBROWSERDATA_API UContentBrowserDataSource : public UObject, public IModularFeature
+UCLASS(MinimalAPI, Abstract, Transient)
+class UContentBrowserDataSource : public UObject, public IModularFeature
 {
 	GENERATED_BODY()
 
 public:
 	//~ UObject interface
-	virtual void BeginDestroy() override final;
+	UE_API virtual void BeginDestroy() override final;
 
 	/**
 	 * Get the name used when registering data source modular feature instances for use with the Content Browser Data Subsystem.
 	 */
-	static FName GetModularFeatureTypeName();
+	static UE_API FName GetModularFeatureTypeName();
 
 	/**
 	 * Register this data source instance for use with the Content Browser Data Subsystem.
 	 * @note This does not activate the data source. @see the EnabledDataSources array in the Content Browser Data Subsystem for more information on activation.
 	 */
-	void RegisterDataSource();
+	UE_API void RegisterDataSource();
 
 	/**
 	 * Unregister this data source instance from the Content Browser Data Subsystem.
 	 */
-	void UnregisterDataSource();
+	UE_API void UnregisterDataSource();
 
 	/**
 	 * Set the data sink that can be used to communicate with the Content Browser Data Subsystem.
 	 * This is set to a valid instance when the data source is activated, and set to null when it is deactivated.
 	 */
-	void SetDataSink(IContentBrowserItemDataSink* InDataSink);
+	UE_API void SetDataSink(IContentBrowserItemDataSink* InDataSink);
 
 	/**
 	 * True if this data source is currently initialized.
 	 * A data source should be initialized before it is registered.
 	 */
-	bool IsInitialized() const;
+	UE_API bool IsInitialized() const;
 
 	/**
 	 * Initialize this data source instance, optionally registering it once the initialization has finished (@see RegisterDataSource).
@@ -161,24 +163,24 @@ public:
 	 *
 	 * @param InAutoRegister True to automatically register this instance once initialization has finished.
 	 */
-	void Initialize(const bool InAutoRegister = true);
+	UE_API void Initialize(const bool InAutoRegister = true);
 
 	/**
 	 * Shutdown this data source instance.
 	 * @note This is called during the BeginDestroy phase of object destruction.
 	 */
-	virtual void Shutdown();
+	UE_API virtual void Shutdown();
 
 	/**
 	 * Tick this data source instance.
 	 * @note Called once every 0.1 seconds, prior to the Content Browser Data Subsystem emitting any pending item update notifications.
 	 */
-	virtual void Tick(const float InDeltaTime);
+	UE_API virtual void Tick(const float InDeltaTime);
 	/**
 	 * Test whether the given virtual path is under the virtual mount root that was passed to Initialize.
 	 * @note This also returns true if the given virtual path *is* the virtual mount root.
 	 */
-	bool IsVirtualPathUnderMountRoot(const FName InPath) const;
+	UE_API bool IsVirtualPathUnderMountRoot(const FName InPath) const;
 
 	/**
 	 * Given a path and a data filter, produce an optimized filter that can be used to efficiently enumerate items that match it, and also query whether an item would pass it.
@@ -190,7 +192,7 @@ public:
 	 * @param InFilter Rules describing how items should be filtered.
 	 * @param OutCompiledFilter The compiled filter instance to fill with the result.
 	 */
-	virtual void CompileFilter(const FName InPath, const FContentBrowserDataFilter& InFilter, FContentBrowserDataCompiledFilter& OutCompiledFilter);
+	UE_API virtual void CompileFilter(const FName InPath, const FContentBrowserDataFilter& InFilter, FContentBrowserDataCompiledFilter& OutCompiledFilter);
 
 	/**
 	 * Enumerate items that match the given compiled filter, invoking the callback for each matching item.
@@ -201,8 +203,8 @@ public:
 	 * @param InCallback The function to invoke for each matching item (return true to continue enumeration).
 	 * @param InSink Parameter provided to receive output as either a callback or an array.
 	 */
-	virtual void EnumerateItemsMatchingFilter(const FContentBrowserDataCompiledFilter& InFilter, TFunctionRef<bool(FContentBrowserItemData&&)> InCallback);
-	virtual void EnumerateItemsMatchingFilter(const FContentBrowserDataCompiledFilter& InFilter, const TGetOrEnumerateSink<FContentBrowserItemData>& InSink);
+	UE_API virtual void EnumerateItemsMatchingFilter(const FContentBrowserDataCompiledFilter& InFilter, TFunctionRef<bool(FContentBrowserItemData&&)> InCallback);
+	UE_API virtual void EnumerateItemsMatchingFilter(const FContentBrowserDataCompiledFilter& InFilter, const TGetOrEnumerateSink<FContentBrowserItemData>& InSink);
 
 	/**
 	 * Enumerate items that have the given virtual path, optionally filtering by type, and invoking the callback for each matching item.
@@ -213,8 +215,8 @@ public:
 	 * @param InCallback The function to invoke for each matching item (return true to continue enumeration).
 	 * @param InSink Parameter provided to receive output as either a callback or an array.
 	 */
-	virtual void EnumerateItemsAtPath(const FName InPath, const EContentBrowserItemTypeFilter InItemTypeFilter, TFunctionRef<bool(FContentBrowserItemData&&)> InCallback);
-	virtual void EnumerateItemsAtPath(const FName InPath, const EContentBrowserItemTypeFilter InItemTypeFilter, const TGetOrEnumerateSink<FContentBrowserItemData>& InSink);
+	UE_API virtual void EnumerateItemsAtPath(const FName InPath, const EContentBrowserItemTypeFilter InItemTypeFilter, TFunctionRef<bool(FContentBrowserItemData&&)> InCallback);
+	UE_API virtual void EnumerateItemsAtPath(const FName InPath, const EContentBrowserItemTypeFilter InItemTypeFilter, const TGetOrEnumerateSink<FContentBrowserItemData>& InSink);
 	
 	/**
 	 * Enumerate the items (folders and/or files) that exist at the given content browser paths.
@@ -224,7 +226,7 @@ public:
 	 * @param InItemTypeFilter The types of items we want to find.
 	 * @param InCallback The function to invoke for each matching item (return true to continue enumeration).
 	 */
-	virtual bool EnumerateItemsAtPaths(const TArrayView<FContentBrowserItemPath> InPaths, const EContentBrowserItemTypeFilter InItemTypeFilter, TFunctionRef<bool(FContentBrowserItemData&&)> InCallback);
+	UE_API virtual bool EnumerateItemsAtPaths(const TArrayView<FContentBrowserItemPath> InPaths, const EContentBrowserItemTypeFilter InItemTypeFilter, TFunctionRef<bool(FContentBrowserItemData&&)> InCallback);
 
 	/**
 	 * Enumerate the items (files) that exist for the given objects.
@@ -233,7 +235,7 @@ public:
 	 * @param InObjects The objects to enumerate
 	 * @param InCallback The function to invoke for each matching item (return true to continue enumeration).
 	 */
-	virtual bool EnumerateItemsForObjects(const TArrayView<UObject*> InObjects, TFunctionRef<bool(FContentBrowserItemData&&)> InCallback);
+	UE_API virtual bool EnumerateItemsForObjects(const TArrayView<UObject*> InObjects, TFunctionRef<bool(FContentBrowserItemData&&)> InCallback);
 
 	/**
 	 * Get a list of other paths that the data source may be using to represent a specific path
@@ -241,21 +243,21 @@ public:
 	 * @param The internal path (or object path) of an asset to get aliases for
 	 * @return All alternative paths that represent the input path (not including the input path itself)
 	 */
-	virtual TArray<FContentBrowserItemPath> GetAliasesForPath(const FSoftObjectPath& InInternalPath) const;
+	UE_API virtual TArray<FContentBrowserItemPath> GetAliasesForPath(const FSoftObjectPath& InInternalPath) const;
 
 	UE_DEPRECATED(5.1, "FNames containing full asset paths are deprecated. Use FSoftObjectPath instead.")
-	TArray<FContentBrowserItemPath> GetAliasesForPath(FName InInternalPath) const;
+	UE_API TArray<FContentBrowserItemPath> GetAliasesForPath(FName InInternalPath) const;
 
 	/**
 	 * Query whether this data source instance is currently discovering content, and retrieve an optional status message that can be shown in the UI.
 	 */
-	virtual bool IsDiscoveringItems(FText* OutStatus = nullptr);
+	UE_API virtual bool IsDiscoveringItems(FText* OutStatus = nullptr);
 
 	/**
 	 * If possible, attempt to prioritize content discovery for the given virtual path.
 	 * @note This will be called in response to a user attempting to browse content under the given virtual path.
 	 */
-	virtual bool PrioritizeSearchPath(const FName InPath);
+	UE_API virtual bool PrioritizeSearchPath(const FName InPath);
 
 	/**
 	 * Query whether the given virtual folder should be visible in the UI.
@@ -264,7 +266,7 @@ public:
 	 * @param Flags Currently unused
 	 * @param ContentsFilter Optional filter for limiting visibility to only folders with a certain type of contents based on desired view filtering. If this is unset, folders may be visible regardless of their contents or empty status.
 	 */
-	virtual bool IsFolderVisible(const FName Path, const EContentBrowserIsFolderVisibleFlags Flags, const FContentBrowserFolderContentsFilter& ContentsFilter);
+	UE_API virtual bool IsFolderVisible(const FName Path, const EContentBrowserIsFolderVisibleFlags Flags, const FContentBrowserFolderContentsFilter& ContentsFilter);
 
 	/*
 	 * Query whether a folder can be created at the given virtual path, optionally providing error information if it cannot.
@@ -274,7 +276,7 @@ public:
 	 *
 	 * @return True if the folder can be created, false otherwise.
 	 */
-	virtual bool CanCreateFolder(const FName InPath, FText* OutErrorMsg);
+	UE_API virtual bool CanCreateFolder(const FName InPath, FText* OutErrorMsg);
 
 	/*
 	 * Attempt to begin the process of asynchronously creating a folder at the given virtual path, populating a temporary item that can be finalized or canceled by the user.
@@ -284,7 +286,7 @@ public:
 	 *
 	 * @return True if the pending folder was created, false otherwise.
 	 */
-	virtual bool CreateFolder(const FName InPath, const TSharedPtr<IContentBrowserHideFolderIfEmptyFilter>& HideFolderIfEmptyFilter, FContentBrowserItemDataTemporaryContext& OutPendingItem);
+	UE_API virtual bool CreateFolder(const FName InPath, const TSharedPtr<IContentBrowserHideFolderIfEmptyFilter>& HideFolderIfEmptyFilter, FContentBrowserItemDataTemporaryContext& OutPendingItem);
 
 	/*
 	 * Query whether the given item passes the given compiled filter. Should be called after ConvertItemForFilter
@@ -296,7 +298,7 @@ public:
 	 *
 	 * @return True if the item passes the filter, false otherwise.
 	 */
-	virtual bool DoesItemPassFilter(const FContentBrowserItemData& InItem, const FContentBrowserDataCompiledFilter& InFilter);
+	UE_API virtual bool DoesItemPassFilter(const FContentBrowserItemData& InItem, const FContentBrowserDataCompiledFilter& InFilter);
 
 
 	/*
@@ -309,7 +311,7 @@ public:
 	 * 
 	 * @return True if the item was converted by the filter.
 	 */
-	virtual bool ConvertItemForFilter(FContentBrowserItemData& Item, const FContentBrowserDataCompiledFilter& InFilter);
+	UE_API virtual bool ConvertItemForFilter(FContentBrowserItemData& Item, const FContentBrowserDataCompiledFilter& InFilter);
 
 	/**
 	 * Query the value of the given attribute on the given item.
@@ -321,7 +323,7 @@ public:
 	 *
 	 * @return True if the attribute was found and its value was filled, false otherwise.
 	 */
-	virtual bool GetItemAttribute(const FContentBrowserItemData& InItem, const bool InIncludeMetaData, const FName InAttributeKey, FContentBrowserItemDataAttributeValue& OutAttributeValue);
+	UE_API virtual bool GetItemAttribute(const FContentBrowserItemData& InItem, const bool InIncludeMetaData, const FName InAttributeKey, FContentBrowserItemDataAttributeValue& OutAttributeValue);
 
 	/**
 	 * Query the values of all attributes on the given item.
@@ -332,7 +334,7 @@ public:
 	 *
 	 * @return True attributes were found and their values reported, false otherwise.
 	 */
-	virtual bool GetItemAttributes(const FContentBrowserItemData& InItem, const bool InIncludeMetaData, FContentBrowserItemDataAttributeValues& OutAttributeValues);
+	UE_API virtual bool GetItemAttributes(const FContentBrowserItemData& InItem, const bool InIncludeMetaData, FContentBrowserItemDataAttributeValues& OutAttributeValues);
 
 	/**
 	 * Query the physical (on-disk) path of the given item.
@@ -343,7 +345,7 @@ public:
 	 *
 	 * @return True if the physical path was found and the result filled, false otherwise.
 	 */
-	virtual bool GetItemPhysicalPath(const FContentBrowserItemData& InItem, FString& OutDiskPath);
+	UE_API virtual bool GetItemPhysicalPath(const FContentBrowserItemData& InItem, FString& OutDiskPath);
 
 	/**
 	 * Query the EAssetAccessSpecifier of the item
@@ -353,7 +355,7 @@ public:
 	 *
 	 * @return True if OutAssetAccessSpecifier was filled.
 	 */
-	virtual bool GetItemAssetAccessSpecifier(const FContentBrowserItemData& InItem, EAssetAccessSpecifier& OutAssetAccessSpecifier);
+	UE_API virtual bool GetItemAssetAccessSpecifier(const FContentBrowserItemData& InItem, EAssetAccessSpecifier& OutAssetAccessSpecifier);
 
 	/**
 	 * Query the item's ability to modify its EAssetAccessSpecifier
@@ -362,7 +364,7 @@ public:
 	 *
 	 * @return True if the item has an EAssetAccessSpecifier that can be modified.
 	 */
-	virtual bool CanModifyItemAssetAccessSpecifier(const FContentBrowserItemData& InItem);
+	UE_API virtual bool CanModifyItemAssetAccessSpecifier(const FContentBrowserItemData& InItem);
 
 	/**
 	 * Query whether the given item is considered dirty (ie, has unsaved changes).
@@ -371,7 +373,7 @@ public:
 	 *
 	 * @return True if the item is considered dirty, false otherwise.
 	 */
-	virtual bool IsItemDirty(const FContentBrowserItemData& InItem);
+	UE_API virtual bool IsItemDirty(const FContentBrowserItemData& InItem);
 
 	/**
 	 * Query whether the given item is can be edited, optionally providing error information if it cannot.
@@ -381,7 +383,7 @@ public:
 	 *
 	 * @return True if the item can be edited, false otherwise.
 	 */
-	virtual bool CanEditItem(const FContentBrowserItemData& InItem, FText* OutErrorMsg);
+	UE_API virtual bool CanEditItem(const FContentBrowserItemData& InItem, FText* OutErrorMsg);
 
 	/**
 	 * Attempt to open the given item for editing.
@@ -390,7 +392,7 @@ public:
 	 *
 	 * @return True if the item was opened for editing, false otherwise.
 	 */
-	virtual bool EditItem(const FContentBrowserItemData& InItem);
+	UE_API virtual bool EditItem(const FContentBrowserItemData& InItem);
 
 	/**
 	 * Attempt to open the given item for editing.
@@ -400,7 +402,7 @@ public:
 	 *
 	 * @return True if any items were opened for editing, false otherwise.
 	 */
-	virtual bool BulkEditItems(TArrayView<const FContentBrowserItemData> InItems);
+	UE_API virtual bool BulkEditItems(TArrayView<const FContentBrowserItemData> InItems);
 
 	/**
      * Query whether the given item is can be viewed (a read-only asset editor), optionally providing error information if it cannot.
@@ -410,7 +412,7 @@ public:
      *
      * @return True if the item can be viewed in a read-only editor, false otherwise.
      */
-    virtual bool CanViewItem(const FContentBrowserItemData& InItem, FText* OutErrorMsg);
+    UE_API virtual bool CanViewItem(const FContentBrowserItemData& InItem, FText* OutErrorMsg);
 
     /**
      * Attempt to open the given item for read-only viewing.
@@ -419,7 +421,7 @@ public:
      *
      * @return True if the item was opened for read-only viewing, false otherwise.
      */
-    virtual bool ViewItem(const FContentBrowserItemData& InItem);
+    UE_API virtual bool ViewItem(const FContentBrowserItemData& InItem);
 
     /**
      * Attempt to open the given items for read-only viewing.
@@ -429,7 +431,7 @@ public:
      *
      * @return True if any items were opened for read-only viewing, false otherwise.
      */
-    virtual bool BulkViewItems(TArrayView<const FContentBrowserItemData> InItems);
+    UE_API virtual bool BulkViewItems(TArrayView<const FContentBrowserItemData> InItems);
 
 	/**
 	 * Query whether the given item is can be previewed, optionally providing error information if it cannot.
@@ -439,7 +441,7 @@ public:
 	 *
 	 * @return True if the item can be previewed, false otherwise.
 	 */
-	virtual bool CanPreviewItem(const FContentBrowserItemData& InItem, FText* OutErrorMsg);
+	UE_API virtual bool CanPreviewItem(const FContentBrowserItemData& InItem, FText* OutErrorMsg);
 
 	/**
 	 * Attempt to preview the given item.
@@ -448,7 +450,7 @@ public:
 	 *
 	 * @return True if the item was previewed, false otherwise.
 	 */
-	virtual bool PreviewItem(const FContentBrowserItemData& InItem);
+	UE_API virtual bool PreviewItem(const FContentBrowserItemData& InItem);
 
 	/**
 	 * Attempt to preview the given items.
@@ -458,7 +460,7 @@ public:
 	 *
 	 * @return True if any items were previewed, false otherwise.
 	 */
-	virtual bool BulkPreviewItems(TArrayView<const FContentBrowserItemData> InItems);
+	UE_API virtual bool BulkPreviewItems(TArrayView<const FContentBrowserItemData> InItems);
 
 	/**
 	 * Query whether the given item is can be duplicated, optionally providing error information if it cannot.
@@ -468,7 +470,7 @@ public:
 	 *
 	 * @return True if the item can be duplicated, false otherwise.
 	 */
-	virtual bool CanDuplicateItem(const FContentBrowserItemData& InItem, FText* OutErrorMsg);
+	UE_API virtual bool CanDuplicateItem(const FContentBrowserItemData& InItem, FText* OutErrorMsg);
 
 	/*
 	 * Attempt to begin the process of asynchronously duplicating the given item, populating a temporary item that can be finalized or canceled by the user.
@@ -479,7 +481,7 @@ public:
 	 *
 	 * @return True if the pending item was created, false otherwise.
 	 */
-	virtual bool DuplicateItem(const FContentBrowserItemData& InItem, FContentBrowserItemDataTemporaryContext& OutPendingItem);
+	UE_API virtual bool DuplicateItem(const FContentBrowserItemData& InItem, FContentBrowserItemDataTemporaryContext& OutPendingItem);
 
 	/**
 	 * Attempt to synchronously duplicate the given items.
@@ -490,7 +492,7 @@ public:
 	 *
 	 * @return True if any items were duplicated, false otherwise.
 	 */
-	virtual bool BulkDuplicateItems(TArrayView<const FContentBrowserItemData> InItems, TArray<FContentBrowserItemData>& OutNewItems);
+	UE_API virtual bool BulkDuplicateItems(TArrayView<const FContentBrowserItemData> InItems, TArray<FContentBrowserItemData>& OutNewItems);
 
 	/**
 	 * Query whether the given item is can be saved, optionally providing error information if it cannot.
@@ -501,7 +503,7 @@ public:
 	 *
 	 * @return True if the item can be saved, false otherwise.
 	 */
-	virtual bool CanSaveItem(const FContentBrowserItemData& InItem, const EContentBrowserItemSaveFlags InSaveFlags, FText* OutErrorMsg);
+	UE_API virtual bool CanSaveItem(const FContentBrowserItemData& InItem, const EContentBrowserItemSaveFlags InSaveFlags, FText* OutErrorMsg);
 
 	/**
 	 * Attempt to save the given item.
@@ -511,7 +513,7 @@ public:
 	 *
 	 * @return True if the item was saved, false otherwise.
 	 */
-	virtual bool SaveItem(const FContentBrowserItemData& InItem, const EContentBrowserItemSaveFlags InSaveFlags);
+	UE_API virtual bool SaveItem(const FContentBrowserItemData& InItem, const EContentBrowserItemSaveFlags InSaveFlags);
 
 	/**
 	 * Attempt to save the given items.
@@ -522,7 +524,7 @@ public:
 	 *
 	 * @return True if any items were saved, false otherwise.
 	 */
-	virtual bool BulkSaveItems(TArrayView<const FContentBrowserItemData> InItems, const EContentBrowserItemSaveFlags InSaveFlags);
+	UE_API virtual bool BulkSaveItems(TArrayView<const FContentBrowserItemData> InItems, const EContentBrowserItemSaveFlags InSaveFlags);
 
 	/**
 	 * Query whether the given item is can be deleted, optionally providing error information if it cannot.
@@ -532,7 +534,7 @@ public:
 	 *
 	 * @return True if the item can be deleted, false otherwise.
 	 */
-	virtual bool CanDeleteItem(const FContentBrowserItemData& InItem, FText* OutErrorMsg);
+	UE_API virtual bool CanDeleteItem(const FContentBrowserItemData& InItem, FText* OutErrorMsg);
 
 	/**
 	 * Attempt to delete the given item.
@@ -541,7 +543,7 @@ public:
 	 *
 	 * @return True if the item was deleted, false otherwise.
 	 */
-	virtual bool DeleteItem(const FContentBrowserItemData& InItem);
+	UE_API virtual bool DeleteItem(const FContentBrowserItemData& InItem);
 
 	/**
 	 * Attempt to delete the given items.
@@ -551,7 +553,7 @@ public:
 	 *
 	 * @return True if any items were deleted, false otherwise.
 	 */
-	virtual bool BulkDeleteItems(TArrayView<const FContentBrowserItemData> InItems);
+	UE_API virtual bool BulkDeleteItems(TArrayView<const FContentBrowserItemData> InItems);
 
 	/**
 	* Query whether the given item can be privatized, optionally providing error information if it cannot.
@@ -561,7 +563,7 @@ public:
 	* 
 	* @return True if the item was deleted, false otherwise.
 	*/
-	virtual bool CanPrivatizeItem(const FContentBrowserItemData& InItem, FText* OutErrorMsg);
+	UE_API virtual bool CanPrivatizeItem(const FContentBrowserItemData& InItem, FText* OutErrorMsg);
 
 	/**
 	* Attempt to mark the given item as private (NotExternallyReferenceable).
@@ -570,7 +572,7 @@ public:
 	*
 	* @return True if the item was marked private, false otherwise
 	*/
-	virtual bool PrivatizeItem(const FContentBrowserItemData& InItem, const EAssetAccessSpecifier InAssetAccessSpecifier = EAssetAccessSpecifier::Private);
+	UE_API virtual bool PrivatizeItem(const FContentBrowserItemData& InItem, const EAssetAccessSpecifier InAssetAccessSpecifier = EAssetAccessSpecifier::Private);
 
 	/**
 	* Attempt to mark the given items as InAssetAccessSpecifier
@@ -580,7 +582,7 @@ public:
 	* 
 	* @return True if any items were marked InAssetAccessSpecifier, false otherwise
 	*/
-	virtual bool BulkPrivatizeItems(TArrayView<const FContentBrowserItemData> InItems, const EAssetAccessSpecifier InAssetAccessSpecifier = EAssetAccessSpecifier::Private);
+	UE_API virtual bool BulkPrivatizeItems(TArrayView<const FContentBrowserItemData> InItems, const EAssetAccessSpecifier InAssetAccessSpecifier = EAssetAccessSpecifier::Private);
 
 	/**
 	 * Query whether the given item is can be renamed, optionally providing error information if it cannot.
@@ -591,7 +593,7 @@ public:
 	 *
 	 * @return True if the item can be renamed, false otherwise.
 	 */
-	virtual bool CanRenameItem(const FContentBrowserItemData& InItem, const FString* InNewName, const IContentBrowserHideFolderIfEmptyFilter* HideFolderIfEmptyFilter, FText* OutErrorMsg);
+	UE_API virtual bool CanRenameItem(const FContentBrowserItemData& InItem, const FString* InNewName, const IContentBrowserHideFolderIfEmptyFilter* HideFolderIfEmptyFilter, FText* OutErrorMsg);
 
 	/**
 	 * Attempt to rename the given item.
@@ -602,7 +604,7 @@ public:
 	 *
 	 * @return True if the item was renamed, false otherwise.
 	 */
-	virtual bool RenameItem(const FContentBrowserItemData& InItem, const FString& InNewName, FContentBrowserItemData& OutNewItem);
+	UE_API virtual bool RenameItem(const FContentBrowserItemData& InItem, const FString& InNewName, FContentBrowserItemData& OutNewItem);
 
 	/**
 	 * Query whether the given item is can be copied, optionally providing error information if it cannot.
@@ -613,7 +615,7 @@ public:
 	 *
 	 * @return True if the item can be copied, false otherwise.
 	 */
-	virtual bool CanCopyItem(const FContentBrowserItemData& InItem, const FName InDestPath, FText* OutErrorMsg);
+	UE_API virtual bool CanCopyItem(const FContentBrowserItemData& InItem, const FName InDestPath, FText* OutErrorMsg);
 
 	/**
 	 * Attempt to copy the given item to the given virtual path.
@@ -623,7 +625,7 @@ public:
 	 *
 	 * @return True if the item was copied, false otherwise.
 	 */
-	virtual bool CopyItem(const FContentBrowserItemData& InItem, const FName InDestPath);
+	UE_API virtual bool CopyItem(const FContentBrowserItemData& InItem, const FName InDestPath);
 
 	/**
 	 * Attempt to copy the given items to the given virtual path.
@@ -634,7 +636,7 @@ public:
 	 *
 	 * @return True if any items were copied, false otherwise.
 	 */
-	virtual bool BulkCopyItems(TArrayView<const FContentBrowserItemData> InItems, const FName InDestPath);
+	UE_API virtual bool BulkCopyItems(TArrayView<const FContentBrowserItemData> InItems, const FName InDestPath);
 
 	/**
 	 * Query whether the given item is can be moved, optionally providing error information if it cannot.
@@ -645,7 +647,7 @@ public:
 	 *
 	 * @return True if the item can be moved, false otherwise.
 	 */
-	virtual bool CanMoveItem(const FContentBrowserItemData& InItem, const FName InDestPath, FText* OutErrorMsg);
+	UE_API virtual bool CanMoveItem(const FContentBrowserItemData& InItem, const FName InDestPath, FText* OutErrorMsg);
 
 	/**
 	 * Attempt to move the given item to the given virtual path.
@@ -655,7 +657,7 @@ public:
 	 *
 	 * @return True if the item was moved, false otherwise.
 	 */
-	virtual bool MoveItem(const FContentBrowserItemData& InItem, const FName InDestPath);
+	UE_API virtual bool MoveItem(const FContentBrowserItemData& InItem, const FName InDestPath);
 
 	/**
 	 * Attempt to move the given items to the given virtual path.
@@ -666,7 +668,7 @@ public:
 	 *
 	 * @return True if any items were moved, false otherwise.
 	 */
-	virtual bool BulkMoveItems(TArrayView<const FContentBrowserItemData> InItems, const FName InDestPath);
+	UE_API virtual bool BulkMoveItems(TArrayView<const FContentBrowserItemData> InItems, const FName InDestPath);
 
 	/**
 	 * Attempt to append any path references for the given item to the given string.
@@ -677,7 +679,7 @@ public:
 	 *
 	 * @return True if references were appended, false otherwise.
 	 */
-	virtual bool AppendItemReference(const FContentBrowserItemData& InItem, FString& InOutStr);
+	UE_API virtual bool AppendItemReference(const FContentBrowserItemData& InItem, FString& InOutStr);
 
 	/**
 	 * Attempt to append any object path for the given item to the given string.
@@ -688,7 +690,7 @@ public:
 	 *
 	 * @return True if package name were appended, false otherwise.
 	 */
-	virtual bool AppendItemObjectPath(const FContentBrowserItemData& InItem, FString& InOutStr);
+	UE_API virtual bool AppendItemObjectPath(const FContentBrowserItemData& InItem, FString& InOutStr);
 
 	/**
 	 * Attempt to append any package name for the given item to the given string.
@@ -699,7 +701,7 @@ public:
 	 *
 	 * @return True if package name were appended, false otherwise.
 	 */
-	virtual bool AppendItemPackageName(const FContentBrowserItemData& InItem, FString& InOutStr);
+	UE_API virtual bool AppendItemPackageName(const FContentBrowserItemData& InItem, FString& InOutStr);
 
 	/**
 	 * Attempt to update the thumbnail associated with the given item.
@@ -709,7 +711,7 @@ public:
 	 *
 	 * @return True if the thumbnail was updated, false otherwise.
 	 */
-	virtual bool UpdateThumbnail(const FContentBrowserItemData& InItem, FAssetThumbnail& InThumbnail);
+	UE_API virtual bool UpdateThumbnail(const FContentBrowserItemData& InItem, FAssetThumbnail& InThumbnail);
 
 	/**
 	 * Called to provide custom drag and drop handling when starting a drag event.
@@ -719,7 +721,7 @@ public:
 	 *
 	 * @return A custom drag operation, or null to allow another data source (or the default handler) to deal with it instead.
 	 */
-	virtual TSharedPtr<FDragDropOperation> CreateCustomDragOperation(TArrayView<const FContentBrowserItemData> InItems);
+	UE_API virtual TSharedPtr<FDragDropOperation> CreateCustomDragOperation(TArrayView<const FContentBrowserItemData> InItems);
 
 	/**
 	 * Called to provide custom drag and drop handling when a drag event enters an item, such as performing validation and reporting error information.
@@ -729,7 +731,7 @@ public:
 	 *
 	 * @return True if this data source can handle the drag event (even if it won't because it's invalid), or false to allow another data source (or the default handler) to deal with it instead.
 	 */
-	virtual bool HandleDragEnterItem(const FContentBrowserItemData& InItem, const FDragDropEvent& InDragDropEvent);
+	UE_API virtual bool HandleDragEnterItem(const FContentBrowserItemData& InItem, const FDragDropEvent& InDragDropEvent);
 
 	/**
 	 * Called to provide custom drag and drop handling while a drag event is over an item, such as performing validation and reporting error information.
@@ -739,7 +741,7 @@ public:
 	 *
 	 * @return True if this data source can handle the drag event (even if it won't because it's invalid), or false to allow another data source (or the default handler) to deal with it instead.
 	 */
-	virtual bool HandleDragOverItem(const FContentBrowserItemData& InItem, const FDragDropEvent& InDragDropEvent);
+	UE_API virtual bool HandleDragOverItem(const FContentBrowserItemData& InItem, const FDragDropEvent& InDragDropEvent);
 
 	/**
 	 * Called to provide custom drag and drop handling when a drag event leaves an item, such as clearing any error information set during earlier validation.
@@ -749,7 +751,7 @@ public:
 	 *
 	 * @return True if this data source can handle the drag event (even if it won't because it's invalid), or false to allow another data source (or the default handler) to deal with it instead.
 	 */
-	virtual bool HandleDragLeaveItem(const FContentBrowserItemData& InItem, const FDragDropEvent& InDragDropEvent);
+	UE_API virtual bool HandleDragLeaveItem(const FContentBrowserItemData& InItem, const FDragDropEvent& InDragDropEvent);
 
 	/**
 	 * Called to provide custom drag and drop handling when a drag event is dropped on an item.
@@ -759,7 +761,7 @@ public:
 	 *
 	 * @return True if this data source can handle the drag event (even if it didn't because it was invalid), or false to allow another data source (or the default handler) to deal with it instead.
 	 */
-	virtual bool HandleDragDropOnItem(const FContentBrowserItemData& InItem, const FDragDropEvent& InDragDropEvent);
+	UE_API virtual bool HandleDragDropOnItem(const FContentBrowserItemData& InItem, const FDragDropEvent& InDragDropEvent);
 
 	/**
 	 * Attempt to retrieve the identifier that should be used when storing a reference to the given item within a collection.
@@ -769,7 +771,7 @@ public:
 	 *
 	 * @return True if the ID was retrieved, false otherwise.
 	 */
-	virtual bool TryGetCollectionId(const FContentBrowserItemData& InItem, FSoftObjectPath& OutCollectionId);
+	UE_API virtual bool TryGetCollectionId(const FContentBrowserItemData& InItem, FSoftObjectPath& OutCollectionId);
 
 	/**
 	 * Attempt to retrieve the package path associated with the given item.
@@ -781,7 +783,7 @@ public:
 	 *
 	 * @return True if the package path was retrieved, false otherwise.
 	 */
-	virtual bool Legacy_TryGetPackagePath(const FContentBrowserItemData& InItem, FName& OutPackagePath);
+	UE_API virtual bool Legacy_TryGetPackagePath(const FContentBrowserItemData& InItem, FName& OutPackagePath);
 
 	/**
 	 * Attempt to retrieve the asset data associated with the given item.
@@ -793,7 +795,7 @@ public:
 	 *
 	 * @return True if the asset data was retrieved, false otherwise.
 	 */
-	virtual bool Legacy_TryGetAssetData(const FContentBrowserItemData& InItem, FAssetData& OutAssetData);
+	UE_API virtual bool Legacy_TryGetAssetData(const FContentBrowserItemData& InItem, FAssetData& OutAssetData);
 
 	/**
 	 * Attempt to convert the given package path to a virtual path associated with this data source.
@@ -805,7 +807,7 @@ public:
 	 *
 	 * @return True if the package path was mapped, false otherwise.
 	 */
-	virtual bool Legacy_TryConvertPackagePathToVirtualPath(const FName InPackagePath, FName& OutPath);
+	UE_API virtual bool Legacy_TryConvertPackagePathToVirtualPath(const FName InPackagePath, FName& OutPath);
 
 	/**
 	 * Attempt to convert the given asset data to a virtual path associated with this data source.
@@ -818,17 +820,17 @@ public:
 	 *
 	 * @return True if the asset data was mapped, false otherwise.
 	 */
-	virtual bool Legacy_TryConvertAssetDataToVirtualPath(const FAssetData& InAssetData, const bool InUseFolderPaths, FName& OutPath);
+	UE_API virtual bool Legacy_TryConvertAssetDataToVirtualPath(const FAssetData& InAssetData, const bool InUseFolderPaths, FName& OutPath);
 
 	/**
 	 * Sets a flag to force rebuild of virtual path tree with next call to RefreshVirtualPathTreeIfNeeded()
 	 */
-	void SetVirtualPathTreeNeedsRebuild();
+	UE_API void SetVirtualPathTreeNeedsRebuild();
 
 	/**
 	 * Call after a change that could affect rules of virtual path generation.
 	 */
-	void RefreshVirtualPathTreeIfNeeded();
+	UE_API void RefreshVirtualPathTreeIfNeeded();
 
 	/**
 	 * Attempt to convert the given virtual path
@@ -837,22 +839,22 @@ public:
 	 *
 	 * @return None if virtual path prefix of InPath does not exist, Virtual if path exists and is fully virtual (stops before it reaches internal root), Internal if virtual path part of prefix exists and there is text after the virtual prefix
 	 */
-	EContentBrowserPathType TryConvertVirtualPath(const FStringView InPath, FStringBuilderBase& OutPath) const;
-	EContentBrowserPathType TryConvertVirtualPath(const FStringView InPath, FString& OutPath) const;
-	EContentBrowserPathType TryConvertVirtualPath(const FStringView InPath, FName& OutPath) const;
-	EContentBrowserPathType TryConvertVirtualPath(const FName InPath, FName& OutPath) const;
+	UE_API EContentBrowserPathType TryConvertVirtualPath(const FStringView InPath, FStringBuilderBase& OutPath) const;
+	UE_API EContentBrowserPathType TryConvertVirtualPath(const FStringView InPath, FString& OutPath) const;
+	UE_API EContentBrowserPathType TryConvertVirtualPath(const FStringView InPath, FName& OutPath) const;
+	UE_API EContentBrowserPathType TryConvertVirtualPath(const FName InPath, FName& OutPath) const;
 
 	/**
 	 * Rebuilds the tree of virtual paths that ends with internal roots
 	 */
-	virtual void BuildRootPathVirtualTree();
+	UE_API virtual void BuildRootPathVirtualTree();
 
 	const FContentBrowserVirtualPathTree& GetRootPathVirtualTree() const { return RootPathVirtualTree; }
 
 	/**
 	 * Creates item data for a fully virtual folder.
 	 */
-	FContentBrowserItemData CreateVirtualFolderItem(const FName InFolderPath);
+	UE_API FContentBrowserItemData CreateVirtualFolderItem(const FName InFolderPath);
 
 	/**
 	 * Convert a virtualized path to its internal form, based on the mount root set on this data source.
@@ -863,7 +865,7 @@ public:
 	 *
 	 * @return True if the virtual path was mapped, false otherwise.
 	 */
-	virtual bool TryConvertVirtualPathToInternal(const FName InPath, FName& OutInternalPath);
+	UE_API virtual bool TryConvertVirtualPathToInternal(const FName InPath, FName& OutInternalPath);
 
 	/**
 	 * Convert an internal path to its virtualized form, based on the mount root set on this data source.
@@ -874,17 +876,17 @@ public:
 	 *
 	 * @return True if the internal path was mapped, false otherwise.
 	 */
-	virtual bool TryConvertInternalPathToVirtual(const FName InInternalPath, FName& OutPath);
+	UE_API virtual bool TryConvertInternalPathToVirtual(const FName InInternalPath, FName& OutPath);
 
 	/**
 	 * Tell the data source to remove any cached data for the filter compilation that might not be needed any more.
 	 */
-	virtual void RemoveUnusedCachedFilterData(const FContentBrowserDataFilterCacheIDOwner& IDOwner, TArrayView<const FName> InVirtualPathsInUse, const FContentBrowserDataFilter& DataFilter);
+	UE_API virtual void RemoveUnusedCachedFilterData(const FContentBrowserDataFilterCacheIDOwner& IDOwner, TArrayView<const FName> InVirtualPathsInUse, const FContentBrowserDataFilter& DataFilter);
 
 	/**
 	 * Tell the data source to remove the cached data for the filter compilation for this specific owner. 
 	 */
-	virtual void ClearCachedFilterData(const FContentBrowserDataFilterCacheIDOwner& IDOwner);
+	UE_API virtual void ClearCachedFilterData(const FContentBrowserDataFilterCacheIDOwner& IDOwner);
 
 protected:
 
@@ -894,23 +896,23 @@ protected:
 	 *
 	 * @param InUpdate The update describing how the item has changed.
 	 */
-	void QueueItemDataUpdate(FContentBrowserItemDataUpdate&& InUpdate);
+	UE_API void QueueItemDataUpdate(FContentBrowserItemDataUpdate&& InUpdate);
 
 	/**
 	 * Notify a wholesale item data update, for data sources that can't provide delta-updates.
 	 */
 	 UE_DEPRECATED(5.5, "No data source should be able request a full data refresh. On large projects this can quickly destroy the editor performance. To avoid that all data sources should support queueing the required incremental data updates for their items. Look for uses of the function QueueItemDataUpdate for examples.")
-	void NotifyItemDataRefreshed();
+	UE_API void NotifyItemDataRefreshed();
 
 	/**
 	 * Adds internal root path to virtual path tree
 	 */
-	void RootPathAdded(const FStringView InInternalPath);
+	UE_API void RootPathAdded(const FStringView InInternalPath);
 
 	/**
 	 * Removes internal root path from virtual path tree
 	 */
-	void RootPathRemoved(const FStringView InInternalPath);
+	UE_API void RootPathRemoved(const FStringView InInternalPath);
 
 
 	/**
@@ -946,3 +948,5 @@ private:
 	 */
 	IContentBrowserItemDataSink* DataSink = nullptr;
 };
+
+#undef UE_API
