Fixed forward declaration of FContentBrowserDataFilter having dllstorage

diff --git a/Engine/Source/Editor/ContentBrowser/Public/ContentBrowserTelemetry.h b/Engine/Source/Editor/ContentBrowser/Public/ContentBrowserTelemetry.h
--- a/Engine/Source/Editor/ContentBrowser/Public/ContentBrowserTelemetry.h
+++ b/Engine/Source/Editor/ContentBrowser/Public/ContentBrowserTelemetry.h
@@ -7,7 +7,7 @@
 #include "Misc/Guid.h"
 #include "Misc/Optional.h"
 
-struct CONTENTBROWSERDATA_API FContentBrowserDataFilter;
+struct FContentBrowserDataFilter;
 
 namespace UE::Telemetry::ContentBrowser
 {
