Content Browser now allows modules to customize item attributes. The first use case of this is customization of item folder icons.

Updated ContentBrowserItem's GetItemAttribute to support overrides from other modules, allowing deeper customization of CB items.
Updated ContentBrowserUtils's GetFolder icons to support retrieving the icon brush from an item attribute.
Updated SAssetTreeItem's GetFolderIcon to support retrieving the icon brush from an item attribute.

[RN] Content Browser now allows modules to customize item attributes.

diff --git a/Engine/Plugins/Editor/ContentBrowser/ContentBrowserAssetDataSource/Source/ContentBrowserAssetDataSource/Private/ContentBrowserAssetDataCore.cpp b/Engine/Plugins/Editor/ContentBrowser/ContentBrowserAssetDataSource/Source/ContentBrowserAssetDataSource/Private/ContentBrowserAssetDataCore.cpp
--- a/Engine/Plugins/Editor/ContentBrowser/ContentBrowserAssetDataSource/Source/ContentBrowserAssetDataSource/Private/ContentBrowserAssetDataCore.cpp
+++ b/Engine/Plugins/Editor/ContentBrowser/ContentBrowserAssetDataSource/Source/ContentBrowserAssetDataSource/Private/ContentBrowserAssetDataCore.cpp
@@ -1721,6 +1721,16 @@ void GetGenericItemAttribute(const FName InTagKey, const FString& InTagValue, co
 
 bool GetItemAttribute(const UContentBrowserDataSource* InOwnerDataSource, const FContentBrowserItemData& InItem, const bool InIncludeMetaData, const FName InAttributeKey, FContentBrowserItemDataAttributeValue& OutAttributeValue)
 {
+	// First see if another module hooked this to provide its own customizations. If it did, we can be done here.
+	FContentBrowserModule& ContentBrowserModule = FModuleManager::Get().GetModuleChecked<FContentBrowserModule>(TEXT("ContentBrowser"));
+	for(const FGetItemDataAttribute& Delegate : ContentBrowserModule.GetAllItemDataAttributeDelegates())
+	{
+		if (Delegate.Execute(InItem, InIncludeMetaData, InAttributeKey, OutAttributeValue))
+		{
+			return true;
+		}
+	}
+
 	if (TSharedPtr<const FContentBrowserAssetFolderItemDataPayload> FolderPayload = GetAssetFolderItemPayload(InOwnerDataSource, InItem))
 	{
 		return GetAssetFolderItemAttribute(*FolderPayload, InIncludeMetaData, InAttributeKey, OutAttributeValue);
diff --git a/Engine/Source/Editor/ContentBrowser/Private/ContentBrowserUtils.cpp b/Engine/Source/Editor/ContentBrowser/Private/ContentBrowserUtils.cpp
--- a/Engine/Source/Editor/ContentBrowser/Private/ContentBrowserUtils.cpp
+++ b/Engine/Source/Editor/ContentBrowser/Private/ContentBrowserUtils.cpp
@@ -574,33 +574,43 @@ bool ContentBrowserUtils::TryGetFolderBrushAndShadowName(const FContentBrowserIt
 	}
 
 	OutShadowBrushName = TEXT("ContentBrowser.FolderItem.DropShadow");
-	const bool bDeveloperFolder = IsItemDeveloperContent(InFolder);
-	const bool bCodeFolder = EnumHasAnyFlags(InFolder.GetItemCategory(), EContentBrowserItemFlags::Category_Class);
-	const FContentBrowserItemDataAttributeValue VirtualAttributeValue = InFolder.GetItemAttribute(ContentBrowserItemAttributes::ItemIsCustomVirtualFolder);
-	const bool bVirtualFolder = VirtualAttributeValue.IsValid() && VirtualAttributeValue.GetValue<bool>();
-	const bool bPluginFolder = IsItemPluginRootFolder(InFolder);
 
-	if (bDeveloperFolder)
+	FContentBrowserItemDataAttributeValue FolderBrushAttribute = InFolder.GetItemAttribute(ContentBrowserItemAttributes::ItemListViewFolderBrush);
+	if (FolderBrushAttribute.IsValid())
 	{
-		OutBrushName = TEXT("ContentBrowser.ListViewDeveloperFolderIcon");
-	}
-	else if (bCodeFolder)
-	{
-		OutBrushName = TEXT("ContentBrowser.ListViewCodeFolderIcon");
-	}
-	else if (bVirtualFolder && ShouldShowCustomVirtualFolderIcon())
-	{
-		OutBrushName = TEXT("ContentBrowser.ListViewVirtualFolderIcon");
-		OutShadowBrushName = TEXT("ContentBrowser.ListViewVirtualFolderShadow");
-	}
-	else if (bPluginFolder && ShouldShowPluginFolderIcon())
-	{
-		OutBrushName = TEXT("ContentBrowser.ListViewPluginFolderIcon");
+		OutBrushName = FolderBrushAttribute.GetValueName();
 	}
 	else
 	{
-		OutBrushName = TEXT("ContentBrowser.ListViewFolderIcon");
+		const bool bDeveloperFolder = IsItemDeveloperContent(InFolder);
+		const bool bCodeFolder = EnumHasAnyFlags(InFolder.GetItemCategory(), EContentBrowserItemFlags::Category_Class);
+		const FContentBrowserItemDataAttributeValue VirtualAttributeValue = InFolder.GetItemAttribute(ContentBrowserItemAttributes::ItemIsCustomVirtualFolder);
+		const bool bVirtualFolder = VirtualAttributeValue.IsValid() && VirtualAttributeValue.GetValue<bool>();
+		const bool bPluginFolder = IsItemPluginRootFolder(InFolder);
+
+		if (bDeveloperFolder)
+		{
+			OutBrushName = TEXT("ContentBrowser.ListViewDeveloperFolderIcon");
+		}
+		else if (bCodeFolder)
+		{
+			OutBrushName = TEXT("ContentBrowser.ListViewCodeFolderIcon");
+		}
+		else if (bVirtualFolder && ShouldShowCustomVirtualFolderIcon())
+		{
+			OutBrushName = TEXT("ContentBrowser.ListViewVirtualFolderIcon");
+			OutShadowBrushName = TEXT("ContentBrowser.ListViewVirtualFolderShadow");
+		}
+		else if (bPluginFolder && ShouldShowPluginFolderIcon())
+		{
+			OutBrushName = TEXT("ContentBrowser.ListViewPluginFolderIcon");
+		}
+		else
+		{
+			OutBrushName = TEXT("ContentBrowser.ListViewFolderIcon");
+		}
 	}
+
 	return true;
 }
 
@@ -612,33 +622,43 @@ bool ContentBrowserUtils::TryGetFolderBrushAndShadowNameSmall(const FContentBrow
 	}
 
 	OutShadowBrushName = TEXT("ContentBrowser.FolderItem.DropShadow");
-	const bool bDeveloperFolder = IsItemDeveloperContent(InFolder);
-	const bool bCodeFolder = EnumHasAnyFlags(InFolder.GetItemCategory(), EContentBrowserItemFlags::Category_Class);
-	const FContentBrowserItemDataAttributeValue VirtualAttributeValue = InFolder.GetItemAttribute(ContentBrowserItemAttributes::ItemIsCustomVirtualFolder);
-	const bool bVirtualFolder = VirtualAttributeValue.IsValid() && VirtualAttributeValue.GetValue<bool>();
-	const bool bPluginFolder = IsItemPluginRootFolder(InFolder);
 
-	if (bDeveloperFolder)
+	FContentBrowserItemDataAttributeValue FolderBrushAttribute = InFolder.GetItemAttribute(ContentBrowserItemAttributes::ItemAssetTreeFolderClosedBrush);
+	if (FolderBrushAttribute.IsValid())
 	{
-		OutBrushName = TEXT("ContentBrowser.AssetTreeFolderClosedDeveloper");
-	}
-	else if (bCodeFolder)
-	{
-		OutBrushName = TEXT("ContentBrowser.AssetTreeFolderClosedCode");
-	}
-	else if (bVirtualFolder && ShouldShowCustomVirtualFolderIcon())
-	{
-		OutBrushName = TEXT("ContentBrowser.AssetTreeFolderClosedVirtual");
-		OutShadowBrushName = TEXT("ContentBrowser.ListViewVirtualFolderShadow");
-	}
-	else if (bPluginFolder && ShouldShowPluginFolderIcon())
-	{
-		OutBrushName = TEXT("ContentBrowser.AssetTreeFolderClosedPluginRoot");
+		OutBrushName = FolderBrushAttribute.GetValueName();
 	}
 	else
 	{
-		OutBrushName = TEXT("ContentBrowser.AssetTreeFolderClosed");
+		const bool bDeveloperFolder = IsItemDeveloperContent(InFolder);
+		const bool bCodeFolder = EnumHasAnyFlags(InFolder.GetItemCategory(), EContentBrowserItemFlags::Category_Class);
+		const FContentBrowserItemDataAttributeValue VirtualAttributeValue = InFolder.GetItemAttribute(ContentBrowserItemAttributes::ItemIsCustomVirtualFolder);
+		const bool bVirtualFolder = VirtualAttributeValue.IsValid() && VirtualAttributeValue.GetValue<bool>();
+		const bool bPluginFolder = IsItemPluginRootFolder(InFolder);
+
+		if (bDeveloperFolder)
+		{
+			OutBrushName = TEXT("ContentBrowser.AssetTreeFolderClosedDeveloper");
+		}
+		else if (bCodeFolder)
+		{
+			OutBrushName = TEXT("ContentBrowser.AssetTreeFolderClosedCode");
+		}
+		else if (bVirtualFolder && ShouldShowCustomVirtualFolderIcon())
+		{
+			OutBrushName = TEXT("ContentBrowser.AssetTreeFolderClosedVirtual");
+			OutShadowBrushName = TEXT("ContentBrowser.ListViewVirtualFolderShadow");
+		}
+		else if (bPluginFolder && ShouldShowPluginFolderIcon())
+		{
+			OutBrushName = TEXT("ContentBrowser.AssetTreeFolderClosedPluginRoot");
+		}
+		else
+		{
+			OutBrushName = TEXT("ContentBrowser.AssetTreeFolderClosed");
+		}
 	}
+
 	return true;
 }
 
diff --git a/Engine/Source/Editor/ContentBrowser/Private/SourcesViewWidgets.cpp b/Engine/Source/Editor/ContentBrowser/Private/SourcesViewWidgets.cpp
--- a/Engine/Source/Editor/ContentBrowser/Private/SourcesViewWidgets.cpp
+++ b/Engine/Source/Editor/ContentBrowser/Private/SourcesViewWidgets.cpp
@@ -317,6 +317,20 @@ bool SAssetTreeItem::IsReadOnly() const
 
 const FSlateBrush* SAssetTreeItem::GetFolderIcon() const
 {
+	// First see if this item has a custom implementation
+	if (TSharedPtr<FTreeItem> TreeItemPin = TreeItem.Pin())
+	{
+		FContentBrowserItemDataAttributeValue FolderAttributeValue = TreeItemPin->GetItem().GetItemAttribute(IsItemExpanded.Get() ? ContentBrowserItemAttributes::ItemAssetTreeFolderOpenBrush : ContentBrowserItemAttributes::ItemAssetTreeFolderClosedBrush);
+		if (FolderAttributeValue.IsValid())
+		{
+			FName BrushName = FolderAttributeValue.GetValue<FName>();
+
+			const FSlateStyleSet& ContentBrowserStyle = UE::ContentBrowser::Private::FContentBrowserStyle::Get();
+			return ContentBrowserStyle.GetBrush(BrushName);
+		}
+	}
+
+	// If we made it here, nothing overrode it, so use the default brushes.
 	FAssetTreeItemBrushes& Brushes = FAssetTreeItemBrushes::Get();
 	switch( FolderType )
 	{
diff --git a/Engine/Source/Editor/ContentBrowser/Public/ContentBrowserDelegates.h b/Engine/Source/Editor/ContentBrowser/Public/ContentBrowserDelegates.h
--- a/Engine/Source/Editor/ContentBrowser/Public/ContentBrowserDelegates.h
+++ b/Engine/Source/Editor/ContentBrowser/Public/ContentBrowserDelegates.h
@@ -10,6 +10,8 @@
 #include "Containers/ArrayView.h"
 #include "ContentBrowserDataLegacyBridge.h"
 
+class FContentBrowserItemData;
+class FContentBrowserItemDataAttributeValue;
 struct FARFilter;
 struct FCollectionNameType;
 struct FContentBrowserItem;
@@ -116,6 +118,9 @@ DECLARE_DELEGATE_RetVal_OneParam(TSharedRef<SWidget>, FOnGenerateAssetViewExtraS
 /** Called to generate state information icons for path view items. */
 DECLARE_DELEGATE_RetVal_OneParam(TSharedPtr<SWidget>, FPathViewStateIconGenerator, const FContentBrowserItem& /*ContentBrowserItem*/);
 
+/** Called to get ContentBrowserItemData attribute values. */
+DECLARE_DELEGATE_RetVal_FourParams(bool, FGetItemDataAttribute, const FContentBrowserItemData& /*InItem*/, const bool /*InIncludeMetaData*/, const FName /*InAttributeKey*/, FContentBrowserItemDataAttributeValue& /*OutAttributeValue*/);
+
 /** Called when an asset item visualizes its tooltip */
 DECLARE_DELEGATE_RetVal_TwoParams(bool, FOnVisualizeAssetToolTip, const TSharedPtr<SWidget>& /*ToolTipContent*/, FAssetData& /*AssetData*/);
 
diff --git a/Engine/Source/Editor/ContentBrowser/Public/ContentBrowserModule.h b/Engine/Source/Editor/ContentBrowser/Public/ContentBrowserModule.h
--- a/Engine/Source/Editor/ContentBrowser/Public/ContentBrowserModule.h
+++ b/Engine/Source/Editor/ContentBrowser/Public/ContentBrowserModule.h
@@ -175,6 +175,9 @@ public:
 	/** Delegates to be called to extend list of content browser Plugin Filters*/
 	virtual TArray<FAddPathViewPluginFilters>& GetAddPathViewPluginFilters() { return PathViewPluginFilters; }
 
+	/** Delegates to call to provide the value for a ContentBrowserItemData's attribute */
+	virtual TArray<FGetItemDataAttribute>& GetAllItemDataAttributeDelegates() { return ItemDataAttributeDelegates; }
+
 	/** Register a delegate that creates a custom view that can show up inside SAssetView in the Content Browser */
 	virtual void SetContentBrowserViewExtender(const FCreateViewExtender& InCreateViewExtender);
 	/** Get the currently registered custom view */
@@ -254,6 +257,9 @@ private:
 	/** All delegates that extend available path view plugin filters */
 	TArray<FAddPathViewPluginFilters> PathViewPluginFilters;
 
+	/** All delegates that return values for ContentItemData attributes */
+	TArray<FGetItemDataAttribute> ItemDataAttributeDelegates;
+
 	FOnFilterChanged OnFilterChanged;
 	FOnSearchBoxChanged OnSearchBoxChanged;
 	FOnAssetSelectionChanged OnAssetSelectionChanged;
diff --git a/Engine/Source/Editor/ContentBrowserData/Public/ContentBrowserDataSource.h b/Engine/Source/Editor/ContentBrowserData/Public/ContentBrowserDataSource.h
--- a/Engine/Source/Editor/ContentBrowserData/Public/ContentBrowserDataSource.h
+++ b/Engine/Source/Editor/ContentBrowserData/Public/ContentBrowserDataSource.h
@@ -105,6 +105,24 @@ namespace ContentBrowserItemAttributes
 	 * Type: FLinearColor (as FString).
 	 */
 	const FName ItemColor = "ItemColor";
+
+	/**
+	 * Attribute key that can be used to query the item's brush to use for an open folder in an asset tree.
+	 * Type: Brush name (as FName).
+	 */
+	const FName ItemAssetTreeFolderOpenBrush = "ItemAssetTreeFolderOpenBrush";
+
+	/**
+	 * Attribute key that can be used to query the item's brush to use for a closed folder in an asset tree.
+	 * Type: Brush name (as FName).
+	 */
+	const FName ItemAssetTreeFolderClosedBrush = "ItemAssetTreeFolderClosedBrush";
+
+	/**
+	 * Attribute key that can be used to query the item's brush to use for a folder in a list view.
+	 * Type: Brush name (as FName).
+	 */
+	const FName ItemListViewFolderBrush = "ItemListViewFolderBrush";
 }
 
 /**
