[ContentBrowser] Set the tiny size of the List and Column view back to 16 in height

diff --git a/Engine/Source/Editor/ContentBrowser/Private/SAssetView.cpp b/Engine/Source/Editor/ContentBrowser/Private/SAssetView.cpp
--- a/Engine/Source/Editor/ContentBrowser/Private/SAssetView.cpp
+++ b/Engine/Source/Editor/ContentBrowser/Private/SAssetView.cpp
@@ -5903,7 +5903,7 @@ void SAssetView::UpdateThumbnailSizeValue()
 	case EThumbnailSize::Tiny:
 		MinThumbnailSize = 64.f;
 		MaxThumbnailSize = 80.f;
-		ListViewItemHeight = 22.f;
+		ListViewItemHeight = 16.f;
 		break;
 	case EThumbnailSize::Small:
 		MinThumbnailSize = 80.f;
