Centralized Asset Access Specifier iconography
Updated Asset Context Menu and Tooltips to render the icon for the given asset access specifier state.

#rnx

diff --git a/Engine/Source/Editor/ContentBrowser/Private/AssetContextMenu.cpp b/Engine/Source/Editor/ContentBrowser/Private/AssetContextMenu.cpp
--- a/Engine/Source/Editor/ContentBrowser/Private/AssetContextMenu.cpp
+++ b/Engine/Source/Editor/ContentBrowser/Private/AssetContextMenu.cpp
@@ -656,7 +656,7 @@ bool FAssetContextMenu::AddPublicStateMenuOptions(UToolMenu* Menu)
 					"PublicAsset",
 					LOCTEXT("PublicAssetToggle", "Public Asset"),
 					LOCTEXT("PublicAssetToggleTooltip", "Sets the asset to be referencable by other Plugins"),
-					FSlateIcon(),
+					FSlateIcon(FAppStyle::GetAppStyleSetName(), FAssetViewItem::GetPublicAssetAccessSpecifierIconStyleName(TEXT(".Small"))),
 					FUIAction(
 						FExecuteAction::CreateSP(this, &FAssetContextMenu::ExecuteSetAssetAccessSpecifier, EAssetAccessSpecifier::Public),
 						FCanExecuteAction::CreateSP(this, &FAssetContextMenu::CanSetAssetAccessSpecifier, EAssetAccessSpecifier::Public),
@@ -671,7 +671,7 @@ bool FAssetContextMenu::AddPublicStateMenuOptions(UToolMenu* Menu)
 						"EpicInternalAsset",
 						LOCTEXT("SetInternalAsset", "Epic Internal Asset"),
 						LOCTEXT("SetInternalAssetTooltip", "Sets the asset to be referencable by Epic internal plugins and mount points"),
-						FSlateIcon(),
+						FSlateIcon(FAppStyle::GetAppStyleSetName(), FAssetViewItem::GetEpicInternalAssetAccessSpecifierIconStyleName(TEXT(".Small"))),
 						FUIAction(
 							FExecuteAction::CreateSP(this, &FAssetContextMenu::ExecuteSetAssetAccessSpecifier, EAssetAccessSpecifier::EpicInternal),
 							FCanExecuteAction::CreateSP(this, &FAssetContextMenu::CanSetAssetAccessSpecifier, EAssetAccessSpecifier::EpicInternal),
@@ -685,7 +685,7 @@ bool FAssetContextMenu::AddPublicStateMenuOptions(UToolMenu* Menu)
 					"PrivateAsset",
 					LOCTEXT("SetPrivateAsset", "Private Asset"),
 					LOCTEXT("SetAssetPrivateTooltip", "Sets the asset so it can't be referenced by other Plugins"),
-					FSlateIcon(),
+					FSlateIcon(FAppStyle::GetAppStyleSetName(), FAssetViewItem::GetPrivateAssetAccessSpecifierIconStyleName(TEXT(".Small"))),
 					FUIAction(
 						FExecuteAction::CreateSP(this, &FAssetContextMenu::ExecuteSetAssetAccessSpecifier, EAssetAccessSpecifier::Private),
 						FCanExecuteAction::CreateSP(this, &FAssetContextMenu::CanSetAssetAccessSpecifier, EAssetAccessSpecifier::Private),
diff --git a/Engine/Source/Editor/ContentBrowser/Private/AssetSystemContentBrowserInfoProvider.cpp b/Engine/Source/Editor/ContentBrowser/Private/AssetSystemContentBrowserInfoProvider.cpp
--- a/Engine/Source/Editor/ContentBrowser/Private/AssetSystemContentBrowserInfoProvider.cpp
+++ b/Engine/Source/Editor/ContentBrowser/Private/AssetSystemContentBrowserInfoProvider.cpp
@@ -120,7 +120,7 @@ void FAssetSystemContentBrowserInfoProvider::PopulateAssetInfo(TArray<FAssetDisp
 			if (!PublicStateText.IsEmpty())
 			{
 				// An icon is required for the PublicStateText to display. So, set one if there's a value to show.
-				PublicStateIcon = FAppStyle::GetBrush("Icons.Info");
+				PublicStateIcon = FAppStyle::GetBrush(AssetItem->GetItemAssetAccessSpecifierIconStyleName(TEXT(".Small")));
 			}
 		}
 
diff --git a/Engine/Source/Editor/ContentBrowser/Private/AssetViewTypes.cpp b/Engine/Source/Editor/ContentBrowser/Private/AssetViewTypes.cpp
--- a/Engine/Source/Editor/ContentBrowser/Private/AssetViewTypes.cpp
+++ b/Engine/Source/Editor/ContentBrowser/Private/AssetViewTypes.cpp
@@ -241,6 +241,61 @@ FText FAssetViewItem::GetItemAssetAccessSpecifierText() const
 	return LOCTEXT("UnknownAssetState", "Unknown");
 }
 
+FName FAssetViewItem::GetItemAssetAccessSpecifierIconStyleName(const TCHAR* InSizeSuffix) const
+{
+	// Items that don't have a modifable access specifier shouldn't render anything.
+	FName IconName = NAME_None;
+
+	EAssetAccessSpecifier AssetAccessSpecifier = EAssetAccessSpecifier::Private;
+	if (CanModifyItemAssetAccessSpecifier() && GetItemAssetAccessSpecifier(AssetAccessSpecifier) )
+	{
+		switch (AssetAccessSpecifier)
+		{
+			case EAssetAccessSpecifier::Public:
+			{
+				IconName = GetPublicAssetAccessSpecifierIconStyleName(InSizeSuffix);
+				break;
+			}
+
+			case EAssetAccessSpecifier::Private:
+			{
+				IconName = GetPrivateAssetAccessSpecifierIconStyleName(InSizeSuffix);
+				break;
+			}
+
+			case EAssetAccessSpecifier::EpicInternal:
+			{
+				IconName = GetEpicInternalAssetAccessSpecifierIconStyleName(InSizeSuffix);
+				break;
+			}
+
+			default:
+			{
+				ensureMsgf(0, TEXT("Unrecognized EAssetAccessSpecifier value '%d' found for asset: '%s'"), AssetAccessSpecifier, *GetItem().GetItemName().ToString());
+				break;
+			}
+		}
+	}
+
+	return IconName;
+}
+
+FName FAssetViewItem::GetPublicAssetAccessSpecifierIconStyleName(const TCHAR* InSizeSuffix)
+{
+	return FName(*FString::Printf(TEXT("ContentBrowser.Shared%s"), InSizeSuffix));
+}
+
+FName FAssetViewItem::GetPrivateAssetAccessSpecifierIconStyleName(const TCHAR* InSizeSuffix)
+{
+	return FName(*FString::Printf(TEXT("ContentBrowser.Private%s"), InSizeSuffix));
+}
+
+FName FAssetViewItem::GetEpicInternalAssetAccessSpecifierIconStyleName(const TCHAR* InSizeSuffix)
+{
+	static const FName IconName("UnrealCircle.Thin");
+	return IconName;
+}
+
 bool FAssetViewItem::CanModifyItemAssetAccessSpecifier() const
 {
 	return Item.CanModifyItemAssetAccessSpecifier();
diff --git a/Engine/Source/Editor/ContentBrowser/Private/AssetViewWidgets.cpp b/Engine/Source/Editor/ContentBrowser/Private/AssetViewWidgets.cpp
--- a/Engine/Source/Editor/ContentBrowser/Private/AssetViewWidgets.cpp
+++ b/Engine/Source/Editor/ContentBrowser/Private/AssetViewWidgets.cpp
@@ -1887,38 +1887,10 @@ void SAssetViewItem::UpdateAssetAccessSpecifierIcon()
 	TSharedRef<SWidget> IconWidget = SNullWidget::NullWidget;
 
 	// If there's a valid, modifiable EAssetAccessSpecifier, render its state. If there isn't one, or it isn't modifiable, it's pointless to show it.
-	EAssetAccessSpecifier AssetAccessSpecifier;
-	if (AssetItem.IsValid() && AssetItem->CanModifyItemAssetAccessSpecifier() && AssetItem->GetItemAssetAccessSpecifier(AssetAccessSpecifier))
+	EAssetAccessSpecifier AssetAccessSpecifier = EAssetAccessSpecifier::Private;
+	if (AssetItem.IsValid())
 	{
-		FName IconName;
-		
-		switch (AssetAccessSpecifier)
-		{
-			case EAssetAccessSpecifier::Public:
-			{
-				IconName = TEXT("ContentBrowser.Shared");
-				break;
-			}
-
-			case EAssetAccessSpecifier::Private:
-			{
-				IconName = TEXT("ContentBrowser.Private");
-				break;
-			}
-
-			case EAssetAccessSpecifier::EpicInternal:
-			{
-				IconName = TEXT("ContentBrowser.Local");
-				break;
-			}
-
-			default:
-			{
-				ensureMsgf(0, TEXT("Unrecognized EAssetAccessSpecifier value '%d' found for asset: '%s'"), AssetAccessSpecifier , *AssetItem->GetItem().GetItemName().ToString());
-				break;
-			}
-		}
-
+		FName IconName = AssetItem->GetItemAssetAccessSpecifierIconStyleName();
 		if (IconName != NAME_None)
 		{
 			// Wrap the icon in an outer background layer to improve visibility on light thumbnails
diff --git a/Engine/Source/Editor/ContentBrowser/Public/AssetViewTypes.h b/Engine/Source/Editor/ContentBrowser/Public/AssetViewTypes.h
--- a/Engine/Source/Editor/ContentBrowser/Public/AssetViewTypes.h
+++ b/Engine/Source/Editor/ContentBrowser/Public/AssetViewTypes.h
@@ -89,6 +89,13 @@ public:
 	 */
 	FText GetItemAssetAccessSpecifierText() const;
 
+	/**
+	 * Query the icon brush style name to use for the EAssetAccessSpecifier of the item.
+
+	 * @return If the item has a modifiable EAssetAccessSpecifier, then the style name of the brush to use. Otherwise NAME_None.
+	 */
+	FName GetItemAssetAccessSpecifierIconStyleName(const TCHAR* InSizeSuffix = TEXT("")) const;
+
 	/**
 	 * Query the item's ability to modify its EAssetAccessSpecifier.
 	 *
@@ -96,6 +103,15 @@ public:
 	 */
 	bool CanModifyItemAssetAccessSpecifier() const;
 
+	 /** Returns the brush style name to use for a public asset access specifier */
+	static FName GetPublicAssetAccessSpecifierIconStyleName(const TCHAR* InSizeSuffix = TEXT(""));
+
+	/** Returns the brush style name to use for a private asset access specifier */
+	static FName GetPrivateAssetAccessSpecifierIconStyleName(const TCHAR* InSizeSuffix = TEXT(""));
+
+	/** Returns the brush style name to use for an Epic Internal asset access specifier */
+	static FName GetEpicInternalAssetAccessSpecifierIconStyleName(const TCHAR* InSizeSuffix = TEXT(""));
+
 	/** Get the underlying Content Browser item */
 	CONTENTBROWSER_API const FContentBrowserItem& GetItem() const;
 
