Fix bulk set assets to 'Epic Internal' does not set all assets to 'Epic Internal' if any asset selected is already 'Epic Internal'
#rnx

diff --git a/Engine/Source/Editor/ContentBrowser/Private/AssetContextMenu.cpp b/Engine/Source/Editor/ContentBrowser/Private/AssetContextMenu.cpp
--- a/Engine/Source/Editor/ContentBrowser/Private/AssetContextMenu.cpp
+++ b/Engine/Source/Editor/ContentBrowser/Private/AssetContextMenu.cpp
@@ -1420,7 +1420,7 @@ void FAssetContextMenu::ExecuteBulkSetAssetAccessSpecifier(EAssetAccessSpecifier
 			const EAssetAccessSpecifier OldAssetAccessSpecifier = ItemAssetPackage->GetAssetAccessSpecifier();
 			if (OldAssetAccessSpecifier == DestScope)
 			{
-				break;
+				continue;
 			}
 
 			if (OldAssetAccessSpecifier == EAssetAccessSpecifier::Private)
