Fixed the Content Browser navigation bar bread crumbs missing some folders that appear in the path view.

diff --git a/Engine/Source/Editor/ContentBrowser/Private/SPathView.cpp b/Engine/Source/Editor/ContentBrowser/Private/SPathView.cpp
--- a/Engine/Source/Editor/ContentBrowser/Private/SPathView.cpp
+++ b/Engine/Source/Editor/ContentBrowser/Private/SPathView.cpp
@@ -1701,7 +1701,7 @@ bool SPathView::InternalPathPassesBlockLists(const FStringView InInternalPath, c
 
 	for (const FPathPermissionList* Filter : BlockLists)
 	{
-		if (!Filter->PassesStartsWithFilter(InInternalPath))
+		if (!Filter->PassesStartsWithFilter(InInternalPath, /*bAllowParentPaths*/true))
 		{
 			return false;
 		}
