Fixed the content browser not navigating to various user provided paths from the navigation bar.  Namely pure virtual paths when there is a permission filter, class paths, and aliased paths when there is no original item (e.g. Verse devices).
Fixed missing content browser navigation bar path suggestions for the root path (/).
Added content browser navigation bar path suggestions for well formed paths that the content browser now fully supports navigating to such as package paths and Verse paths.

diff --git a/Engine/Plugins/Editor/ContentBrowser/ContentBrowserAliasDataSource/Source/ContentBrowserAliasDataSource/Private/ContentBrowserAliasDataSource.cpp b/Engine/Plugins/Editor/ContentBrowser/ContentBrowserAliasDataSource/Source/ContentBrowserAliasDataSource/Private/ContentBrowserAliasDataSource.cpp
--- a/Engine/Plugins/Editor/ContentBrowser/ContentBrowserAliasDataSource/Source/ContentBrowserAliasDataSource/Private/ContentBrowserAliasDataSource.cpp
+++ b/Engine/Plugins/Editor/ContentBrowser/ContentBrowserAliasDataSource/Source/ContentBrowserAliasDataSource/Private/ContentBrowserAliasDataSource.cpp
@@ -274,6 +274,72 @@ void UContentBrowserAliasDataSource::EnumerateItemsAtPath(const FName InPath, co
 	}
 }
 
+void UContentBrowserAliasDataSource::EnumerateItemsAtUserProvidedPath(const FStringView InPath, const EContentBrowserItemTypeFilter InItemTypeFilter, TFunctionRef<bool(FContentBrowserItemData&&)> InCallback)
+{
+	if (EnumHasAnyFlags(InItemTypeFilter, EContentBrowserItemTypeFilter::IncludeFiles) && !AliasesForObjectPath.IsEmpty())
+	{
+		auto EnumerateItemsAtObjectPath = [this, &InCallback](const FSoftObjectPath& ObjectPath)
+		{
+			if (ObjectPath.IsValid())
+			{
+				if (const TArray<FName>* FoundAliases = AliasesForObjectPath.Find(ObjectPath))
+				{
+					for (const FName Alias : *FoundAliases)
+					{
+						FContentBrowserItemData Item = CreateAssetFileItem(FContentBrowserUniqueAlias(ObjectPath, Alias));
+						if (ensureMsgf(Item.IsValid(), TEXT("Could not find alias data for: %s"), *Alias.ToString()))
+						{
+							InCallback(MoveTemp(Item));
+						}
+					}
+					return true;
+				}
+			}
+			return false;
+		};
+
+		// Could be an object path.
+		{
+			FSoftObjectPath ObjectPath(FPackageName::ExportTextPathToObjectPath(InPath));
+
+			if (EnumerateItemsAtObjectPath(ObjectPath))
+			{
+				return;
+			}
+		}
+
+		// Could be a package name or file name.
+		if (ensure(AssetRegistry))
+		{
+			FName PackageName;
+			if (FPackageName::IsValidLongPackageName(InPath))
+			{
+				PackageName = FName(InPath, FNAME_Find);
+			}
+			else
+			{
+				FNameBuilder PackageNameBuilder;
+				if (FPackageName::TryConvertFilenameToLongPackageName(InPath, PackageNameBuilder))
+				{
+					PackageName = FName(PackageNameBuilder, FNAME_Find);
+				}
+			}
+
+			if (!PackageName.IsNone())
+			{
+				TArray<FAssetData> Assets;
+				if (AssetRegistry->GetAssetsByPackageName(PackageName, Assets))
+				{
+					for (const FAssetData& AssetData : Assets)
+					{
+						EnumerateItemsAtObjectPath(AssetData.GetSoftObjectPath());
+					}
+				}
+			}
+		}
+	}
+}
+
 bool UContentBrowserAliasDataSource::IsFolderVisible(const FName Path, const EContentBrowserIsFolderVisibleFlags Flags, const FContentBrowserFolderContentsFilter& ContentsFilter)
 {
 	FName InternalPath;
diff --git a/Engine/Plugins/Editor/ContentBrowser/ContentBrowserAliasDataSource/Source/ContentBrowserAliasDataSource/Public/ContentBrowserAliasDataSource.h b/Engine/Plugins/Editor/ContentBrowser/ContentBrowserAliasDataSource/Source/ContentBrowserAliasDataSource/Public/ContentBrowserAliasDataSource.h
--- a/Engine/Plugins/Editor/ContentBrowser/ContentBrowserAliasDataSource/Source/ContentBrowserAliasDataSource/Public/ContentBrowserAliasDataSource.h
+++ b/Engine/Plugins/Editor/ContentBrowser/ContentBrowserAliasDataSource/Source/ContentBrowserAliasDataSource/Public/ContentBrowserAliasDataSource.h
@@ -60,6 +60,7 @@ public:
 	UE_API virtual void CompileFilter(const FName InPath, const FContentBrowserDataFilter& InFilter, FContentBrowserDataCompiledFilter& OutCompiledFilter) override;
 	UE_API virtual void EnumerateItemsMatchingFilter(const FContentBrowserDataCompiledFilter& InFilter, TFunctionRef<bool(FContentBrowserItemData&&)> InCallback) override;
 	UE_API virtual void EnumerateItemsAtPath(const FName InPath, const EContentBrowserItemTypeFilter InItemTypeFilter, TFunctionRef<bool(FContentBrowserItemData&&)> InCallback) override;
+	UE_API virtual void EnumerateItemsAtUserProvidedPath(const FStringView InPath, const EContentBrowserItemTypeFilter InItemTypeFilter, TFunctionRef<bool(FContentBrowserItemData&&)> InCallback) override;
 
 	UE_API virtual bool IsFolderVisible(const FName Path, const EContentBrowserIsFolderVisibleFlags Flags, const FContentBrowserFolderContentsFilter& ContentsFilter) override;
 
diff --git a/Engine/Plugins/Editor/ContentBrowser/ContentBrowserAssetDataSource/Source/ContentBrowserAssetDataSource/Private/ContentBrowserAssetDataSource.cpp b/Engine/Plugins/Editor/ContentBrowser/ContentBrowserAssetDataSource/Source/ContentBrowserAssetDataSource/Private/ContentBrowserAssetDataSource.cpp
--- a/Engine/Plugins/Editor/ContentBrowser/ContentBrowserAssetDataSource/Source/ContentBrowserAssetDataSource/Private/ContentBrowserAssetDataSource.cpp
+++ b/Engine/Plugins/Editor/ContentBrowser/ContentBrowserAssetDataSource/Source/ContentBrowserAssetDataSource/Private/ContentBrowserAssetDataSource.cpp
@@ -2302,6 +2302,85 @@ bool UContentBrowserAssetDataSource::EnumerateItemsAtPaths(const TArrayView<FCon
 	return true;
 }
 
+void UContentBrowserAssetDataSource::EnumerateItemsAtUserProvidedPath(const FStringView InPath, const EContentBrowserItemTypeFilter InItemTypeFilter, TFunctionRef<bool(FContentBrowserItemData&&)> InCallback)
+{
+	if (EnumHasAnyFlags(InItemTypeFilter, EContentBrowserItemTypeFilter::IncludeFolders))
+	{
+		// Could be a package path.
+		if (ensure(AssetRegistry))
+		{
+			FStringView Path = InPath;
+
+			// Internal paths have no trailing /.
+			while (FPathViews::HasRedundantTerminatingSeparator(Path))
+			{
+				Path.LeftChopInline(1);
+			}
+
+			FName InternalFolderPath(Path, FNAME_Find);
+			if (!InternalFolderPath.IsNone() && AssetRegistry->PathExists(InternalFolderPath))
+			{
+				InCallback(CreateAssetFolderItem(InternalFolderPath));
+			}
+		}
+	}
+
+	if (EnumHasAnyFlags(InItemTypeFilter, EContentBrowserItemTypeFilter::IncludeFiles))
+	{
+		// Could be an object path.
+		if (ensure(AssetRegistry))
+		{
+			FSoftObjectPath ObjectPath(FPackageName::ExportTextPathToObjectPath(InPath));
+			if (ObjectPath.IsValid())
+			{
+				FAssetData AssetData = AssetRegistry->GetAssetByObjectPath(ObjectPath);
+				if (AssetData.IsValid())
+				{
+					if (ContentBrowserAssetData::IsPrimaryAsset(AssetData))
+					{
+						InCallback(CreateAssetFileItem(AssetData));
+					}
+
+					return;
+				}
+			}
+		}
+
+		// Could be a package name or file name.
+		if (ensure(AssetRegistry))
+		{
+			FName PackageName;
+			if (FPackageName::IsValidLongPackageName(InPath))
+			{
+				PackageName = FName(InPath, FNAME_Find);
+			}
+			else
+			{
+				FNameBuilder PackageNameBuilder;
+				if (FPackageName::TryConvertFilenameToLongPackageName(InPath, PackageNameBuilder))
+				{
+					PackageName = FName(PackageNameBuilder, FNAME_Find);
+				}
+			}
+
+			if (!PackageName.IsNone())
+			{
+				TArray<FAssetData> Assets;
+				if (AssetRegistry->GetAssetsByPackageName(PackageName, Assets))
+				{
+					for (const FAssetData& AssetData : Assets)
+					{
+						if (ContentBrowserAssetData::IsPrimaryAsset(AssetData))
+						{
+							InCallback(CreateAssetFileItem(AssetData));
+						}
+					}
+				}
+			}
+		}
+	}
+}
+
 bool UContentBrowserAssetDataSource::IsDiscoveringItems(FText* OutStatus)
 {
 	if (AssetRegistry->IsLoadingAssets())
diff --git a/Engine/Plugins/Editor/ContentBrowser/ContentBrowserAssetDataSource/Source/ContentBrowserAssetDataSource/Public/ContentBrowserAssetDataSource.h b/Engine/Plugins/Editor/ContentBrowser/ContentBrowserAssetDataSource/Source/ContentBrowserAssetDataSource/Public/ContentBrowserAssetDataSource.h
--- a/Engine/Plugins/Editor/ContentBrowser/ContentBrowserAssetDataSource/Source/ContentBrowserAssetDataSource/Public/ContentBrowserAssetDataSource.h
+++ b/Engine/Plugins/Editor/ContentBrowser/ContentBrowserAssetDataSource/Source/ContentBrowserAssetDataSource/Public/ContentBrowserAssetDataSource.h
@@ -249,6 +249,8 @@ public:
 
 	UE_API virtual bool EnumerateItemsAtPaths(const TArrayView<FContentBrowserItemPath> InPaths, const EContentBrowserItemTypeFilter InItemTypeFilter, TFunctionRef<bool(FContentBrowserItemData&&)> InCallback) override;
 
+	UE_API virtual void EnumerateItemsAtUserProvidedPath(const FStringView InPath, const EContentBrowserItemTypeFilter InItemTypeFilter, TFunctionRef<bool(FContentBrowserItemData&&)> InCallback) override;
+
 	UE_API virtual bool IsDiscoveringItems(FText* OutStatus = nullptr) override;
 
 	UE_API virtual bool PrioritizeSearchPath(const FName InPath) override;
diff --git a/Engine/Plugins/Editor/ContentBrowser/ContentBrowserClassDataSource/Source/ContentBrowserClassDataSource/Private/ContentBrowserClassDataSource.cpp b/Engine/Plugins/Editor/ContentBrowser/ContentBrowserClassDataSource/Source/ContentBrowserClassDataSource/Private/ContentBrowserClassDataSource.cpp
--- a/Engine/Plugins/Editor/ContentBrowser/ContentBrowserClassDataSource/Source/ContentBrowserClassDataSource/Private/ContentBrowserClassDataSource.cpp
+++ b/Engine/Plugins/Editor/ContentBrowser/ContentBrowserClassDataSource/Source/ContentBrowserClassDataSource/Private/ContentBrowserClassDataSource.cpp
@@ -379,8 +379,30 @@ void UContentBrowserClassDataSource::EnumerateItemsAtPath(const FName InPath, co
 	{
 		if (TSharedPtr<const FNativeClassHierarchyNode> ClassNode = NativeClassHierarchy->FindNode(InternalPath, ENativeClassHierarchyNodeType::Class))
 		{
-			FNativeClassHierarchyGetClassPathCache Cache;
-			InCallback(CreateClassFileItem(ClassNode->Class, Cache));
+			InCallback(CreateClassFileItem(InternalPath, ClassNode));
+		}
+	}
+}
+
+void UContentBrowserClassDataSource::EnumerateItemsAtUserProvidedPath(const FStringView InPath, const EContentBrowserItemTypeFilter InItemTypeFilter, TFunctionRef<bool(FContentBrowserItemData&&)> InCallback)
+{
+	if (EnumHasAnyFlags(InItemTypeFilter, EContentBrowserItemTypeFilter::IncludeFiles) && InPath.StartsWith(TEXT("/Script/")))
+	{
+		const FTopLevelAssetPath Path(InPath);
+		if (Path.IsValid() && !Path.GetAssetName().IsNone())
+		{
+			if (UClass* Class = FindObject<UClass>(Path))
+			{
+				FString InternalPath;
+				if (NativeClassHierarchy->GetClassPath(Class, InternalPath, NativeClassHierarchyGetClassPathCache.GameModules))
+				{
+					FName ClassPath = FName(InternalPath);
+					if (TSharedPtr<const FNativeClassHierarchyNode> ClassNode = NativeClassHierarchy->FindNode(ClassPath, ENativeClassHierarchyNodeType::Class))
+					{
+						InCallback(CreateClassFileItem(ClassPath, ClassNode));
+					}
+				}
+			}
 		}
 	}
 }
@@ -398,11 +420,12 @@ bool UContentBrowserClassDataSource::EnumerateItemsForObjects(const TArrayView<U
 			if (NativeClassHierarchy->GetClassPath(InClass, InternalPath, NativeClassHierarchyGetClassPathCache.GameModules))
 			{
 				FName ClassPath = FName(InternalPath);
-				TSharedPtr<const FNativeClassHierarchyNode> ClassNode = NativeClassHierarchy->FindNode(ClassPath, ENativeClassHierarchyNodeType::Class);
-				const FContentBrowserItemPath ContentBrowserItemPath(InternalPath, EContentBrowserPathType::Internal);
-				if (!InCallback(ContentBrowserClassData::CreateClassFileItem(this, ContentBrowserItemPath.GetVirtualPathName(), ClassPath, InClass, ClassNode->LoadedFrom.IsSet())))
+				if (TSharedPtr<const FNativeClassHierarchyNode> ClassNode = NativeClassHierarchy->FindNode(ClassPath, ENativeClassHierarchyNodeType::Class))
 				{
-					return false;
+					if (!InCallback(CreateClassFileItem(ClassPath, ClassNode)))
+					{
+						return false;
+					}
 				}
 			}
 		}
diff --git a/Engine/Plugins/Editor/ContentBrowser/ContentBrowserClassDataSource/Source/ContentBrowserClassDataSource/Public/ContentBrowserClassDataSource.h b/Engine/Plugins/Editor/ContentBrowser/ContentBrowserClassDataSource/Source/ContentBrowserClassDataSource/Public/ContentBrowserClassDataSource.h
--- a/Engine/Plugins/Editor/ContentBrowser/ContentBrowserClassDataSource/Source/ContentBrowserClassDataSource/Public/ContentBrowserClassDataSource.h
+++ b/Engine/Plugins/Editor/ContentBrowser/ContentBrowserClassDataSource/Source/ContentBrowserClassDataSource/Public/ContentBrowserClassDataSource.h
@@ -46,6 +46,8 @@ public:
 
 	UE_API virtual void EnumerateItemsAtPath(const FName InPath, const EContentBrowserItemTypeFilter InItemTypeFilter, TFunctionRef<bool(FContentBrowserItemData&&)> InCallback) override;
 
+	UE_API virtual void EnumerateItemsAtUserProvidedPath(const FStringView InPath, const EContentBrowserItemTypeFilter InItemTypeFilter, TFunctionRef<bool(FContentBrowserItemData&&)> InCallback) override;
+
 	UE_API virtual bool EnumerateItemsForObjects(const TArrayView<UObject*> InObjects, TFunctionRef<bool(FContentBrowserItemData&&)> InCallback) override;
 
 	UE_API virtual bool IsFolderVisible(const FName InPath, const EContentBrowserIsFolderVisibleFlags InFlags, const FContentBrowserFolderContentsFilter& InContentsFilter) override;
diff --git a/Engine/Plugins/Editor/ContentBrowser/ContentBrowserFileDataSource/Source/ContentBrowserFileDataSource/Private/ContentBrowserFileDataSource.cpp b/Engine/Plugins/Editor/ContentBrowser/ContentBrowserFileDataSource/Source/ContentBrowserFileDataSource/Private/ContentBrowserFileDataSource.cpp
--- a/Engine/Plugins/Editor/ContentBrowser/ContentBrowserFileDataSource/Source/ContentBrowserFileDataSource/Private/ContentBrowserFileDataSource.cpp
+++ b/Engine/Plugins/Editor/ContentBrowser/ContentBrowserFileDataSource/Source/ContentBrowserFileDataSource/Private/ContentBrowserFileDataSource.cpp
@@ -978,6 +978,51 @@ void UContentBrowserFileDataSource::EnumerateItemsAtPath(const FName InPath, con
 	}
 }
 
+void UContentBrowserFileDataSource::EnumerateItemsAtUserProvidedPath(const FStringView InPath, const EContentBrowserItemTypeFilter InItemTypeFilter, TFunctionRef<bool(FContentBrowserItemData&&)> InCallback)
+{
+	TStringBuilder<256> PathBuilder;
+	PathBuilder.Append(InPath);
+	FPathViews::NormalizeDirectoryName(PathBuilder);
+
+	const FStringView Path(PathBuilder);
+	// Find the RegisteredFileMounts the Path is under.
+	for (const TPair<FName, FFileMount>& RegisteredFileMount : RegisteredFileMounts)
+	{
+		if (Path.StartsWith(RegisteredFileMount.Value.DiskPath) &&
+			(Path.Len() == RegisteredFileMount.Value.DiskPath.Len() || Path[RegisteredFileMount.Value.DiskPath.Len()] == TEXT('/')))
+		{
+			TStringBuilder<256> InternalPathBuilder;
+			RegisteredFileMount.Key.AppendString(InternalPathBuilder);
+			InternalPathBuilder.Append(FStringView(Path).RightChop(RegisteredFileMount.Value.DiskPath.Len()));
+
+			FName InternalPath(InternalPathBuilder, FNAME_Find);
+			if (!InternalPath.IsNone())
+			{
+				if (const FDiscoveredItem* DiscoveredItem = DiscoveredItems.Find(InternalPath))
+				{
+					switch (DiscoveredItem->Type)
+					{
+					case FDiscoveredItem::EType::Directory:
+						if (EnumHasAnyFlags(InItemTypeFilter, EContentBrowserItemTypeFilter::IncludeFolders))
+						{
+							InCallback(CreateFolderItem(InternalPath, DiscoveredItem->DiskPath));
+						}
+						break;
+					case FDiscoveredItem::EType::File:
+						if (EnumHasAnyFlags(InItemTypeFilter, EContentBrowserItemTypeFilter::IncludeFiles))
+						{
+							InCallback(CreateFolderItem(InternalPath, DiscoveredItem->DiskPath));
+						}
+						break;
+					}
+
+					return;
+				}
+			}
+		}
+	}
+}
+
 bool UContentBrowserFileDataSource::CanCreateFolder(const FName InPath, FText* OutErrorMsg)
 {
 	FName InternalPath;
diff --git a/Engine/Plugins/Editor/ContentBrowser/ContentBrowserFileDataSource/Source/ContentBrowserFileDataSource/Public/ContentBrowserFileDataSource.h b/Engine/Plugins/Editor/ContentBrowser/ContentBrowserFileDataSource/Source/ContentBrowserFileDataSource/Public/ContentBrowserFileDataSource.h
--- a/Engine/Plugins/Editor/ContentBrowser/ContentBrowserFileDataSource/Source/ContentBrowserFileDataSource/Public/ContentBrowserFileDataSource.h
+++ b/Engine/Plugins/Editor/ContentBrowser/ContentBrowserFileDataSource/Source/ContentBrowserFileDataSource/Public/ContentBrowserFileDataSource.h
@@ -53,6 +53,8 @@ public:
 
 	UE_API virtual void EnumerateItemsAtPath(const FName InPath, const EContentBrowserItemTypeFilter InItemTypeFilter, TFunctionRef<bool(FContentBrowserItemData&&)> InCallback) override;
 
+	UE_API virtual void EnumerateItemsAtUserProvidedPath(const FStringView InPath, const EContentBrowserItemTypeFilter InItemTypeFilter, TFunctionRef<bool(FContentBrowserItemData&&)> InCallback) override;
+
 	UE_API virtual bool CanCreateFolder(const FName InPath, FText* OutErrorMsg) override;
 
 	UE_API virtual bool CreateFolder(const FName InPath, const TSharedPtr<IContentBrowserHideFolderIfEmptyFilter>& HideFolderIfEmptyFilter, FContentBrowserItemDataTemporaryContext& OutPendingItem) override;
diff --git a/Engine/Source/Editor/ContentBrowser/Private/ContentBrowserUtils.cpp b/Engine/Source/Editor/ContentBrowser/Private/ContentBrowserUtils.cpp
--- a/Engine/Source/Editor/ContentBrowser/Private/ContentBrowserUtils.cpp
+++ b/Engine/Source/Editor/ContentBrowser/Private/ContentBrowserUtils.cpp
@@ -3,6 +3,8 @@
 
 #include "ContentBrowserUtils.h"
 
+#include "Algo/MinElement.h"
+#include "Algo/Unique.h"
 #include "AssetRegistry/ARFilter.h"
 #include "AssetRegistry/AssetData.h"
 #include "CollectionManagerModule.h"
@@ -88,6 +90,48 @@ namespace ContentBrowserUtils
 		IContentBrowserDataModule::Get().GetSubsystem()->TryConvertVirtualPath(FName(VirtualPath), ConvertedPath);
 		return ConvertedPath.ToString();
 	}
+
+	bool VirtualPathPassesBlockLists(const FContentBrowserItem& Item, const SPathView& PathViewForFiltering)
+	{
+		check(Item.GetInternalPath().IsNone());
+
+		bool bPasses = false;
+		// Test if any child internal paths pass for this fully virtual path
+		for (const FContentBrowserItemData& ItemData : Item.GetInternalItems())
+		{
+			UContentBrowserDataSource* ItemDataSource = ItemData.GetOwnerDataSource();
+			if (!ItemDataSource)
+			{
+				continue;
+			}
+
+			auto OnSubPath = [&PathViewForFiltering, &bPasses](FName VirtualSubPath, FName InternalPath)
+			{
+				if (InternalPath.IsNone())
+				{
+					// Keep enumerating, this path only exist virtually e.g. /All/Plugins
+					return true;
+				}
+
+				const FNameBuilder InternalPathBuilder(InternalPath);
+				if (ContentBrowserDataUtils::PathPassesAttributeFilter(InternalPathBuilder, 0, PathViewForFiltering.GetContentBrowserItemAttributeFilter())
+					&& PathViewForFiltering.InternalPathPassesBlockLists(InternalPathBuilder))
+				{
+					bPasses = true;
+					// Stop enumerating
+					return false;
+				}
+				return true;
+			};
+			ItemDataSource->GetRootPathVirtualTree().EnumerateSubPaths(Item.GetVirtualPath(), OnSubPath, /*bRecurse*/ true);
+
+			if (bPasses)
+			{
+				break;
+			}
+		}
+		return bPasses;
+	}
 }
 
 class SContentBrowserPopup : public SCompoundWidget
@@ -1207,7 +1251,7 @@ EContentBrowserItemAttributeFilter ContentBrowserUtils::GetContentBrowserItemAtt
 FContentBrowserItem ContentBrowserUtils::TryGetItemFromUserProvidedPath(FStringView RequestedPathView)
 {
 	// For all types of accepted input we can trim a trailing slash if it exists
-	if (RequestedPathView.EndsWith('/'))
+	while (FPathViews::HasRedundantTerminatingSeparator(RequestedPathView))
 	{
 		RequestedPathView.LeftChopInline(1);
 	}
@@ -1222,26 +1266,17 @@ FContentBrowserItem ContentBrowserUtils::TryGetItemFromUserProvidedPath(FStringV
 		return Item;
 	}
 
-	// If the path is a non-virtual path like /Game/Maps transform it into a virtual path and try and find an item there
-	FName VirtualPath = ContentBrowserData->ConvertInternalPathToVirtual(RequestedPath);
-	if (!VirtualPath.IsNone())
+	auto ProjectVirtualPath = [](const FContentBrowserItem& Item)
 	{
-		Item = ContentBrowserData->GetItemAtPath(VirtualPath, EContentBrowserItemTypeFilter::IncludeAll);
-		if (Item.IsValid())
-		{
-			return Item;
-		}
-	}
+		return Item.GetVirtualPath();
+	};
 
-	// If the string is a complete object path (with or without class), sync to that asset
-	FStringView ObjectPathView = RequestedPathView;
-	if (FPackageName::IsValidObjectPath(ObjectPathView) || FPackageName::ParseExportTextPath(RequestedPathView, nullptr, &ObjectPathView))
 	{
-		ContentBrowserData->ConvertInternalPathToVirtual(ObjectPathView, VirtualPath);
-		Item = ContentBrowserData->GetItemAtPath(VirtualPath, EContentBrowserItemTypeFilter::IncludeFiles);
-		if (Item.IsValid())
+		TArray<FContentBrowserItem> Items = ContentBrowserData->GetItemsAtUserProvidedPath(RequestedPathView, EContentBrowserItemTypeFilter::IncludeAll);
+		
+		if (!Items.IsEmpty())
 		{
-			return Item;
+			return *Algo::MinElementBy(Items, ProjectVirtualPath, FNameLexicalLess());
 		}
 	}
 
@@ -1252,93 +1287,25 @@ FContentBrowserItem ContentBrowserUtils::TryGetItemFromUserProvidedPath(FStringV
 		UE::Core::FVersePath::TryMake(VersePath, FString(RequestedPathView));
 	}
 
-	// If the string is a complete Verse path to an asset, sync to that asset
 	if (VersePath.IsValid())
 	{
+		// If the string is a complete Verse path to an asset, sync to that asset
 		FAssetData AssetData = AssetTools.FindAssetByVersePath(VersePath);
 		if (AssetData.IsValid())
 		{
 			FNameBuilder NameBuilder;
 			AssetData.GetSoftObjectPath().ToString(NameBuilder);
 
-			ContentBrowserData->ConvertInternalPathToVirtual(NameBuilder, VirtualPath);
-			Item = ContentBrowserData->GetItemAtPath(VirtualPath, EContentBrowserItemTypeFilter::IncludeFiles);
-			if (Item.IsValid())
-			{
-				return Item;
-			}
-		}
-	}
+			TArray<FContentBrowserItem> Items = ContentBrowserData->GetItemsAtUserProvidedPath(NameBuilder, EContentBrowserItemTypeFilter::IncludeFiles);
 
-	auto GetItemFromPackageName = [ContentBrowserData](const FString& PackageName) -> FContentBrowserItem {
-		// Packages like /Game/Characters/Knight do not map to virtual paths in data source, assets like /Game/Characters/Knight.Knight do.
-		// See if there's an item if we duplicate the last part of the path 
-		TStringBuilder<320> InternalPath(InPlace, PackageName, TEXTVIEW("."), FPackageName::GetShortName(PackageName));
-		FName VirtualPath;
-		ContentBrowserData->ConvertInternalPathToVirtual(InternalPath, VirtualPath);
-		FContentBrowserItem Item = ContentBrowserData->GetItemAtPath(VirtualPath, EContentBrowserItemTypeFilter::IncludeFiles);
-		if (Item.IsValid())
-		{
-			return Item;
-		}
-
-		// Otherwise go up to the package path and enumerate items to see if there's an asset with the desired package name 
-		FString PackagePath = FPackageName::GetLongPackagePath(PackageName);
-		ContentBrowserData->ConvertInternalPathToVirtual(PackagePath, VirtualPath);
-		FContentBrowserDataFilter Filter;
-		Filter.bRecursivePaths = false;
-		Filter.ItemTypeFilter = EContentBrowserItemTypeFilter::IncludeFiles;
-		ContentBrowserData->EnumerateItemsUnderPath(VirtualPath, Filter, [&Item, &PackageName](FContentBrowserItem&& InItem) { 
-			FName InternalPath = InItem.GetInternalPath();	
-			if (WriteToString<256>(InternalPath).ToView().StartsWith(PackageName))
+			if (!Items.IsEmpty())
 			{
-				Item = MoveTemp(InItem);
-				return false; 
+				return *Algo::MinElementBy(Items, ProjectVirtualPath, FNameLexicalLess());
 			}
-			return true;
-		});
-		if (Item.IsValid())
-		{
-			return Item;
 		}
-		return FContentBrowserItem();
-	};
 
-	// If the string is an incomplete virtual path that looks more like a package name 
-	// e.g. /All/Game/Maps/Arena rather than /Game/Maps/Arena or /All/Game/Maps/Arena.Arena
-	// try and convert it to an internal path, then try and use it as a package name 
-	{
-		FName ConvertedPath;
-		FString PackageName;
-		if (ContentBrowserData->TryConvertVirtualPath(RequestedPath, ConvertedPath) == EContentBrowserPathType::Internal)
-		{
-			if (FPackageName::IsValidLongPackageName(WriteToString<256>(ConvertedPath)))
-			{
-				PackageName = ConvertedPath.ToString();
-				Item = GetItemFromPackageName(PackageName);
-				if (Item.IsValid())
-				{
-					return Item;
-				}
-			}
-		}
-	}
-
-	// If the string is a filesystem path to a package, sync to that asset
-	FString PackageName;
-	if (FPackageName::IsValidLongPackageName(PackageName) || FPackageName::TryConvertFilenameToLongPackageName(FString(RequestedPathView), PackageName))
-	{
-		Item = GetItemFromPackageName(PackageName);
-		if (Item.IsValid())
-		{
-			return Item;
-		}
-	}
-
-	// If the string is a complete Verse path to a module, sync to that folder
-	if (VersePath.IsValid())
-	{
-		IPluginManager::Get().EnumerateVersePathMountPoints(VersePath, [ContentBrowserData, &VirtualPath, &Item](FStringView MountedAssetPath, FStringView LeafPath)
+		// If the string is a complete Verse path to a module, sync to that folder
+		IPluginManager::Get().EnumerateVersePathMountPoints(VersePath, [ContentBrowserData, &Item](FStringView MountedAssetPath, FStringView LeafPath)
 		{
 			FNameBuilder NameBuilder;
 			NameBuilder.Append(MountedAssetPath);
@@ -1354,6 +1321,7 @@ FContentBrowserItem ContentBrowserUtils::TryGetItemFromUserProvidedPath(FStringV
 					NameBuilder.Append(LeafPath);
 				}
 
+				FName VirtualPath;
 				ContentBrowserData->ConvertInternalPathToVirtual(NameBuilder, VirtualPath);
 				Item = ContentBrowserData->GetItemAtPath(VirtualPath, EContentBrowserItemTypeFilter::IncludeFolders);
 			}
@@ -1388,6 +1356,174 @@ FContentBrowserItem ContentBrowserUtils::TryGetItemFromUserProvidedPath(FStringV
 	return FContentBrowserItem();
 }
 
+TArray<FString> ContentBrowserUtils::GetChildPathsFromUserProvidedPath(FStringView RequestedPathView, EContentBrowserItemCategoryFilter ItemCategoryFilter, EContentBrowserItemAttributeFilter ItemAttributeFilter, const FName& ConfigInstanceName, const SPathView& PathViewForFiltering)
+{
+	// For all types of accepted input, we can trim a trailing slash if it exists.
+	while (FPathViews::HasRedundantTerminatingSeparator(RequestedPathView))
+	{
+		RequestedPathView.LeftChopInline(1);
+	}
+
+	UContentBrowserDataSubsystem* const ContentBrowserData = IContentBrowserDataModule::Get().GetSubsystem();
+
+	TArray<FContentBrowserItem> Items;
+	EContentBrowserPathType PathType = EContentBrowserPathType::None;
+	{
+		FName RequestedPath(RequestedPathView);
+
+		// If the path is a valid virtual path, use that.
+		FContentBrowserItem Item = ContentBrowserData->GetItemAtPath(RequestedPath, EContentBrowserItemTypeFilter::IncludeFolders);
+		if (Item.IsValid())
+		{
+			Items.Add(MoveTemp(Item));
+			PathType = EContentBrowserPathType::Virtual;
+		}
+		else
+		{
+			// Otherwise it may be some other type of valid path we understand, such as a package path.
+			Items = ContentBrowserData->GetItemsAtUserProvidedPath(RequestedPathView, EContentBrowserItemTypeFilter::IncludeFolders);
+
+			if (!Items.IsEmpty())
+			{
+				PathType = EContentBrowserPathType::Internal;
+			}
+			else // Lastly, it could be a Verse path.
+			{
+				const IAssetTools& AssetTools = IAssetTools::Get();
+				UE::Core::FVersePath VersePath;
+				if (AssetTools.ShowingContentVersePath())
+				{
+					UE::Core::FVersePath::TryMake(VersePath, FString(RequestedPathView));
+				}
+
+				if (VersePath.IsValid())
+				{
+					IPluginManager::Get().EnumerateVersePathMountPoints(VersePath, [ContentBrowserData, &Items](FStringView MountedAssetPath, FStringView LeafPath)
+					{
+						FNameBuilder NameBuilder;
+						NameBuilder.Append(MountedAssetPath);
+						if (ensure(NameBuilder.Len() > 0 && NameBuilder.LastChar() == TEXT('/')))
+						{
+							if (LeafPath.IsEmpty())
+							{
+								// ContentBrowserData expects no trailing /.
+								NameBuilder.RemoveSuffix(1);
+							}
+							else
+							{
+								NameBuilder.Append(LeafPath);
+							}
+
+							FName VirtualPath;
+							ContentBrowserData->ConvertInternalPathToVirtual(NameBuilder, VirtualPath);
+							FContentBrowserItem Item = ContentBrowserData->GetItemAtPath(VirtualPath, EContentBrowserItemTypeFilter::IncludeFolders);
+							if (Item.IsValid())
+							{
+								Items.Add(MoveTemp(Item));
+							}
+						}
+						return true;
+					});
+				}
+			}
+		}
+	}
+
+	if (Items.IsEmpty())
+	{
+		return {};
+	}
+
+	const UContentBrowserSettings* const ContentBrowserSettings = GetDefault<UContentBrowserSettings>();
+	bool bDisplayEmptyFolders = ContentBrowserSettings->DisplayEmptyFolders;
+	// Check to see if we have an instance config that overrides the default in UContentBrowserSettings.
+	if (const FContentBrowserInstanceConfig* EditorConfig = GetConstInstanceConfig(ConfigInstanceName))
+	{
+		bDisplayEmptyFolders = EditorConfig->bShowEmptyFolders;
+	}
+
+	const EContentBrowserIsFolderVisibleFlags IsFolderVisibleFlags = ContentBrowserUtils::GetIsFolderVisibleFlags(bDisplayEmptyFolders);
+
+	FContentBrowserFolderContentsFilter FolderFilter;
+	if (bDisplayEmptyFolders)
+	{
+		FolderFilter.HideFolderIfEmptyFilter = ContentBrowserData->CreateHideFolderIfEmptyFilter();
+	}
+	else
+	{
+		FolderFilter.ItemCategoryFilter = ItemCategoryFilter;
+	}
+
+	FContentBrowserDataFilter SubItemsFilter;
+	SubItemsFilter.ItemTypeFilter = EContentBrowserItemTypeFilter::IncludeFolders;
+	SubItemsFilter.bRecursivePaths = false;
+	SubItemsFilter.ItemCategoryFilter = ItemCategoryFilter;
+	SubItemsFilter.ItemAttributeFilter = ItemAttributeFilter;
+
+	TArray<FString> Results;
+
+	for (const FContentBrowserItem& Item : Items)
+	{
+		TArray<FContentBrowserItem> SubItems = ContentBrowserData->GetItemsUnderPath(Item.GetVirtualPath(), SubItemsFilter);
+
+		for (const FContentBrowserItem& SubItem : SubItems)
+		{
+			// Filter subitems we can't see.
+			if (!SubItem.GetInternalPath().IsNone())
+			{
+				if (!PathViewForFiltering.InternalPathPassesBlockLists(FNameBuilder(SubItem.GetInternalPath())))
+				{
+					continue;
+				}
+			}
+			else if (!VirtualPathPassesBlockLists(SubItem, PathViewForFiltering))
+			{
+				continue;
+			}
+
+			if (!ContentBrowserData->IsFolderVisible(SubItem.GetVirtualPath(), IsFolderVisibleFlags, FolderFilter))
+			{
+				continue;
+			}
+
+			// Get the subitem path matching the requested path format.
+			const FName Path = PathType == EContentBrowserPathType::Virtual ? SubItem.GetVirtualPath() : SubItem.GetInternalPath();
+			const FNameBuilder PathBuilder(Path);
+
+			FStringView PathView;
+			UE::Core::FVersePath VersePath;
+			if (PathType == EContentBrowserPathType::None)
+			{
+				// If not a virtual or internal path, it must be a Verse path.
+				// At the time of authoring the only possible option that GetItemsAtUserProvidedPath could have parsed is a package path.
+				VersePath = FPackageName::LongPackagePathToVersePath(PathBuilder);
+				if (!VersePath.IsValid())
+				{
+					continue;
+				}
+				PathView = VersePath.AsStringView();
+			}
+			else
+			{
+				PathView = PathBuilder.ToView();
+			}
+
+			// Verify that the new path is indeed syntactically a subpath.
+			if (PathView.StartsWith(RequestedPathView) &&
+				(RequestedPathView.EndsWith(TEXT('/')) || (PathView.Len() > RequestedPathView.Len() && PathView[RequestedPathView.Len()] == TEXT('/'))))
+			{
+				Results.Emplace(PathView);
+			}
+		}
+	}
+
+	// Check for duplicates - these should only be likely with Verse paths, but we don't know the internal implementation of custom data sources, so we'll always check just to be sure.
+	Results.Sort();
+	Results.SetNum(Algo::Unique(Results));
+
+	return Results;
+}
+
 FString ContentBrowserUtils::FormatCollectionCrumbData(const ICollectionContainer& CollectionContainer, const FCollectionNameType& Collection)
 {
 	return FString::Printf(
@@ -1470,16 +1606,24 @@ void ContentBrowserUtils::UpdateNavigationBar(const TSharedPtr<SNavigationBar>&
 	}
 }
 
-TArray<FContentBrowserItem> ContentBrowserUtils::FilterOrAliasItems(TArrayView<const FContentBrowserItem> Items)
+TArray<FContentBrowserItem> ContentBrowserUtils::FilterOrAliasItems(TArrayView<const FContentBrowserItem> Items, const SPathView& PathViewForFiltering)
 {
 	UContentBrowserDataSubsystem* ContentBrowserData = IContentBrowserDataModule::Get().GetSubsystem();
 	const TSharedRef<FPathPermissionList>& FolderPermissions = FAssetToolsModule::GetModule().Get().GetFolderPermissionList();
 
 	// If any of the items to sync don't pass the permission filter, try to find an alias that might be visible
 	TArray<FContentBrowserItem> NewItems;
+	NewItems.Reserve(Items.Num());
 	for (const FContentBrowserItem& Item : Items)
 	{
-		if (FolderPermissions->PassesStartsWithFilter(Item.GetInternalPath()))
+		if (Item.GetInternalPath().IsNone())
+		{
+			if (VirtualPathPassesBlockLists(Item, PathViewForFiltering))
+			{
+				NewItems.Add(Item);
+			}
+		}
+		else if (FolderPermissions->PassesStartsWithFilter(Item.GetInternalPath()))
 		{
 			NewItems.Add(Item);
 		}
@@ -1570,49 +1714,10 @@ TArray<FContentBrowserItem> ContentBrowserUtils::GetChildItemsFromVirtualPath(
 				continue;
 			}
 		}
-		else
+		else if (!VirtualPathPassesBlockLists(Item, PathViewForFiltering))
 		{
-			// Test if any child internal paths pass for this fully virtual path
-			bool bPasses = false;
-			for (const FContentBrowserItemData& ItemData : Item.GetInternalItems())
-			{
-				UContentBrowserDataSource* ItemDataSource = ItemData.GetOwnerDataSource();
-				if (!ItemDataSource) 
-				{
-					continue;
-				}
-
-				auto OnSubPath = [&PathViewForFiltering, &bPasses, &SubItemsFilter](FName VirtualSubPath, FName InternalPath)
-				{
-					if (InternalPath.IsNone())
-					{
-						// Keep enumerating, this path only exist virtually e.g. /All/Plugins
-						return true;
-					}
-
-					const FNameBuilder InternalPathBuilder(InternalPath);
-					if (ContentBrowserDataUtils::PathPassesAttributeFilter(InternalPathBuilder, 0, SubItemsFilter.ItemAttributeFilter) 
-					&& PathViewForFiltering.InternalPathPassesBlockLists(InternalPathBuilder))
-					{
-						bPasses = true;
-						// Stop enumerating
-						return false;
-					}
-					return true;
-				}; 
-				ItemDataSource->GetRootPathVirtualTree().EnumerateSubPaths(Item.GetVirtualPath(), OnSubPath, /*bRecurse*/ true);
-
-				if (bPasses)
-				{
-					break;
-				}
-			}
-
-			if (!bPasses)
-			{
-				It.RemoveCurrent();
-				continue;
-			}
+			It.RemoveCurrent();
+			continue;
 		}
 
 		if (!ContentBrowserData->IsFolderVisible(Item.GetVirtualPath(), FolderFlags, FolderFilter))
diff --git a/Engine/Source/Editor/ContentBrowser/Private/ContentBrowserUtils.h b/Engine/Source/Editor/ContentBrowser/Private/ContentBrowserUtils.h
--- a/Engine/Source/Editor/ContentBrowser/Private/ContentBrowserUtils.h
+++ b/Engine/Source/Editor/ContentBrowser/Private/ContentBrowserUtils.h
@@ -209,13 +209,16 @@ namespace ContentBrowserUtils
 	 */
 	FContentBrowserItem TryGetItemFromUserProvidedPath(FStringView RequestedPathView);
 
+	/** Finds the child paths that are children of the supplied path, while obeying the supplied filters. */
+	TArray<FString> GetChildPathsFromUserProvidedPath(FStringView RequestedPathView, EContentBrowserItemCategoryFilter ItemCategoryFilter, EContentBrowserItemAttributeFilter ItemAttributeFilter, const FName& ConfigInstanceName, const SPathView& PathViewForFiltering);
+
 	FString FormatCollectionCrumbData(const ICollectionContainer& CollectionContainer, const FCollectionNameType& Collection);
 	
 	/** Updates a navigation bar based on the current state of an asset view and path view */
 	void UpdateNavigationBar(const TSharedPtr<SNavigationBar>& NavigationBar, const TSharedPtr<SAssetView>& AssetView, const TSharedPtr<SPathView>& PathView);
 
 	/** Filter a list of items based on permissions or find a valid alias. */
-	TArray<FContentBrowserItem> FilterOrAliasItems(TArrayView<const FContentBrowserItem> Items);
+	TArray<FContentBrowserItem> FilterOrAliasItems(TArrayView<const FContentBrowserItem> Items, const SPathView& PathViewForFiltering);
 
 	/** Retrieves a named config. */
 	const FContentBrowserInstanceConfig* GetConstInstanceConfig(const FName& ForInstance);
diff --git a/Engine/Source/Editor/ContentBrowser/Private/SAssetDialog.cpp b/Engine/Source/Editor/ContentBrowser/Private/SAssetDialog.cpp
--- a/Engine/Source/Editor/ContentBrowser/Private/SAssetDialog.cpp
+++ b/Engine/Source/Editor/ContentBrowser/Private/SAssetDialog.cpp
@@ -1036,17 +1036,15 @@ TArray<FNavigationBarComboOption> SAssetDialog::OnCompletePathPrefix(const FStri
 {
 	FStringView PrefixView = Prefix;
 	
-	// Strip to last path separator
-	FName Parent;
+	// Strip to last path separator, but keep the very first path separator
 	if (int32 Index = UE::String::FindLastChar(PrefixView, '/'); Index != INDEX_NONE)
 	{	
-		PrefixView.LeftInline(Index);
-		Parent = FName(PrefixView);
+		PrefixView.LeftInline(FMath::Max(1, Index));
 	}
 
 	// Find PrefixView in the available tree of data sources, get its direct children, and filter them by SuffixView
-	TArray<FContentBrowserItem> SubItems = ContentBrowserUtils::GetChildItemsFromVirtualPath(
-		Parent,
+	TArray<FString> ChildPaths = ContentBrowserUtils::GetChildPathsFromUserProvidedPath(
+		PrefixView,
 		PathPicker->GetPathView()->GetContentBrowserItemCategoryFilter(),
 		PathPicker->GetPathView()->GetContentBrowserItemAttributeFilter(),
 		NAME_None,
@@ -1054,14 +1052,9 @@ TArray<FNavigationBarComboOption> SAssetDialog::OnCompletePathPrefix(const FStri
 	);
 	
 	TArray<FNavigationBarComboOption> Results;
-	for (const FContentBrowserItem& Item : SubItems)
+	for (const FString& ChildPath : ChildPaths)
 	{
-		FName Path = Item.GetVirtualPath();
-		FNameBuilder PathBuilder(Path);
-		if (PathBuilder.ToView().StartsWith(Prefix))
-		{
-			Results.Emplace(Path);
-		}
+		Results.Emplace(FText::AsCultureInvariant(ChildPath), ChildPath);
 	}
 	return Results;
 }
@@ -1196,7 +1189,7 @@ void SAssetDialog::SetCurrentlySelectedPath(const FString& NewPath, const EConte
 
 void SAssetDialog::SyncToItems(TArrayView<const FContentBrowserItem> ItemsToSync, const bool bAllowImplicitSync)
 {
-	TArray<FContentBrowserItem> NewItemsToSync = ContentBrowserUtils::FilterOrAliasItems(ItemsToSync);
+	TArray<FContentBrowserItem> NewItemsToSync = ContentBrowserUtils::FilterOrAliasItems(ItemsToSync, *PathPicker->GetPathView());
 	ItemsToSync = NewItemsToSync;
 	
 	// Tell the sources view first so the asset view will be up to date by the time we request the sync
diff --git a/Engine/Source/Editor/ContentBrowser/Private/SContentBrowser.cpp b/Engine/Source/Editor/ContentBrowser/Private/SContentBrowser.cpp
--- a/Engine/Source/Editor/ContentBrowser/Private/SContentBrowser.cpp
+++ b/Engine/Source/Editor/ContentBrowser/Private/SContentBrowser.cpp
@@ -1604,7 +1604,7 @@ void SContentBrowser::SyncToItems(TArrayView<const FContentBrowserItem> ItemsToS
 	// For now we just forcefully enable the legacy content source when a sync is requested so it succeeds
 	ContentSourcesContainer->ActivateLegacyContentSource();
 
-	TArray<FContentBrowserItem> NewItemsToSync = ContentBrowserUtils::FilterOrAliasItems(ItemsToSync);
+	TArray<FContentBrowserItem> NewItemsToSync = ContentBrowserUtils::FilterOrAliasItems(ItemsToSync, *LegacyContentSourceWidgets->PathViewPtr);
 	ItemsToSync = NewItemsToSync;
 	PrepareToSyncItems(ItemsToSync, bDisableFiltersThatHideAssets);
 
@@ -2631,31 +2631,24 @@ TArray<FNavigationBarComboOption> SContentBrowser::OnCompletePathPrefix(const FS
 {
 	FStringView PrefixView = Prefix;
 	
-	// Strip to last path separator
-	FName Parent;
+	// Strip to last path separator, but keep the very first path separator
 	if (int32 Index = UE::String::FindLastChar(PrefixView, '/'); Index != INDEX_NONE)
 	{	
-		PrefixView.LeftInline(Index);
-		Parent = FName(PrefixView);
+		PrefixView.LeftInline(FMath::Max(1, Index));
 	}
 
 	// Find PrefixView in the available tree of data sources, get its direct children, and filter them by SuffixView 
 	UContentBrowserDataSubsystem* ContentBrowserData = IContentBrowserDataModule::Get().GetSubsystem();
-	TArray<FContentBrowserItem> SubItems = ContentBrowserUtils::GetChildItemsFromVirtualPath(
-		Parent,
+	TArray<FString> ChildPaths = ContentBrowserUtils::GetChildPathsFromUserProvidedPath(
+		PrefixView,
 		LegacyContentSourceWidgets->PathViewPtr->GetContentBrowserItemCategoryFilter(),
 		LegacyContentSourceWidgets->PathViewPtr->GetContentBrowserItemAttributeFilter(),
 		InstanceName,
 		*LegacyContentSourceWidgets->PathViewPtr); 
 	TArray<FNavigationBarComboOption> Results;
-	for (const FContentBrowserItem& Item : SubItems)
+	for (const FString& ChildPath : ChildPaths)
 	{
-		FName Path = Item.GetVirtualPath();
-		FNameBuilder PathBuilder(Path);
-		if (PathBuilder.ToView().StartsWith(Prefix))
-		{
-			Results.Emplace(Item.GetVirtualPath());
-		}
+		Results.Emplace(FText::AsCultureInvariant(ChildPath), ChildPath);
 	}
 	return Results;
 }
diff --git a/Engine/Source/Editor/ContentBrowser/Private/SNavigationBar.cpp b/Engine/Source/Editor/ContentBrowser/Private/SNavigationBar.cpp
--- a/Engine/Source/Editor/ContentBrowser/Private/SNavigationBar.cpp
+++ b/Engine/Source/Editor/ContentBrowser/Private/SNavigationBar.cpp
@@ -31,7 +31,7 @@ enum class ELocationSource
 struct FLocationItem
 {
 	FText DisplayPath;
-	FString VirtualPath;
+	FString Path;
 	ELocationSource Source;
 };
 
@@ -76,7 +76,7 @@ public:
 		Items.Reset(InItems.Num());
 		for (FNavigationBarComboOption& Item : InItems)
 		{
-			Items.Emplace(MakeShared<FLocationItem>(FLocationItem{ MoveTemp(Item.DisplayPath), MoveTemp(Item.VirtualPath), InSource }));
+			Items.Emplace(MakeShared<FLocationItem>(FLocationItem{ MoveTemp(Item.DisplayPath), MoveTemp(Item.Path), InSource }));
 		}
 	}
 	
@@ -104,15 +104,9 @@ protected:
 	TArray<TSharedPtr<FLocationItem>> Items;
 };
 
-FNavigationBarComboOption::FNavigationBarComboOption(FName InVirtualPath)
-{
-	VirtualPath = InVirtualPath.ToString();
-	DisplayPath = FText::FromString(VirtualPath);
-}
-
-FNavigationBarComboOption::FNavigationBarComboOption(const FText& InDisplayPath, const FString& InVirtualPath)
+FNavigationBarComboOption::FNavigationBarComboOption(const FText& InDisplayPath, const FString& InPath)
 	: DisplayPath(InDisplayPath)
-	, VirtualPath(InVirtualPath)
+	, Path(InPath)
 {
 }
 
@@ -341,7 +335,7 @@ FReply SNavigationBar::HandleComboKeyChar(const FGeometry& MyGeometry, const FCh
 	TArray<TSharedPtr<FLocationItem>> SelectedItems = ComboListView->GetSelectedItems();
 	if (SelectedItems.Num() != 0)
 	{
-		FString NewText = SelectedItems[0]->VirtualPath;
+		FString NewText = SelectedItems[0]->Path;
 		NewText += InCharEvent.GetCharacter();
 		EditableText->SetText(FText::FromString(NewText));
 		EditableText->SetSelectAllTextWhenFocused(false);
@@ -387,7 +381,7 @@ FReply SNavigationBar::HandleComboKeyDown(const FGeometry& MyGeometry, const FKe
 		TArray<TSharedPtr<FLocationItem>> SelectedItems = ComboListView->GetSelectedItems();
 		if (SelectedItems.Num() > 0)
 		{
-			FString NewText = SelectedItems[0]->VirtualPath + TEXT("/"); // TODO: Ideally only add slash if this prefix has children itself 
+			FString NewText = SelectedItems[0]->Path + TEXT("/"); // TODO: Ideally only add slash if this prefix has children itself 
 			EditableText->SetText(FText::FromString(NewText));
 			EditableText->SetSelectAllTextWhenFocused(false);
 			return FReply::Handled().SetUserFocus(EditableText.ToSharedRef(), EFocusCause::SetDirectly);
@@ -537,12 +531,12 @@ void SNavigationBar::HandleComboSelectionChanged(TSharedPtr<FLocationItem> Selec
 			// Stop editing/hide popup and navigate to chosen path
 			EditTextVisibility = EVisibility::Hidden;
 			SetIsOpen(false);
-			OnNavigateToPath.ExecuteIfBound(SelectedItem->VirtualPath);
+			OnNavigateToPath.ExecuteIfBound(SelectedItem->Path);
 			break;
 		case ESelectInfo::OnNavigation:
 			// Mirror chosen item into text box without moving focus back or triggering new suggestions
 			TGuardValue<bool> BlockSuggestions{bNoSuggestionsFromTextChange, true};
-			EditableText->SetText(FText::FromString(SelectedItem->VirtualPath));
+			EditableText->SetText(FText::FromString(SelectedItem->Path));
 			break;
 	}
 }
diff --git a/Engine/Source/Editor/ContentBrowser/Private/SNavigationBar.h b/Engine/Source/Editor/ContentBrowser/Private/SNavigationBar.h
--- a/Engine/Source/Editor/ContentBrowser/Private/SNavigationBar.h
+++ b/Engine/Source/Editor/ContentBrowser/Private/SNavigationBar.h
@@ -25,11 +25,10 @@ class USlateWidgetStyleAsset;
 struct FNavigationBarComboOption
 {
 	FNavigationBarComboOption() = default;
-	explicit FNavigationBarComboOption(FName InVirtualPath);
-	FNavigationBarComboOption(const FText& InDisplayPath, const FString& InVirtualPath);
+	FNavigationBarComboOption(const FText& InDisplayPath, const FString& InPath);
 
 	FText DisplayPath;
-	FString VirtualPath;
+	FString Path;
 };
 
 using FOnPathClicked = TDelegate<void(const FString&)>;
diff --git a/Engine/Source/Editor/ContentBrowserData/Private/ContentBrowserDataSource.cpp b/Engine/Source/Editor/ContentBrowserData/Private/ContentBrowserDataSource.cpp
--- a/Engine/Source/Editor/ContentBrowserData/Private/ContentBrowserDataSource.cpp
+++ b/Engine/Source/Editor/ContentBrowserData/Private/ContentBrowserDataSource.cpp
@@ -240,6 +240,15 @@ bool UContentBrowserDataSource::EnumerateItemsAtPaths(const TArrayView<FContentB
 	return true;
 }
 
+void UContentBrowserDataSource::EnumerateItemsAtUserProvidedPath(const FStringView InPath, const EContentBrowserItemTypeFilter InItemTypeFilter, TFunctionRef<bool(FContentBrowserItemData&&)> InCallback)
+{
+}
+
+void UContentBrowserDataSource::EnumerateItemsAtUserProvidedPath(const FStringView InPath, const EContentBrowserItemTypeFilter InItemTypeFilter, const TGetOrEnumerateSink<FContentBrowserItemData>& InSink)
+{
+	EnumerateItemsAtUserProvidedPath(InPath, InItemTypeFilter, [&InSink](FContentBrowserItemData&& Item) { return InSink.ProduceItem(MoveTemp(Item)); });
+}
+
 bool UContentBrowserDataSource::EnumerateItemsForObjects(const TArrayView<UObject*> InObjects, TFunctionRef<bool(FContentBrowserItemData&&)> InCallback)
 {
 	return true;
diff --git a/Engine/Source/Editor/ContentBrowserData/Private/ContentBrowserDataSubsystem.cpp b/Engine/Source/Editor/ContentBrowserData/Private/ContentBrowserDataSubsystem.cpp
--- a/Engine/Source/Editor/ContentBrowserData/Private/ContentBrowserDataSubsystem.cpp
+++ b/Engine/Source/Editor/ContentBrowserData/Private/ContentBrowserDataSubsystem.cpp
@@ -513,6 +513,24 @@ bool UContentBrowserDataSubsystem::EnumerateItemsAtPaths(const TArrayView<struct
 	return true;
 }
 
+void UContentBrowserDataSubsystem::EnumerateItemsAtUserProvidedPath(const FStringView InPath, const EContentBrowserItemTypeFilter InItemTypeFilter, TFunctionRef<bool(FContentBrowserItem&&)> InCallback) const
+{
+	EnumerateItemsAtUserProvidedPath(InPath, InItemTypeFilter, [&InCallback](FContentBrowserItemData&& InItemData)
+	{
+		checkf(InItemData.IsValid(), TEXT("Enumerated items must be valid!"));
+		return InCallback(FContentBrowserItem(MoveTemp(InItemData)));
+	});
+}
+
+void UContentBrowserDataSubsystem::EnumerateItemsAtUserProvidedPath(const FStringView InPath, const EContentBrowserItemTypeFilter InItemTypeFilter, TFunctionRef<bool(FContentBrowserItemData&&)> InCallback) const
+{
+	for (const TPair<FName, UContentBrowserDataSource*>& ActiveDataSourcePair : ActiveDataSources)
+	{
+		UContentBrowserDataSource* DataSource = ActiveDataSourcePair.Value;
+		DataSource->EnumerateItemsAtUserProvidedPath(InPath, InItemTypeFilter, InCallback);
+	}
+}
+
 bool UContentBrowserDataSubsystem::EnumerateItemsForObjects(TArrayView<UObject*> InObjects, TFunctionRef<bool(FContentBrowserItemData&&)> InCallback) const
 {
 	for (const auto& ActiveDataSourcePair : ActiveDataSources)
@@ -569,6 +587,24 @@ FContentBrowserItem UContentBrowserDataSubsystem::GetItemAtPath(const FName InPa
 	return FoundItem;
 }
 
+TArray<FContentBrowserItem> UContentBrowserDataSubsystem::GetItemsAtUserProvidedPath(const FStringView InPath, const EContentBrowserItemTypeFilter InItemTypeFilter) const
+{
+	TMap<FContentBrowserItemKey, FContentBrowserItem> FoundItems;
+	EnumerateItemsAtUserProvidedPath(InPath, InItemTypeFilter, [&FoundItems](FContentBrowserItemData&& InItemData)
+	{
+		checkf(InItemData.IsValid(), TEXT("Enumerated items must be valid!"));
+
+		FContentBrowserItem& FoundItem = FoundItems.FindOrAdd(FContentBrowserItemKey(InItemData));
+		FoundItem.Append(InItemData);
+
+		return true;
+	});
+
+	TArray<FContentBrowserItem> FoundItemsArray;
+	FoundItems.GenerateValueArray(FoundItemsArray);
+	return FoundItemsArray;
+}
+
 TArray<FContentBrowserItemPath> UContentBrowserDataSubsystem::GetAliasesForPath(const FContentBrowserItemPath InPath) const
 {
 	return GetAliasesForPath(InPath.GetInternalPathName());
diff --git a/Engine/Source/Editor/ContentBrowserData/Public/ContentBrowserDataSource.h b/Engine/Source/Editor/ContentBrowserData/Public/ContentBrowserDataSource.h
--- a/Engine/Source/Editor/ContentBrowserData/Public/ContentBrowserDataSource.h
+++ b/Engine/Source/Editor/ContentBrowserData/Public/ContentBrowserDataSource.h
@@ -246,6 +246,19 @@ public:
 	 */
 	UE_API virtual bool EnumerateItemsAtPaths(const TArrayView<FContentBrowserItemPath> InPaths, const EContentBrowserItemTypeFilter InItemTypeFilter, TFunctionRef<bool(FContentBrowserItemData&&)> InCallback);
 
+	/**
+	 * Enumerate items at the given user provided path, optionally filtering by type, and invoking the callback for each matching item.
+	 * The user provided path can be any path format the data source supports other than a virtual path. The data source will deduce the path type.
+	 * @note This function *must not* block waiting on content discovery! It should use the current state as known at this point in time.
+	 *
+	 * @param InPath The user provided path to find items for.
+	 * @param InItemTypeFilter The types of items we want to find.
+	 * @param InCallback The function to invoke for each matching item (return true to continue enumeration).
+	 * @param InSink Parameter provided to receive output as either a callback or an array.
+	 */
+	UE_API virtual void EnumerateItemsAtUserProvidedPath(const FStringView InPath, const EContentBrowserItemTypeFilter InItemTypeFilter, TFunctionRef<bool(FContentBrowserItemData&&)> InCallback);
+	UE_API virtual void EnumerateItemsAtUserProvidedPath(const FStringView InPath, const EContentBrowserItemTypeFilter InItemTypeFilter, const TGetOrEnumerateSink<FContentBrowserItemData>& InSink);
+
 	/**
 	 * Enumerate the items (files) that exist for the given objects.
 	 * @note Multiple items may have the same virtual path if they are different types, or come from different data sources.
diff --git a/Engine/Source/Editor/ContentBrowserData/Public/ContentBrowserDataSubsystem.h b/Engine/Source/Editor/ContentBrowserData/Public/ContentBrowserDataSubsystem.h
--- a/Engine/Source/Editor/ContentBrowserData/Public/ContentBrowserDataSubsystem.h
+++ b/Engine/Source/Editor/ContentBrowserData/Public/ContentBrowserDataSubsystem.h
@@ -231,6 +231,14 @@ public:
 	 */
 	UE_API bool EnumerateItemsAtPaths(const TArrayView<struct FContentBrowserItemPath> InItemPaths, const EContentBrowserItemTypeFilter InItemTypeFilter, TFunctionRef<bool(FContentBrowserItemData&&)> InCallback) const;
 
+	/**
+	 * Enumerate the items (folders and/or files) that exist at the given user provided path.
+	 * The user provided path can be any path format any of the data sources support other than a virtual path. The data sources will deduce the path type.
+	 * @note Multiple items may have the same virtual path if they are different types, or come from different data sources.
+	 */
+	UE_API void EnumerateItemsAtUserProvidedPath(const FStringView InPath, const EContentBrowserItemTypeFilter InItemTypeFilter, TFunctionRef<bool(FContentBrowserItem&&)> InCallback) const;
+	UE_API void EnumerateItemsAtUserProvidedPath(const FStringView InPath, const EContentBrowserItemTypeFilter InItemTypeFilter, TFunctionRef<bool(FContentBrowserItemData&&)> InCallback) const;
+
 	/**
 	 * Enumerate the items (files) that exist for the given objects.
 	 * @note Multiple items may have the same virtual path if they are different types, or come from different data sources.
@@ -253,6 +261,13 @@ public:
 	UFUNCTION(BlueprintCallable, Category="ContentBrowser")
 	UE_API FContentBrowserItem GetItemAtPath(const FName InPath, const EContentBrowserItemTypeFilter InItemTypeFilter) const;
 
+	/**
+	 * Get the items (folders and/or files) that exist at the given user provided path.
+	 * The user provided path can be any path format any of the data sources support other than a virtual path. The data sources will deduce the path type.
+	 * @note Multiple items may have the same virtual path if they are different types, or come from different data sources.
+	 */
+	UE_API TArray<FContentBrowserItem> GetItemsAtUserProvidedPath(const FStringView InPath, const EContentBrowserItemTypeFilter InItemTypeFilter) const;
+
 	/**
 	 * Get a list of other paths that the data source may be using to represent a specific path
 	 *
