Optimization to the path view widget, only populate the folder tree when widget is about to be drawn.
Throw away the item updates when waiting for a population request.
This change could result in a hitch the first time a Content browser is shown into view but overall it help the engine boot faster.

diff --git a/Engine/Source/Editor/ContentBrowser/Private/SContentBrowser.cpp b/Engine/Source/Editor/ContentBrowser/Private/SContentBrowser.cpp
--- a/Engine/Source/Editor/ContentBrowser/Private/SContentBrowser.cpp
+++ b/Engine/Source/Editor/ContentBrowser/Private/SContentBrowser.cpp
@@ -2272,19 +2272,25 @@ void SContentBrowser::SetSelectedPaths(const TArray<FString>& FolderPaths, bool
 {
 	if (FolderPaths.Num() > 0)
 	{
+		TArray<FString> NormalizedFolderPaths = FolderPaths;
+		for (FString& NormalizedFolderPath : NormalizedFolderPaths)
+		{
+			FPaths::NormalizeDirectoryName(NormalizedFolderPath);
+		}
+
 		if (bNeedsRefresh)
 		{
-			LegacyContentSourceWidgets->PathViewPtr->Populate();
+			LegacyContentSourceWidgets->PathViewPtr->RequestPopulation(FSimpleMulticastDelegate::FDelegate::CreateSPLambda(this, [this, NormalizedFolderPaths]()
+				{
+					LegacyContentSourceWidgets->PathViewPtr->SetSelectedPaths(NormalizedFolderPaths);
+				}));
 			LegacyContentSourceWidgets->FavoritePathViewPtr->Populate();
 		}
-
-		TArray<FString> NormalizedFolderPaths = FolderPaths;
-		for (FString& NormalizedFolderPath : NormalizedFolderPaths)
+		else
 		{
-			FPaths::NormalizeDirectoryName(NormalizedFolderPath);
+			LegacyContentSourceWidgets->PathViewPtr->SetSelectedPaths(NormalizedFolderPaths);
 		}
 
-		LegacyContentSourceWidgets->PathViewPtr->SetSelectedPaths(NormalizedFolderPaths);
 		LegacyContentSourceWidgets->FavoritePathViewPtr->SetSelectedPaths(NormalizedFolderPaths);
 		PathSelected(NormalizedFolderPaths[0]);
 	}
diff --git a/Engine/Source/Editor/ContentBrowser/Private/SPathView.cpp b/Engine/Source/Editor/ContentBrowser/Private/SPathView.cpp
--- a/Engine/Source/Editor/ContentBrowser/Private/SPathView.cpp
+++ b/Engine/Source/Editor/ContentBrowser/Private/SPathView.cpp
@@ -922,6 +922,11 @@ SPathView::FScopedSelectionChangedEvent::~FScopedSelectionChangedEvent()
 	}
 }
 
+bool SPathView::FScopedSelectionChangedEvent::IsInitialSelectionEmpty() const
+{
+	return InitialSelectionSet.IsEmpty();
+}
+
 TSet<FName> SPathView::FScopedSelectionChangedEvent::GetSelectionSet() const
 {
 	TSet<FName> SelectionSet;
@@ -1154,40 +1159,13 @@ void SPathView::Construct( const FArguments& InArgs )
 	];
 
 	CustomFolderPermissionList = InArgs._CustomFolderPermissionList;
+
 	// Add all paths currently gathered from the asset registry
-	Populate();
+	bInitialPopulationRequested = true;
 
-	for (const FName PathToExpand : GetDefaultPathsToExpand())
-	{
-		if (TSharedPtr<FTreeItem> FoundItem = TreeData->FindTreeItem(PathToExpand))
-		{
-			RecursiveExpandParents(FoundItem);
-			TreeViewPtr->SetItemExpansion(FoundItem, true);
-		}
-	}
+	DefaultPath = InArgs._DefaultPath;
+	bCreateDefaultPath = InArgs._CreateDefaultPath;
 
-	if (!InArgs._DefaultPath.IsEmpty() && InternalPathPassesBlockLists(InArgs._DefaultPath))
-	{
-		const FName VirtualPath = ContentBrowserData->ConvertInternalPathToVirtual(*InArgs._DefaultPath);
-		FName InternalPath;
-		if (ContentBrowserData->TryConvertVirtualPath(VirtualPath, InternalPath) != EContentBrowserPathType::Internal)
-		{
-			InternalPath = FName();
-		}
-		if (InArgs._CreateDefaultPath && !TreeData->FindTreeItem(VirtualPath))
-		{
-			const FString DefaultPathLeafName = FPaths::GetPathLeaf(VirtualPath.ToString());
-			TreeData->AddFolderItem(FContentBrowserItemData(nullptr,
-				EContentBrowserItemFlags::Type_Folder,
-				VirtualPath,
-				*DefaultPathLeafName,
-				FText(),
-				nullptr,
-				InternalPath));
-		}
-
-		SetSelectedPaths({ VirtualPath.ToString() });
-	}
 }
 
 void SPathView::ConfigureTreeView(STreeView<TSharedPtr<FTreeItem>>::FArguments& InArgs)
@@ -1201,6 +1179,11 @@ void SPathView::Tick(const FGeometry& AllottedGeometry, const double InCurrentTi
 {
 	Super::Tick(AllottedGeometry, InCurrentTime, InDeltaTime);
 
+	if (IsPopulationRequested())
+	{
+		Populate();
+	}
+
 	if (LastTreeDataVersion != TreeData->GetVersion())
 	{
 		LastTreeDataVersion = TreeData->GetVersion();
@@ -1272,7 +1255,7 @@ void SPathView::PopulatePathViewFiltersMenu(UToolMenu* Menu)
 void SPathView::PluginPathFilterClicked(TSharedRef<FContentBrowserPluginFilter> Filter)
 {
 	SetPluginPathFilterActive(Filter, !IsPluginPathFilterInUse(Filter));
-	Populate();
+	bPopulationRequested = true;
 }
 
 bool SPathView::IsPluginPathFilterChecked(TSharedRef<FContentBrowserPluginFilter> Filter) const
@@ -1307,7 +1290,7 @@ void SPathView::ResetPluginPathFilters()
 		SetPluginPathFilterActive(Filter, false);
 	}
 
-	Populate();
+	bPopulationRequested = true;
 }
 
 bool SPathView::DisablePluginPathFiltersThatHideItems(TConstArrayView<FContentBrowserItem> Items)
@@ -1987,12 +1970,6 @@ EActiveTimerReturnType SPathView::SetFocusPostConstruct( double InCurrentTime, f
 	return EActiveTimerReturnType::Stop;
 }
 
-EActiveTimerReturnType SPathView::TriggerRepopulate(double InCurrentTime, float InDeltaTime)
-{
-	Populate();
-	return EActiveTimerReturnType::Stop;
-}
-
 TSharedPtr<SWidget> SPathView::MakePathViewContextMenu()
 {
 	if (!bAllowContextMenu || !OnGetItemContextMenu.IsBound())
@@ -2264,6 +2241,17 @@ FText SPathView::GetHighlightText() const
 	return TreeData->GetFolderPathTextFilter().GetRawFilterText();
 }
 
+void SPathView::RequestPopulation(FSimpleMulticastDelegate::FDelegate&& AfterPopulation)
+{
+	bPopulationRequested = true;
+
+	if (AfterPopulation.IsBound())
+	{
+		OnNextPopulate.Add(MoveTemp(AfterPopulation));
+	}
+}
+
+
 void SPathView::Populate(const bool bIsRefreshingFilter)
 {
 	TRACE_CPUPROFILER_EVENT_SCOPE(SPathView::Populate);
@@ -2279,16 +2267,39 @@ void SPathView::Populate(const bool bIsRefreshingFilter)
 	TreeData->FilterFullFolderTree();
 	TreeData->SortRootItems();
 
+	if (bPopulationForSettingChangeRequested)
+	{
+		PostPopulateHandleSettingChanged(!ScopedSelectionChangedEvent.IsInitialSelectionEmpty());
+	}
+
+	if (bInitialPopulationRequested)
+	{
+		PostInitialPopulation();
+	}
+
 	// Select any of our initial paths which aren't currently selected
-	if (Algo::AllOf(PendingInitialPaths, [this](FName VirtualPath) { 
+	if (Algo::AllOf(PendingInitialPaths, [this](FName VirtualPath) {
 		return ExplicitlyAddPathToSelection(VirtualPath);
-	}))
+		}))
 	{
 		UE_LOG(LogPathView, VeryVerbose, TEXT("[%s] Resetting pending initial paths as all are selected"), *WriteToString<256>(OwningContentBrowserName));
 		PendingInitialPaths.Reset();
 	}
+
+	SharedPostPopulate();
+}
+
+void SPathView::SharedPostPopulate()
+{
+	OnNextPopulate.Broadcast();
+	OnNextPopulate.Clear();
+
+	bPopulationRequested = false;
+	bInitialPopulationRequested = false;
+	bPopulationForSettingChangeRequested = false;
 }
 
+
 FReply SPathView::OnFolderDragDetected(const FGeometry& Geometry, const FPointerEvent& MouseEvent)
 {
 	if ( MouseEvent.IsMouseButtonDown( EKeys::LeftMouseButton ) )
@@ -2422,6 +2433,15 @@ bool SPathView::IsTreeItemSelected(TSharedPtr<FTreeItem> TreeItem) const
 
 void SPathView::HandleItemDataUpdated(TArrayView<const FContentBrowserItemDataUpdate> InUpdatedItems)
 {
+	if (IsPopulationRequested())
+	{
+		/**
+		 * Skip the data updates.
+		 * We are about to refresh the view from scratch.
+		 */
+		return;
+	}
+
 	TRACE_CPUPROFILER_EVENT_SCOPE(SPathView::HandleItemDataUpdated);
 
 	if (InUpdatedItems.Num() == 0)
@@ -2462,14 +2482,7 @@ void SPathView::HandleItemDataUpdated(TArrayView<const FContentBrowserItemDataUp
 
 void SPathView::HandleItemDataRefreshed()
 {
-	// Populate immediately, as the path view must be up to date for Content Browser selection to work correctly
-	// and since it defaults to being hidden, it potentially won't be ticked to run this update latently
-	Populate();
-
-	/*
-	// The class hierarchy has changed in some way, so we need to refresh our set of paths
-	RegisterActiveTimer(0.f, FWidgetActiveTimerDelegate::CreateSP(this, &SPathView::TriggerRepopulate));
-	*/
+	bPopulationRequested = true;
 }
 
 void SPathView::HandleItemDataDiscoveryComplete()
@@ -2498,56 +2511,102 @@ void SPathView::HandleSettingChanged(FName PropertyName)
 			*WriteToString<256>(PropertyName));
 		const bool bHadSelectedPath = TreeViewPtr->GetNumItemsSelected() > 0;
 
-		// Update our path view so that it can include/exclude the dev folder
-		Populate();
+		// Update our path view so that it can include/exclude the permitted folder
+		bPopulationForSettingChangeRequested = true;
+	}
+}
 
-		// If folder is no longer visible but we're inside it...
-		if (TreeViewPtr->GetNumItemsSelected() == 0 && bHadSelectedPath)
+void SPathView::PostPopulateHandleSettingChanged(bool bHadSelectedPath)
+{
+	// If folder is no longer visible but we're inside it...
+	if (TreeViewPtr->GetNumItemsSelected() == 0 && bHadSelectedPath)
+	{
+		for (const FName VirtualPath : GetDefaultPathsToSelect())
 		{
-			for (const FName VirtualPath : GetDefaultPathsToSelect())
+			if (TSharedPtr<FTreeItem> TreeItemToSelect = TreeData->FindTreeItem(VirtualPath))
 			{
-				if (TSharedPtr<FTreeItem> TreeItemToSelect = TreeData->FindTreeItem(VirtualPath))
-				{
-					TreeViewPtr->SetSelection(TreeItemToSelect);
-					break;
-				}
+				TreeViewPtr->SetSelection(TreeItemToSelect);
+				break;
 			}
 		}
+	}
 
-		// If the dev or engine folder has become visible and we're inside it...
-		const UContentBrowserSettings* ContentBrowserSettings = GetDefault<UContentBrowserSettings>();
-		bool bDisplayDev = ContentBrowserSettings->GetDisplayDevelopersFolder();
-		bool bDisplayEngine = ContentBrowserSettings->GetDisplayEngineFolder();
-		bool bDisplayPlugins = ContentBrowserSettings->GetDisplayPluginFolders();
-		bool bDisplayL10N = ContentBrowserSettings->GetDisplayL10NFolder();
-		// check to see if we have an instance config that overrides the default in UContentBrowserSettings
-		if (FContentBrowserInstanceConfig* EditorConfig = GetContentBrowserConfig())
+	// If the dev or engine folder has become visible and we're inside it...
+	const UContentBrowserSettings* ContentBrowserSettings = GetDefault<UContentBrowserSettings>();
+	bool bDisplayDev = ContentBrowserSettings->GetDisplayDevelopersFolder();
+	bool bDisplayEngine = ContentBrowserSettings->GetDisplayEngineFolder();
+	bool bDisplayPlugins = ContentBrowserSettings->GetDisplayPluginFolders();
+	bool bDisplayL10N = ContentBrowserSettings->GetDisplayL10NFolder();
+	// check to see if we have an instance config that overrides the default in UContentBrowserSettings
+	if (FContentBrowserInstanceConfig* EditorConfig = GetContentBrowserConfig())
+	{
+		bDisplayDev = EditorConfig->bShowDeveloperContent;
+		bDisplayEngine = EditorConfig->bShowEngineContent;
+		bDisplayPlugins = EditorConfig->bShowPluginContent;
+		bDisplayL10N = EditorConfig->bShowLocalizedContent;
+	}
+
+	if (bDisplayDev || bDisplayEngine || bDisplayPlugins || bDisplayL10N)
+	{
+		const TArray<FContentBrowserItem> NewSelectedItems = GetSelectedFolderItems();
+		if (NewSelectedItems.Num() > 0)
 		{
-			bDisplayDev = EditorConfig->bShowDeveloperContent;
-			bDisplayEngine = EditorConfig->bShowEngineContent;
-			bDisplayPlugins = EditorConfig->bShowPluginContent;
-			bDisplayL10N = EditorConfig->bShowLocalizedContent;
+			const FContentBrowserItem& NewSelectedItem = NewSelectedItems[0];
+
+			if ((bDisplayDev && ContentBrowserUtils::IsItemDeveloperContent(NewSelectedItem)) ||
+				(bDisplayEngine && ContentBrowserUtils::IsItemEngineContent(NewSelectedItem)) ||
+				(bDisplayPlugins && ContentBrowserUtils::IsItemPluginContent(NewSelectedItem)) ||
+				(bDisplayL10N && ContentBrowserUtils::IsItemLocalizedContent(NewSelectedItem))
+				)
+			{
+				// Refresh the contents
+				OnItemSelectionChanged.ExecuteIfBound(NewSelectedItem, ESelectInfo::Direct);
+			}
 		}
+	}
+}
 
-		if (bDisplayDev || bDisplayEngine || bDisplayPlugins || bDisplayL10N)
+void SPathView::PostInitialPopulation()
+{
+	for (const FName PathToExpand : GetDefaultPathsToExpand())
+	{
+		if (TSharedPtr<FTreeItem> FoundItem = TreeData->FindTreeItem(PathToExpand))
 		{
-			const TArray<FContentBrowserItem> NewSelectedItems = GetSelectedFolderItems();
-			if (NewSelectedItems.Num() > 0)
-			{
-				const FContentBrowserItem& NewSelectedItem = NewSelectedItems[0];
+			RecursiveExpandParents(FoundItem);
+			TreeViewPtr->SetItemExpansion(FoundItem, true);
+		}
+	}
 
-				if ((bDisplayDev && ContentBrowserUtils::IsItemDeveloperContent(NewSelectedItem)) ||
-					(bDisplayEngine && ContentBrowserUtils::IsItemEngineContent(NewSelectedItem)) ||
-					(bDisplayPlugins && ContentBrowserUtils::IsItemPluginContent(NewSelectedItem)) ||
-					(bDisplayL10N && ContentBrowserUtils::IsItemLocalizedContent(NewSelectedItem))
-					)
-				{
-					// Refresh the contents
-					OnItemSelectionChanged.ExecuteIfBound(NewSelectedItem, ESelectInfo::Direct);
-				}
-			}
+	if (!DefaultPath.IsEmpty() && InternalPathPassesBlockLists(DefaultPath))
+	{
+		UContentBrowserDataSubsystem* ContentBrowserData = IContentBrowserDataModule::Get().GetSubsystem();
+		const FName VirtualPath = ContentBrowserData->ConvertInternalPathToVirtual(*DefaultPath);
+		FName InternalPath;
+		if (ContentBrowserData->TryConvertVirtualPath(VirtualPath, InternalPath) != EContentBrowserPathType::Internal)
+		{
+			InternalPath = FName();
+		}
+		if (bCreateDefaultPath && !TreeData->FindTreeItem(VirtualPath))
+		{
+			const FString DefaultPathLeafName = FPaths::GetPathLeaf(VirtualPath.ToString());
+			TreeData->AddFolderItem(FContentBrowserItemData(nullptr,
+				EContentBrowserItemFlags::Type_Folder,
+				VirtualPath,
+				*DefaultPathLeafName,
+				FText(),
+				nullptr,
+				InternalPath));
 		}
+
+		SetSelectedPaths({ VirtualPath.ToString() });
 	}
+
+	DefaultPath.Empty();
+}
+
+bool SPathView::IsPopulationRequested() const
+{
+	return bPopulationRequested || bPopulationForSettingChangeRequested || bInitialPopulationRequested;
 }
 
 TArray<FName> SPathView::GetDefaultPathsToSelect() const
@@ -2781,6 +2840,8 @@ void SFavoritePathView::Populate(const bool bIsRefreshingFilter)
 	TreeData->PopulateWithFavorites(CreateCompiledFolderFilter());
 	TreeData->SortRootItems();
 	TreeData->FilterFullFolderTree();
+
+	SharedPostPopulate();
 }
 
 void SFavoritePathView::SaveSettings(const FString& IniFilename, const FString& IniSection, const FString& SettingsString) const
diff --git a/Engine/Source/Editor/ContentBrowser/Private/SPathView.h b/Engine/Source/Editor/ContentBrowser/Private/SPathView.h
--- a/Engine/Source/Editor/ContentBrowser/Private/SPathView.h
+++ b/Engine/Source/Editor/ContentBrowser/Private/SPathView.h
@@ -258,7 +258,15 @@ public:
 	 */
 	bool DisablePluginPathFiltersThatHideItems(TConstArrayView<FContentBrowserItem> Items);
 
-	/** Populates the tree with all folders that are not filtered out */
+	/**
+	 * Request a population of the path view the next time the widget is ticked.
+	 */
+	void RequestPopulation(FSimpleMulticastDelegate::FDelegate&& AfterPopulation = FSimpleMulticastDelegate::FDelegate());
+
+	/**
+	 * Populates the tree with all folders that are not filtered out
+	 * Note: This can be very long in large project. prefer calling request population to avoid multiple population in one frame.
+	 */
 	virtual void Populate(const bool bIsRefreshingFilter = false);
 
 	/** Sets an alternate tree title*/
@@ -348,6 +356,9 @@ protected:
 
 	virtual void ConfigureTreeView(STreeView<TSharedPtr<FTreeItem>>::FArguments& InArgs);
 
+	/** Should be call after a populate */
+	void SharedPostPopulate();
+
 private:
 	/** Selects the given path only if it exists. Returns true if selected. */
 	bool ExplicitlyAddPathToSelection(const FName Path);
@@ -373,12 +384,18 @@ private:
 	/** Delegate called when an editor setting is changed */
 	void HandleSettingChanged(FName PropertyName);
 
+	/** Handle the logic that is done after population when requested by a setting change */
+	void PostPopulateHandleSettingChanged(bool bHadSelectedPath);
+
+	/** Handle the logic that is done after the initial population of the view */
+	void PostInitialPopulation();
+
+	/** Is the view waiting for a population of its tree */
+	bool IsPopulationRequested() const;
+
 	/** One-off active timer to focus the widget post-construct */
 	EActiveTimerReturnType SetFocusPostConstruct(double InCurrentTime, float InDeltaTime);
 
-	/** One-off active timer to repopulate the path view */
-	EActiveTimerReturnType TriggerRepopulate(double InCurrentTime, float InDeltaTime);
-
 	/** Sets the active state of a filter. */
 	void SetPluginPathFilterActive(const TSharedRef<FContentBrowserPluginFilter>& Filter, bool bActive);
 
@@ -466,6 +483,8 @@ protected:
 		FScopedSelectionChangedEvent(const TSharedRef<SPathView>& InPathView, const bool InShouldEmitEvent = true);
 		~FScopedSelectionChangedEvent();
 
+		bool IsInitialSelectionEmpty() const;
+
 	private:
 		TSet<FName> GetSelectionSet() const;
 
@@ -518,6 +537,9 @@ protected:
 	/** The config instance to use. */
 	FName OwningContentBrowserName;
 
+	/** Event(s) called to run any queued logic after a population of the widget */
+	FSimpleMulticastDelegate OnNextPopulate;
+
 private:
 	/** Used to track if the list of last expanded path should be updated */
 	bool bLastExpandedPathsDirty = false;
@@ -554,6 +576,18 @@ private:
 
 	/** If true, plugin content is always shown */
 	bool bForceShowPluginContent : 1;
+	
+	/** Request a population of the path view */
+	bool bPopulationRequested : 1;
+
+	/** Request the initial population of the path view */
+	bool bInitialPopulationRequested : 1;
+
+	/** Request a population of the path view but because of a setting change */
+	bool bPopulationForSettingChangeRequested : 1;
+
+	bool bCreateDefaultPath : 1;
+	FString DefaultPath;
 
 	/** The title of this path view */
 	FText TreeTitle;
@@ -607,7 +641,7 @@ public:
 
 	DECLARE_DELEGATE_OneParam(FOnFolderFavoriteAdd, const TArray<FString>& /*FoldersToAdd*/)
 	void SetOnFolderFavoriteAdd(const FOnFolderFavoriteAdd& InOnFolderFavoriteAdd);
-	
+
 private:
 	virtual TSharedRef<ITableRow> GenerateTreeRow(TSharedPtr<FTreeItem> TreeItem, const TSharedRef<STableViewBase>& OwnerTable) override;
 
