Turning off the delayed pathview population as it can cause some folder navigation issue if the widget is not visible in the content browser.

diff --git a/Engine/Source/Editor/ContentBrowser/Private/SPathView.cpp b/Engine/Source/Editor/ContentBrowser/Private/SPathView.cpp
--- a/Engine/Source/Editor/ContentBrowser/Private/SPathView.cpp
+++ b/Engine/Source/Editor/ContentBrowser/Private/SPathView.cpp
@@ -1166,6 +1166,8 @@ void SPathView::Construct( const FArguments& InArgs )
 	DefaultPath = InArgs._DefaultPath;
 	bCreateDefaultPath = InArgs._CreateDefaultPath;
 
+	// Temporary turn off the delayed population to investigate navigation issue when the view is not populated
+	Populate();
 }
 
 void SPathView::ConfigureTreeView(STreeView<TSharedPtr<FTreeItem>>::FArguments& InArgs)
@@ -1255,7 +1257,11 @@ void SPathView::PopulatePathViewFiltersMenu(UToolMenu* Menu)
 void SPathView::PluginPathFilterClicked(TSharedRef<FContentBrowserPluginFilter> Filter)
 {
 	SetPluginPathFilterActive(Filter, !IsPluginPathFilterInUse(Filter));
+	
 	bPopulationRequested = true;
+
+	// Temporary turn off the delayed population to investigate navigation issue when the view is not populated
+	Populate();
 }
 
 bool SPathView::IsPluginPathFilterChecked(TSharedRef<FContentBrowserPluginFilter> Filter) const
@@ -1291,6 +1297,8 @@ void SPathView::ResetPluginPathFilters()
 	}
 
 	bPopulationRequested = true;
+	// Temporary turn off the delayed population to investigate navigation issue when the view is not populated
+	Populate();
 }
 
 bool SPathView::DisablePluginPathFiltersThatHideItems(TConstArrayView<FContentBrowserItem> Items)
@@ -2249,6 +2257,9 @@ void SPathView::RequestPopulation(FSimpleMulticastDelegate::FDelegate&& AfterPop
 	{
 		OnNextPopulate.Add(MoveTemp(AfterPopulation));
 	}
+
+	// Temporary turn off the delayed population to investigate navigation issue when the view is not populated
+	Populate();
 }
 
 
@@ -2483,6 +2494,9 @@ void SPathView::HandleItemDataUpdated(TArrayView<const FContentBrowserItemDataUp
 void SPathView::HandleItemDataRefreshed()
 {
 	bPopulationRequested = true;
+
+	// Temporary turn off the delayed population to investigate navigation issue when the view is not populated
+	Populate();
 }
 
 void SPathView::HandleItemDataDiscoveryComplete()
@@ -2513,6 +2527,9 @@ void SPathView::HandleSettingChanged(FName PropertyName)
 
 		// Update our path view so that it can include/exclude the permitted folder
 		bPopulationForSettingChangeRequested = true;
+
+		// Temporary turn off the delayed population to investigate navigation issue when the view is not populated
+		Populate();
 	}
 }
 
