Do not propagate the selection change when poppulating the view for the first time. This was the previous behavior before the deleyed population.

diff --git a/Engine/Source/Editor/ContentBrowser/Private/SPathView.cpp b/Engine/Source/Editor/ContentBrowser/Private/SPathView.cpp
--- a/Engine/Source/Editor/ContentBrowser/Private/SPathView.cpp
+++ b/Engine/Source/Editor/ContentBrowser/Private/SPathView.cpp
@@ -2271,9 +2271,9 @@ void SPathView::Populate(const bool bIsRefreshingFilter)
 	const bool bFilteringByText = !TreeData->GetFolderPathTextFilter().GetRawFilterText().IsEmpty();
 
 	// Batch the selection changed event
-	// Only emit events when the user isn't filtering, as the selection may be artificially limited by the filter
+	// Only emit events when the user isn't filtering or doing the first population.
 	FScopedSelectionChangedEvent ScopedSelectionChangedEvent(SharedThis(this),
-		!bFilteringByText && !bIsRefreshingFilter);
+		!bFilteringByText && !bIsRefreshingFilter && !bInitialPopulationRequested);
 	TreeData->PopulateFullFolderTree(CreateCompiledFolderFilter());
 	TreeData->FilterFullFolderTree();
 	TreeData->SortRootItems();
