Added a trace for SAssetView::SortList

#rnx

diff --git a/Engine/Source/Editor/ContentBrowser/Private/SAssetView.cpp b/Engine/Source/Editor/ContentBrowser/Private/SAssetView.cpp
--- a/Engine/Source/Editor/ContentBrowser/Private/SAssetView.cpp
+++ b/Engine/Source/Editor/ContentBrowser/Private/SAssetView.cpp
@@ -4196,6 +4196,8 @@ bool SAssetView::PassesCurrentFrontendFilter(const FContentBrowserItem& Item) co
 
 void SAssetView::SortList(bool bSyncToSelection)
 {
+	TRACE_CPUPROFILER_EVENT_SCOPE(SAssetView::SortList);
+
 	if ( !IsRenamingAsset() )
 	{
 		SortManager->SortList(FilteredAssetItems, MajorityAssetType, CustomColumns);
