Fixing crash when clicking the Standalone Game button with an unsaved level.

diff --git a/Engine/Source/Editor/ContentBrowser/Private/SAssetDialog.cpp b/Engine/Source/Editor/ContentBrowser/Private/SAssetDialog.cpp
--- a/Engine/Source/Editor/ContentBrowser/Private/SAssetDialog.cpp
+++ b/Engine/Source/Editor/ContentBrowser/Private/SAssetDialog.cpp
@@ -961,6 +961,11 @@ void SAssetDialog::OnApplyHistoryData(const FHistoryData& History)
 
 void SAssetDialog::OnUpdateHistoryData(FHistoryData& HistoryData) const
 {
+	if (!AssetPicker)
+	{
+		return;
+	}
+
 	const FAssetViewContentSources& ContentSources = AssetPicker->GetAssetView()->GetContentSources();
 	const TArray<FContentBrowserItem> SelectedItems = AssetPicker->GetAssetView()->GetSelectedItems();
 
