PR #13669: Fix for Content Browser search crash caused by AssetTextFilter

diff --git a/Engine/Source/Editor/ContentBrowser/Private/AssetTextFilter.cpp b/Engine/Source/Editor/ContentBrowser/Private/AssetTextFilter.cpp
--- a/Engine/Source/Editor/ContentBrowser/Private/AssetTextFilter.cpp
+++ b/Engine/Source/Editor/ContentBrowser/Private/AssetTextFilter.cpp
@@ -400,6 +400,21 @@ bool FCompiledAssetTextFilter::PassesFilter(FAssetFilterType InItem)
 		}
 	}
 
+
+    // Convert entire text buffer to uppercase
+    // Assumes that for a TCHAR buffer this won't change number of code points 
+	TextBuffer.ToUpperInline();
+	
+	const TCHAR* Cursor = *TextBuffer;
+	AssetDisplayName = FStringView(Cursor, DisplayNameLen);
+	Cursor += DisplayNameLen;
+	AssetFullPath = FStringView(Cursor, AssetPathLen);
+	Cursor += AssetPathLen;
+	AssetExportTextPath  = FStringView(Cursor, ExportTextPathLen);
+	AssetExportTextPath.TrimStartAndEndInline(); // Backends try to separate export text paths with newlines 
+	Cursor += ExportTextPathLen;
+
+	// Executed last because we pass *this to DynamicQuery and it would try to access StringView before initialization.
 	if ((CollectionContainers.IsValid() || ReferencedDynamicCollections.IsValid()) && bIncludeCollectionNames)
 	{
 		FSoftObjectPath ItemCollectionId;
@@ -434,19 +449,6 @@ bool FCompiledAssetTextFilter::PassesFilter(FAssetFilterType InItem)
 			}
 		}
 	}
-    
-    // Convert entire text buffer to uppercase
-    // Assumes that for a TCHAR buffer this won't change number of code points 
-	TextBuffer.ToUpperInline();
-	
-	const TCHAR* Cursor = *TextBuffer;
-	AssetDisplayName = FStringView(Cursor, DisplayNameLen);
-	Cursor += DisplayNameLen;
-	AssetFullPath = FStringView(Cursor, AssetPathLen);
-	Cursor += AssetPathLen;
-	AssetExportTextPath  = FStringView(Cursor, ExportTextPathLen);
-	AssetExportTextPath.TrimStartAndEndInline(); // Backends try to separate export text paths with newlines 
-	Cursor += ExportTextPathLen;
 
 	const bool bMatched = Evaluator->TestTextFilter(*this);
 	AssetPtr = nullptr;
