Fixed the Private AssetAcessSpecifier icon so its a closed padlock

#rnx

diff --git a/Engine/Source/Editor/ContentBrowser/Private/AssetViewTypes.cpp b/Engine/Source/Editor/ContentBrowser/Private/AssetViewTypes.cpp
--- a/Engine/Source/Editor/ContentBrowser/Private/AssetViewTypes.cpp
+++ b/Engine/Source/Editor/ContentBrowser/Private/AssetViewTypes.cpp
@@ -287,7 +287,8 @@ FName FAssetViewItem::GetPublicAssetAccessSpecifierIconStyleName(const TCHAR* In
 
 FName FAssetViewItem::GetPrivateAssetAccessSpecifierIconStyleName(const TCHAR* InSizeSuffix)
 {
-	return FName(*FString::Printf(TEXT("ContentBrowser.Private%s"), InSizeSuffix));
+	static const FName IconName("Icons.Lock");
+	return IconName;
 }
 
 FName FAssetViewItem::GetEpicInternalAssetAccessSpecifierIconStyleName(const TCHAR* InSizeSuffix)
