Updated the Content Browser navigation bar so clicking in the empty space edits a friendly user facing path (such as a Verse path or package path) instead of a virtual path (unless the current path is pure virtual).
Updated the alias and class data source to no longer return their internal path as a package path via Legacy_TryGetPackagePath since they are not package paths.
Updated the Content Browser copy folder path behavior to only copy package paths.  The effective change of this is to not copy internal paths of class folders (which is a format native only to the class data source internals).

diff --git a/Engine/Plugins/Editor/ContentBrowser/ContentBrowserAliasDataSource/Source/ContentBrowserAliasDataSource/Private/ContentBrowserAliasDataSource.cpp b/Engine/Plugins/Editor/ContentBrowser/ContentBrowserAliasDataSource/Source/ContentBrowserAliasDataSource/Private/ContentBrowserAliasDataSource.cpp
--- a/Engine/Plugins/Editor/ContentBrowser/ContentBrowserAliasDataSource/Source/ContentBrowserAliasDataSource/Private/ContentBrowserAliasDataSource.cpp
+++ b/Engine/Plugins/Editor/ContentBrowser/ContentBrowserAliasDataSource/Source/ContentBrowserAliasDataSource/Private/ContentBrowserAliasDataSource.cpp
@@ -907,7 +907,7 @@ bool UContentBrowserAliasDataSource::GetItemAttributes(const FContentBrowserItem
 
 bool UContentBrowserAliasDataSource::GetItemPhysicalPath(const FContentBrowserItemData& InItem, FString& OutDiskPath)
 {
-	if (InItem.GetItemType() == EContentBrowserItemFlags::Category_Asset)
+	if (InItem.GetItemType() == EContentBrowserItemFlags::Type_File)
 	{
 		return ContentBrowserAssetData::GetItemPhysicalPath(this, InItem, OutDiskPath);
 	}
@@ -1022,17 +1022,29 @@ bool UContentBrowserAliasDataSource::BulkPrivatizeItems(TArrayView<const FConten
 
 bool UContentBrowserAliasDataSource::AppendItemReference(const FContentBrowserItemData& InItem, FString& InOutStr)
 {
-	return ContentBrowserAssetData::AppendItemReference(AssetRegistry, this, InItem, InOutStr);
+	if (InItem.GetItemType() == EContentBrowserItemFlags::Type_File)
+	{
+		return ContentBrowserAssetData::AppendItemReference(AssetRegistry, this, InItem, InOutStr);
+	}
+	return false;
 }
 
 bool UContentBrowserAliasDataSource::AppendItemObjectPath(const FContentBrowserItemData& InItem, FString& InOutStr)
 {
-	return ContentBrowserAssetData::AppendItemObjectPath(AssetRegistry, this, InItem, InOutStr);
+	if (InItem.GetItemType() == EContentBrowserItemFlags::Type_File)
+	{
+		return ContentBrowserAssetData::AppendItemObjectPath(AssetRegistry, this, InItem, InOutStr);
+	}
+	return false;
 }
 
 bool UContentBrowserAliasDataSource::AppendItemPackageName(const FContentBrowserItemData& InItem, FString& InOutStr)
 {
-	return ContentBrowserAssetData::AppendItemPackageName(AssetRegistry, this, InItem, InOutStr);
+	if (InItem.GetItemType() == EContentBrowserItemFlags::Type_File)
+	{
+		return ContentBrowserAssetData::AppendItemPackageName(AssetRegistry, this, InItem, InOutStr);
+	}
+	return false;
 }
 
 bool UContentBrowserAliasDataSource::UpdateThumbnail(const FContentBrowserItemData& InItem, FAssetThumbnail& InThumbnail)
@@ -1062,11 +1074,6 @@ bool UContentBrowserAliasDataSource::CanModifyItemAssetAccessSpecifier(const FCo
 
 bool UContentBrowserAliasDataSource::Legacy_TryGetPackagePath(const FContentBrowserItemData& InItem, FName& OutPackagePath)
 {
-	if (TSharedPtr<const FContentBrowserAssetFolderItemDataPayload> FolderPayload = ContentBrowserAssetData::GetAssetFolderItemPayload(this, InItem))
-	{
-		OutPackagePath = FolderPayload->GetInternalPath();
-		return true;
-	}
 	return false;
 }
 
@@ -1082,8 +1089,7 @@ bool UContentBrowserAliasDataSource::Legacy_TryGetAssetData(const FContentBrowse
 
 bool UContentBrowserAliasDataSource::Legacy_TryConvertPackagePathToVirtualPath(const FName InPackagePath, FName& OutPath)
 {
-	return PathTree.PathExists(InPackagePath) // Ignore unknown content paths
-		&& TryConvertInternalPathToVirtual(InPackagePath, OutPath);
+	return false;
 }
 
 bool UContentBrowserAliasDataSource::Legacy_TryConvertAssetDataToVirtualPath(const FAssetData& InAssetData, const bool InUseFolderPaths, FName& OutPath)
diff --git a/Engine/Plugins/Editor/ContentBrowser/ContentBrowserClassDataSource/Source/ContentBrowserClassDataSource/Private/ContentBrowserClassDataSource.cpp b/Engine/Plugins/Editor/ContentBrowser/ContentBrowserClassDataSource/Source/ContentBrowserClassDataSource/Private/ContentBrowserClassDataSource.cpp
--- a/Engine/Plugins/Editor/ContentBrowser/ContentBrowserClassDataSource/Source/ContentBrowserClassDataSource/Private/ContentBrowserClassDataSource.cpp
+++ b/Engine/Plugins/Editor/ContentBrowser/ContentBrowserClassDataSource/Source/ContentBrowserClassDataSource/Private/ContentBrowserClassDataSource.cpp
@@ -596,11 +596,6 @@ bool UContentBrowserClassDataSource::TryGetCollectionId(const FContentBrowserIte
 
 bool UContentBrowserClassDataSource::Legacy_TryGetPackagePath(const FContentBrowserItemData& InItem, FName& OutPackagePath)
 {
-	if (TSharedPtr<const FContentBrowserClassFolderItemDataPayload> FolderPayload = GetClassFolderItemPayload(InItem))
-	{
-		OutPackagePath = FolderPayload->GetInternalPath();
-		return true;
-	}
 	return false;
 }
 
@@ -616,8 +611,7 @@ bool UContentBrowserClassDataSource::Legacy_TryGetAssetData(const FContentBrowse
 
 bool UContentBrowserClassDataSource::Legacy_TryConvertPackagePathToVirtualPath(const FName InPackagePath, FName& OutPath)
 {
-	return IsKnownClassPath(InPackagePath) // Ignore non-class paths
-		&& TryConvertInternalPathToVirtual(InPackagePath, OutPath);
+	return false;
 }
 
 bool UContentBrowserClassDataSource::Legacy_TryConvertAssetDataToVirtualPath(const FAssetData& InAssetData, const bool InUseFolderPaths, FName& OutPath)
diff --git a/Engine/Source/Editor/ContentBrowser/Private/ContentBrowserUtils.cpp b/Engine/Source/Editor/ContentBrowser/Private/ContentBrowserUtils.cpp
--- a/Engine/Source/Editor/ContentBrowser/Private/ContentBrowserUtils.cpp
+++ b/Engine/Source/Editor/ContentBrowser/Private/ContentBrowserUtils.cpp
@@ -4,6 +4,7 @@
 #include "ContentBrowserUtils.h"
 
 #include "Algo/MinElement.h"
+#include "Algo/NoneOf.h"
 #include "Algo/Unique.h"
 #include "AssetRegistry/ARFilter.h"
 #include "AssetRegistry/AssetData.h"
@@ -479,10 +480,15 @@ FString ContentBrowserUtils::GetFolderReferencesText(const TArray<FContentBrowse
 	{
 		if (ensure(Item.IsFolder()))
 		{
-			FName InternalPath = Item.GetInternalPath();
-			if (!InternalPath.IsNone())
+			FName PackagePath;
+			if (Item.Legacy_TryGetPackagePath(PackagePath))
 			{
-				Result << InternalPath << LINE_TERMINATOR;
+				if (Result.Len() > 0)
+				{
+					Result << LINE_TERMINATOR;
+				}
+
+				Result << PackagePath;
 			}
 		}
 	}
@@ -1367,7 +1373,13 @@ TArray<FString> ContentBrowserUtils::GetChildPathsFromUserProvidedPath(FStringVi
 	UContentBrowserDataSubsystem* const ContentBrowserData = IContentBrowserDataModule::Get().GetSubsystem();
 
 	TArray<FContentBrowserItem> Items;
-	EContentBrowserPathType PathType = EContentBrowserPathType::None;
+	enum class EUserProvidedPathType
+	{
+		None,
+		VirtualPath,
+		PackagePath,
+		VersePath,
+	} PathType = EUserProvidedPathType::None;
 	{
 		FName RequestedPath(RequestedPathView);
 
@@ -1376,7 +1388,7 @@ TArray<FString> ContentBrowserUtils::GetChildPathsFromUserProvidedPath(FStringVi
 		if (Item.IsValid())
 		{
 			Items.Add(MoveTemp(Item));
-			PathType = EContentBrowserPathType::Virtual;
+			PathType = EUserProvidedPathType::VirtualPath;
 		}
 		else
 		{
@@ -1385,7 +1397,8 @@ TArray<FString> ContentBrowserUtils::GetChildPathsFromUserProvidedPath(FStringVi
 
 			if (!Items.IsEmpty())
 			{
-				PathType = EContentBrowserPathType::Internal;
+				// At the time of authoring the only possible option that GetItemsAtUserProvidedPath could have parsed is a package path.
+				PathType = EUserProvidedPathType::PackagePath;
 			}
 			else // Lastly, it could be a Verse path.
 			{
@@ -1398,7 +1411,7 @@ TArray<FString> ContentBrowserUtils::GetChildPathsFromUserProvidedPath(FStringVi
 
 				if (VersePath.IsValid())
 				{
-					IPluginManager::Get().EnumerateVersePathMountPoints(VersePath, [ContentBrowserData, &Items](FStringView MountedAssetPath, FStringView LeafPath)
+					IPluginManager::Get().EnumerateVersePathMountPoints(VersePath, [ContentBrowserData, &Items, &PathType](FStringView MountedAssetPath, FStringView LeafPath)
 					{
 						FNameBuilder NameBuilder;
 						NameBuilder.Append(MountedAssetPath);
@@ -1420,6 +1433,7 @@ TArray<FString> ContentBrowserUtils::GetChildPathsFromUserProvidedPath(FStringVi
 							if (Item.IsValid())
 							{
 								Items.Add(MoveTemp(Item));
+								PathType = EUserProvidedPathType::VersePath;
 							}
 						}
 						return true;
@@ -1486,26 +1500,46 @@ TArray<FString> ContentBrowserUtils::GetChildPathsFromUserProvidedPath(FStringVi
 				continue;
 			}
 
-			// Get the subitem path matching the requested path format.
-			const FName Path = PathType == EContentBrowserPathType::Virtual ? SubItem.GetVirtualPath() : SubItem.GetInternalPath();
-			const FNameBuilder PathBuilder(Path);
-
-			FStringView PathView;
+			FNameBuilder PathBuilder;
 			UE::Core::FVersePath VersePath;
-			if (PathType == EContentBrowserPathType::None)
+			FStringView PathView;
+
+			// Get the subitem path matching the requested path format.
+			switch (PathType)
 			{
-				// If not a virtual or internal path, it must be a Verse path.
-				// At the time of authoring the only possible option that GetItemsAtUserProvidedPath could have parsed is a package path.
-				VersePath = FPackageName::LongPackagePathToVersePath(PathBuilder);
-				if (!VersePath.IsValid())
+			case EUserProvidedPathType::VirtualPath:
+				SubItem.GetVirtualPath().AppendString(PathBuilder);
+				PathView = PathBuilder;
+				break;
+			case EUserProvidedPathType::PackagePath:
+			{
+				FName PackagePath;
+				if (SubItem.Legacy_TryGetPackagePath(PackagePath))
 				{
-					continue;
+					PackagePath.AppendString(PathBuilder);
+					PathView = PathBuilder;
+					break;
 				}
-				PathView = VersePath.AsStringView();
+				continue;
 			}
-			else
+			case EUserProvidedPathType::VersePath:
 			{
-				PathView = PathBuilder.ToView();
+				FName PackagePath;
+				if (SubItem.Legacy_TryGetPackagePath(PackagePath))
+				{
+					PackagePath.AppendString(PathBuilder);
+					VersePath = FPackageName::LongPackagePathToVersePath(PathBuilder);
+					if (VersePath.IsValid())
+					{
+						PathView = VersePath.AsStringView();
+						break;
+					}
+				}
+				continue;
+			}
+			default:
+				checkNoEntry();
+				break;
 			}
 
 			// Verify that the new path is indeed syntactically a subpath.
@@ -1524,6 +1558,93 @@ TArray<FString> ContentBrowserUtils::GetChildPathsFromUserProvidedPath(FStringVi
 	return Results;
 }
 
+namespace ContentBrowserUtils
+{
+	static bool bEditNavigationBarPathAsVirtualPath = false;
+	static FAutoConsoleVariableRef CVarShowContentSourcesBar(
+		TEXT("ContentBrowser.UI.EditNavigationBarPathAsVirtualPath"),
+		bEditNavigationBarPathAsVirtualPath,
+		TEXT("When editing text in the navigation bar, edit the text as the virtual path if set, otherwise edit it as its native representation if available, such as a package path."));
+}
+
+FText ContentBrowserUtils::GetUserFacingPathFromVirtualPath(FName VirtualPath)
+{
+	UContentBrowserDataSubsystem* ContentBrowserData = IContentBrowserDataModule::Get().GetSubsystem();
+	FContentBrowserItem Item = ContentBrowserData->GetItemAtPath(VirtualPath, EContentBrowserItemTypeFilter::IncludeFolders);
+	if (Item.IsValid())
+	{
+		if (!bEditNavigationBarPathAsVirtualPath)
+		{
+			FName PackagePath;
+			if (Item.Legacy_TryGetPackagePath(PackagePath))
+			{
+				if (IAssetTools::Get().ShowingContentVersePath())
+				{
+					// Check to see if the mount point maps to a unique Verse path and return that.
+					FNameBuilder PackagePathBuilder(PackagePath);
+					FStringView PackagePathView(PackagePathBuilder);
+
+					// If the mount point is invalid, we can't create a Verse path from it.
+					bool bHadClassesPrefix = false;
+					const FStringView MountPointName = FPathViews::GetMountPointNameFromPath(PackagePathView, &bHadClassesPrefix);
+					if (!MountPointName.IsEmpty() && !bHadClassesPrefix)
+					{
+						IPluginManager& PluginManager = IPluginManager::Get();
+
+						// If the package path isn't mounted under a plugin, the package path doesn't have a Verse path.
+						TSharedPtr<IPlugin> Plugin = PluginManager.FindPlugin(MountPointName);
+						if (Plugin)
+						{
+							// If the plugin doesn't have a root Verse path, the package path doesn't have a Verse path.
+							UE::Core::FVersePath PluginVersePath;
+							if (UE::Core::FVersePath::TryMake(PluginVersePath, Plugin->GetVersePath()))
+							{
+								// Make sure there are no other plugins with the same Verse path.
+								if (Algo::NoneOf(PluginManager.GetEnabledPluginsWithContent(), [&Plugin, &PluginVersePath](const TSharedPtr<IPlugin>& OtherPlugin)
+								{
+									if (Plugin == OtherPlugin)
+									{
+										return false;
+									}
+
+									UE::Core::FVersePath OtherPluginVersePath;
+									if (!UE::Core::FVersePath::TryMake(OtherPluginVersePath, OtherPlugin->GetVersePath()))
+									{
+										return false;
+									}
+
+									return PluginVersePath == OtherPluginVersePath;
+								}))
+								{
+									const int32 CharCount = MountPointName.Len() + 1;
+
+									// If the package path is the mount point, return the mount point's Verse path.
+									if (PackagePathView.Len() == CharCount)
+									{
+										return MoveTemp(PluginVersePath).AsText();
+									}
+
+									// Otherwise append the leaf path.
+									UE::Core::FVersePath VersePath;
+									if (UE::Core::FVersePath::TryMake(VersePath, FPaths::Combine(PluginVersePath.AsStringView(), PackagePathView.RightChop(CharCount))))
+									{
+										return MoveTemp(VersePath).AsText();
+									}
+								}
+							}
+						}
+					}
+				}
+
+				// If not showing Verse paths, or the Verse path is not unique, prefer the package path.
+				return FText::FromName(PackagePath);
+			}
+		}
+		return FText::FromName(VirtualPath);
+	}
+	return FText::GetEmpty();
+}
+
 FString ContentBrowserUtils::FormatCollectionCrumbData(const ICollectionContainer& CollectionContainer, const FCollectionNameType& Collection)
 {
 	return FString::Printf(
diff --git a/Engine/Source/Editor/ContentBrowser/Private/ContentBrowserUtils.h b/Engine/Source/Editor/ContentBrowser/Private/ContentBrowserUtils.h
--- a/Engine/Source/Editor/ContentBrowser/Private/ContentBrowserUtils.h
+++ b/Engine/Source/Editor/ContentBrowser/Private/ContentBrowserUtils.h
@@ -212,6 +212,9 @@ namespace ContentBrowserUtils
 	/** Finds the child paths that are children of the supplied path, while obeying the supplied filters. */
 	TArray<FString> GetChildPathsFromUserProvidedPath(FStringView RequestedPathView, EContentBrowserItemCategoryFilter ItemCategoryFilter, EContentBrowserItemAttributeFilter ItemAttributeFilter, const FName& ConfigInstanceName, const SPathView& PathViewForFiltering);
 
+	/** Gets the actual path the user should edit when they start editing a virtual path in the navigation bar */
+	FText GetUserFacingPathFromVirtualPath(FName VirtualPath);
+
 	FString FormatCollectionCrumbData(const ICollectionContainer& CollectionContainer, const FCollectionNameType& Collection);
 	
 	/** Updates a navigation bar based on the current state of an asset view and path view */
diff --git a/Engine/Source/Editor/ContentBrowser/Private/PathContextMenu.cpp b/Engine/Source/Editor/ContentBrowser/Private/PathContextMenu.cpp
--- a/Engine/Source/Editor/ContentBrowser/Private/PathContextMenu.cpp
+++ b/Engine/Source/Editor/ContentBrowser/Private/PathContextMenu.cpp
@@ -164,8 +164,7 @@ void FPathContextMenu::MakePathViewContextMenu(UToolMenu* Menu)
 					);
 			}
 
-			// Assume paths with an on-disk representation also have an internal path to copy
-			if (!Context->bNoFolderOnDisk)
+			if (!Context->GetSelectedPackagePaths().IsEmpty())
 			{
 				Section.AddMenuEntry(
 					"CopyPath",
diff --git a/Engine/Source/Editor/ContentBrowser/Private/SAssetDialog.cpp b/Engine/Source/Editor/ContentBrowser/Private/SAssetDialog.cpp
--- a/Engine/Source/Editor/ContentBrowser/Private/SAssetDialog.cpp
+++ b/Engine/Source/Editor/ContentBrowser/Private/SAssetDialog.cpp
@@ -275,7 +275,7 @@ void SAssetDialog::Construct(const FArguments& InArgs, const FSharedAssetDialogC
 					SAssignNew(NavigationBar, SNavigationBar)
 					.OnPathClicked(this, &SAssetDialog::SetCurrentlySelectedPath, EContentBrowserPathType::Virtual)
 					.OnNavigateToPath(this, &SAssetDialog::OnPathTextEdited)
-					.OnCanEditPathAsText(this, &SAssetDialog::OnCanEditPathAsText)
+					.OnGetEditPathAsText(this, &SAssetDialog::OnGetEditPathAsText)
 					.OnCompletePrefix(this, &SAssetDialog::OnCompletePathPrefix)
 					.GetPathMenuContent(this, &SAssetDialog::OnGetCrumbDelimiterContent)
 					.GetComboOptions(this, &SAssetDialog::GetRecentPaths)
@@ -1018,14 +1018,18 @@ FText SAssetDialog::GetNameErrorLabelText() const
 	return FText::GetEmpty();
 }
 
-bool SAssetDialog::OnCanEditPathAsText(const FString& Text) const
+FText SAssetDialog::OnGetEditPathAsText(const FString& Path) const
 {
 	const FAssetViewContentSources& ContentSources = AssetPicker->GetAssetView()->GetContentSources();
 	if (ContentSources.HasCollections())
 	{
-		return false;
+		return FText::GetEmpty();
+	}
+	else if (ContentSources.HasVirtualPaths())
+	{
+		return ContentBrowserUtils::GetUserFacingPathFromVirtualPath(FName(Path));
 	}
-	return true;
+	return FText::GetEmpty();
 }
 
 void SAssetDialog::OnPathTextEdited(const FString& NewPath)
diff --git a/Engine/Source/Editor/ContentBrowser/Private/SAssetDialog.h b/Engine/Source/Editor/ContentBrowser/Private/SAssetDialog.h
--- a/Engine/Source/Editor/ContentBrowser/Private/SAssetDialog.h
+++ b/Engine/Source/Editor/ContentBrowser/Private/SAssetDialog.h
@@ -110,8 +110,8 @@ private:
 
 	
 	// --- Navigation Bar Behavior ---
-	/** Whether the navigation bar can be edited */
-	bool OnCanEditPathAsText(const FString& Text) const;
+	/** Returns how the currently shown location should be presented to the user as text when beginning to edit the path */
+	FText OnGetEditPathAsText(const FString& Path) const;
 
 	/** Respond to edited text from the navigation bar */
 	void OnPathTextEdited(const FString& NewPath);
diff --git a/Engine/Source/Editor/ContentBrowser/Private/SAssetView.cpp b/Engine/Source/Editor/ContentBrowser/Private/SAssetView.cpp
--- a/Engine/Source/Editor/ContentBrowser/Private/SAssetView.cpp
+++ b/Engine/Source/Editor/ContentBrowser/Private/SAssetView.cpp
@@ -3754,13 +3754,14 @@ void SAssetView::ExecuteCopy(EAssetViewCopyType InCopyType) const
 
 void SAssetView::ExecuteCopyFolders(const TArray<FContentBrowserItem>& InSelectedFolders, FString& OutClipboardText) const
 {
-	if (InSelectedFolders.Num() > 0)
+	const FString FolderReferencesText = ContentBrowserUtils::GetFolderReferencesText(InSelectedFolders);
+	if (!FolderReferencesText.IsEmpty())
 	{
 		if (!OutClipboardText.IsEmpty())
 		{
 			OutClipboardText += LINE_TERMINATOR;
 		}
-		OutClipboardText += ContentBrowserUtils::GetFolderReferencesText(InSelectedFolders);
+		OutClipboardText += FolderReferencesText;
 	}
 }
 
diff --git a/Engine/Source/Editor/ContentBrowser/Private/SContentBrowser.cpp b/Engine/Source/Editor/ContentBrowser/Private/SContentBrowser.cpp
--- a/Engine/Source/Editor/ContentBrowser/Private/SContentBrowser.cpp
+++ b/Engine/Source/Editor/ContentBrowser/Private/SContentBrowser.cpp
@@ -997,7 +997,7 @@ TSharedRef<SWidget> SContentBrowser::CreateAssetView(const FContentBrowserConfig
 						.GetComboOptions(this, &SContentBrowser::GetRecentPaths)
 						.OnNavigateToPath(this, &SContentBrowser::OnNavigateToPath)
 						.OnCompletePrefix(this, &SContentBrowser::OnCompletePathPrefix)
-						.OnCanEditPathAsText(this, &SContentBrowser::OnCanEditPathAsText)
+						.OnGetEditPathAsText(this, &SContentBrowser::OnGetEditPathAsText)
 						.AddMetaData<FTagMetaData>(FTagMetaData(TEXT("ContentBrowserPath")));
 	}
 
@@ -2622,15 +2622,19 @@ void SContentBrowser::OnPathMenuItemClicked(FString ClickedPath)
 	OnPathClicked(ClickedPath);
 }
 
-bool SContentBrowser::OnCanEditPathAsText(const FString& Text) const
+FText SContentBrowser::OnGetEditPathAsText(const FString& Path) const
 {
 	const FAssetViewContentSources& ContentSources = LegacyContentSourceWidgets->AssetViewPtr->GetContentSources();
 	if (ContentSources.HasCollections())
 	{
 		// Do not present collections as text because their names are not very user friendly right now.
-		return false;
+		return FText::GetEmpty();
 	}
-	return true;
+	else if (ContentSources.HasVirtualPaths())
+	{
+		return ContentBrowserUtils::GetUserFacingPathFromVirtualPath(FName(Path));
+	}
+	return FText::GetEmpty();
 }
 
 TArray<FNavigationBarComboOption> SContentBrowser::OnCompletePathPrefix(const FString& Prefix) const
@@ -3726,7 +3730,7 @@ TSharedRef<SWidget> SContentBrowser::CreateLegacyAssetViewWidgets()
 					.GetComboOptions(this, &SContentBrowser::GetRecentPaths)
 					.OnNavigateToPath(this, &SContentBrowser::OnNavigateToPath)
 					.OnCompletePrefix(this, &SContentBrowser::OnCompletePathPrefix)
-					.OnCanEditPathAsText(this, &SContentBrowser::OnCanEditPathAsText)
+					.OnGetEditPathAsText(this, &SContentBrowser::OnGetEditPathAsText)
 					.AddMetaData<FTagMetaData>(FTagMetaData(TEXT("ContentBrowserPath")))
 				]
 
@@ -5034,14 +5038,14 @@ TSharedPtr<SWidget> SContentBrowser::GetItemContextMenu(TArrayView<const FConten
 			if (SelectedFolder.Legacy_TryGetPackagePath(PackagePath))
 			{
 				SelectedPackagePaths.Add(PackagePath.ToString());
+			}
 
-				if (!bPhysicalPathExists)
+			if (!bPhysicalPathExists)
+			{
+				FString PhysicalPath;
+				if (SelectedFolder.GetItemPhysicalPath(PhysicalPath) && FPaths::DirectoryExists(PhysicalPath))
 				{
-					FString PhysicalPath;
-					if (SelectedFolder.GetItemPhysicalPath(PhysicalPath) && FPaths::DirectoryExists(PhysicalPath))
-					{
-						bPhysicalPathExists = true;
-					}
+					bPhysicalPathExists = true;
 				}
 			}
 		}
@@ -5066,7 +5070,6 @@ TSharedPtr<SWidget> SContentBrowser::GetItemContextMenu(TArrayView<const FConten
 
 		if (SelectedPackagePaths.Num() == 0)
 		{
-			Context->bNoFolderOnDisk = true;
 			Context->bCanBeModified = false;
 		}
 
diff --git a/Engine/Source/Editor/ContentBrowser/Private/SContentBrowser.h b/Engine/Source/Editor/ContentBrowser/Private/SContentBrowser.h
--- a/Engine/Source/Editor/ContentBrowser/Private/SContentBrowser.h
+++ b/Engine/Source/Editor/ContentBrowser/Private/SContentBrowser.h
@@ -384,11 +384,8 @@ private:
 	/** Navigate to a location by text identifier. E.g. a virtual path such as /All/Game/MyFolder */
 	void OnNavigateToPath(const FString& Path);
 
-	/** 
-	  * Returns whether the currently shown location should be presented to the user as a string when beginning to edit the path, 
-	  * Otherwise the text box will be initially empty.
-	  */
-	bool OnCanEditPathAsText(const FString& Text) const;
+	/** Returns how the currently shown location should be presented to the user as text when beginning to edit the path */
+	FText OnGetEditPathAsText(const FString& Text) const;
 	
 	/** Returns a list of valid paths starting with Prefix which can be navigated to with OnNavigateToPath */
 	TArray<FNavigationBarComboOption> OnCompletePathPrefix(const FString& Prefix) const;
diff --git a/Engine/Source/Editor/ContentBrowser/Private/SNavigationBar.cpp b/Engine/Source/Editor/ContentBrowser/Private/SNavigationBar.cpp
--- a/Engine/Source/Editor/ContentBrowser/Private/SNavigationBar.cpp
+++ b/Engine/Source/Editor/ContentBrowser/Private/SNavigationBar.cpp
@@ -115,7 +115,7 @@ void SNavigationBar::Construct(const FArguments& InArgs)
 	OnGetComboOptions = InArgs._GetComboOptions;
 	OnCompletePrefix = InArgs._OnCompletePrefix;
 	OnNavigateToPath = InArgs._OnNavigateToPath;
-	OnCanEditPathAsText = InArgs._OnCanEditPathAsText;
+	OnGetEditPathAsText = InArgs._OnGetEditPathAsText;
 	OnPathClicked = InArgs._OnPathClicked;
 	OnGetPathMenuContent = InArgs._GetPathMenuContent;
 
@@ -442,9 +442,9 @@ void SNavigationBar::StartEditingPath()
 	if (BreadcrumbBar->HasCrumbs())
 	{
 		FNavigationCrumb Crumb = BreadcrumbBar->PeekCrumb();
-		if (OnCanEditPathAsText.IsBound() && OnCanEditPathAsText.Execute(Crumb.Path))
+		if (OnGetEditPathAsText.IsBound())
 		{
-			Text = FText::FromString(Crumb.Path);
+			Text = OnGetEditPathAsText.Execute(Crumb.Path);
 		}
 	}
 	EditableText->SetSelectAllTextWhenFocused(true);
diff --git a/Engine/Source/Editor/ContentBrowser/Private/SNavigationBar.h b/Engine/Source/Editor/ContentBrowser/Private/SNavigationBar.h
--- a/Engine/Source/Editor/ContentBrowser/Private/SNavigationBar.h
+++ b/Engine/Source/Editor/ContentBrowser/Private/SNavigationBar.h
@@ -36,7 +36,7 @@ using FGetPathMenuContent = TDelegate<TSharedRef<SWidget>(const FString&)>;
 using FGetComboOptions = TDelegate<TArray<FNavigationBarComboOption>(void)>;
 using FOnCompletePrefix = TDelegate<TArray<FNavigationBarComboOption>(const FString&)>;
 using FOnNavigateToPath = TDelegate<void(const FString&)>;
-using FCanEditPathAsText = TDelegate<bool(const FString&)>;
+using FGetEditPathAsText = TDelegate<FText(const FString&)>;
 
 // Private class for internal list view, stores a full possible location to nagivate to.
 struct FLocationItem;
@@ -93,11 +93,8 @@ public:
 		/** Called to get suggested text completion values for the currently edit text */
 		SLATE_EVENT(FOnCompletePrefix, OnCompletePrefix)
 		
-		/** 
-		 * Called to check if a path (e.g. the last breadcrumb) can be edited as text 
-		 * If not, when editing this location the text box will be emptied
-		 */
-		SLATE_EVENT(FCanEditPathAsText, OnCanEditPathAsText)
+		/** Called to get the edit location of a path (e.g. the last breadcrumb) as text */
+		SLATE_EVENT(FGetEditPathAsText, OnGetEditPathAsText)
 
     SLATE_END_ARGS()
     
@@ -161,7 +158,7 @@ private:
 	FOnNavigateToPath OnNavigateToPath;
 	FGetComboOptions OnGetComboOptions;
 	FOnCompletePrefix OnCompletePrefix;	
-	FCanEditPathAsText OnCanEditPathAsText;
+	FGetEditPathAsText OnGetEditPathAsText;
 	FOnPathClicked OnPathClicked;
 	FGetPathMenuContent OnGetPathMenuContent;
 };
\ No newline at end of file
