Turning back on the delayed population of the FPathView.
Added some just in time partial population for the public functions that need those even when the view is not presented (Sync and select path).
This is needed because unfortunatly the widgets of the CB owns to much of the logic outside of the actual presentation.
Removed a function that is no longer used from the SPathView needed since we do just in time partial population now.

diff --git a/Engine/Source/Editor/ContentBrowser/Private/SContentBrowser.cpp b/Engine/Source/Editor/ContentBrowser/Private/SContentBrowser.cpp
--- a/Engine/Source/Editor/ContentBrowser/Private/SContentBrowser.cpp
+++ b/Engine/Source/Editor/ContentBrowser/Private/SContentBrowser.cpp
@@ -1479,14 +1479,6 @@ void SContentBrowser::PrepareToSyncItems(TArrayView<const FContentBrowserItem> I
 		bSomePluginPathFiltersChanged = true;
 		bRepopulate = true;
 	}
-
-	// Check to see if any item paths don't exist (this can happen if we haven't ticked since the path was created)
-	if (!bRepopulate)
-	{
-		bRepopulate = Algo::AnyOf(ItemsToSync, [this](const FContentBrowserItem& Item) {
-			return Item.IsFolder() && !LegacyContentSourceWidgets->PathViewPtr->DoesItemExist(Item.GetVirtualPath());
-		});
-	}
 	
 	if (bDisableFiltersThatHideAssets)
 	{
diff --git a/Engine/Source/Editor/ContentBrowser/Private/SPathView.cpp b/Engine/Source/Editor/ContentBrowser/Private/SPathView.cpp
--- a/Engine/Source/Editor/ContentBrowser/Private/SPathView.cpp
+++ b/Engine/Source/Editor/ContentBrowser/Private/SPathView.cpp
@@ -346,6 +346,9 @@ public:
 	// Searches all items, not just visible ones according to the current text filter
 	TSharedPtr<FTreeItem> FindBestItemForPath(FStringView InVirtualPath);
 
+	void ProcessVirtualPaths(TConstArrayView<FName> InVirtualPaths, TFunction<FContentBrowserDataCompiledFilter()> CompiledFilterGetter);
+	void ProcessVirtualPaths(TConstArrayView<FName> InVirtualPaths, TFunction<FContentBrowserDataCompiledFilter()> CompiledFilterGetter, TSet <TSharedRef<FTreeItem>>& OutItems);
+
 	// Apply new/modified/removed data notifications to the tree
 	// bFlat parameter adds all items at the root of the tree as in the favorites tree
 	void ProcessDataUpdates(TConstArrayView<FContentBrowserItemDataUpdate> InUpdatedItems,
@@ -726,6 +729,40 @@ TSharedPtr<FTreeItem> FPathViewData::FindBestItemForPath(FStringView InVirtualPa
 	return Found;
 }
 
+void FPathViewData::ProcessVirtualPaths(TConstArrayView<FName> InVirtualPaths, TFunction<FContentBrowserDataCompiledFilter()> CompiledFilterGetter)
+{
+	UContentBrowserDataSubsystem* ContentBrowserData = IContentBrowserDataModule::Get().GetSubsystem();
+
+	TArray<FContentBrowserItemDataUpdate> FolderAtPath;
+	for (const FName& Path : InVirtualPaths)
+	{
+		ContentBrowserData->EnumerateItemsAtPath(Path, EContentBrowserItemTypeFilter::IncludeFolders, [&FolderAtPath](FContentBrowserItemData&& InItem) -> bool
+			{
+				FolderAtPath.Add(FContentBrowserItemDataUpdate::MakeItemAddedUpdate(MoveTemp(InItem)));
+				return true;
+			});
+	}
+
+	if (!FolderAtPath.IsEmpty())
+	{
+		ProcessDataUpdates(FolderAtPath, CompiledFilterGetter());
+	}
+}
+
+void FPathViewData::ProcessVirtualPaths(TConstArrayView<FName> InVirtualPaths, TFunction<FContentBrowserDataCompiledFilter()> CompiledFilterGetter, TSet<TSharedRef<FTreeItem>>& OutItems)
+{
+	ProcessVirtualPaths(InVirtualPaths, MoveTemp(CompiledFilterGetter));
+
+	for (const FName& VirtualPath : InVirtualPaths)
+	{
+		TSharedPtr<FTreeItem> Item = FindTreeItem(VirtualPath);
+		if (Item.IsValid())
+		{
+			OutItems.Add(Item.ToSharedRef());
+		}
+	}
+}
+
 TSharedRef<FTreeItem> FPathViewData::AddFolderItem(FContentBrowserItemData&& InItemData)
 {
 	TSharedRef<FTreeItem> NewOrUpdatedItem = AddFolderItemInternal(MoveTemp(InItemData), nullptr);
@@ -1165,9 +1202,6 @@ void SPathView::Construct( const FArguments& InArgs )
 
 	DefaultPath = InArgs._DefaultPath;
 	bCreateDefaultPath = InArgs._CreateDefaultPath;
-
-	// Temporary turn off the delayed population to investigate navigation issue when the view is not populated
-	Populate();
 }
 
 void SPathView::ConfigureTreeView(STreeView<TSharedPtr<FTreeItem>>::FArguments& InArgs)
@@ -1259,9 +1293,6 @@ void SPathView::PluginPathFilterClicked(TSharedRef<FContentBrowserPluginFilter>
 	SetPluginPathFilterActive(Filter, !IsPluginPathFilterInUse(Filter));
 	
 	bPopulationRequested = true;
-
-	// Temporary turn off the delayed population to investigate navigation issue when the view is not populated
-	Populate();
 }
 
 bool SPathView::IsPluginPathFilterChecked(TSharedRef<FContentBrowserPluginFilter> Filter) const
@@ -1297,8 +1328,6 @@ void SPathView::ResetPluginPathFilters()
 	}
 
 	bPopulationRequested = true;
-	// Temporary turn off the delayed population to investigate navigation issue when the view is not populated
-	Populate();
 }
 
 bool SPathView::DisablePluginPathFiltersThatHideItems(TConstArrayView<FContentBrowserItem> Items)
@@ -1414,6 +1443,25 @@ void SPathView::SetSelectedPaths(const TArray<FString>& Paths)
 	LastSelectedPaths.Empty();
 	TreeViewPtr->ClearSelection();
 
+	{
+		TArray<const FName> InPathToPopulate;
+		for (const FString& Path : Paths)
+		{
+			FName PathAsName = *Path;
+			TSharedPtr<FTreeItem> Item = TreeData->FindTreeItem(PathAsName);
+			if (!Item)
+			{
+				InPathToPopulate.Add(PathAsName);
+			}
+		}
+
+		if (!InPathToPopulate.IsEmpty())
+		{
+			TreeData->ProcessVirtualPaths(InPathToPopulate,
+				[this]() { return CreateCompiledFolderFilter(); });
+		}
+	}
+
 	for (const FString& Path : Paths)
 	{
 		TSharedPtr<FTreeItem> BestItem = TreeData->FindBestItemForPath(Path);
@@ -1747,12 +1795,28 @@ void SPathView::SyncToItems(TArrayView<const FContentBrowserItem> ItemsToSync, c
 void SPathView::SyncToVirtualPaths(TArrayView<const FName> VirtualPathsToSync, const bool bAllowImplicitSync)
 {
 	TSet<TSharedRef<FTreeItem>> SyncTreeItems;
-	for (const FName& VirtualPathToSync : VirtualPathsToSync)
+
 	{
-		TSharedPtr<FTreeItem> Item = TreeData->FindTreeItem(VirtualPathToSync);
-		if (Item.IsValid())
+		TArray<const FName> VirtualPathsNotFound;
+
+		for (const FName& VirtualPathToSync : VirtualPathsToSync)
 		{
-			SyncTreeItems.Add(Item.ToSharedRef());
+			TSharedPtr<FTreeItem> Item = TreeData->FindTreeItem(VirtualPathToSync);
+			if (Item.IsValid())
+			{
+				SyncTreeItems.Add(Item.ToSharedRef());
+			}
+			else
+			{
+				VirtualPathsNotFound.Add(VirtualPathToSync);
+			}
+		}
+
+		if (!VirtualPathsNotFound.IsEmpty())
+		{
+			TreeData->ProcessVirtualPaths(VirtualPathsNotFound,
+		[this]() { return CreateCompiledFolderFilter(); },
+				SyncTreeItems);
 		}
 	}
 
@@ -1823,11 +1887,6 @@ void SPathView::SyncToLegacy(TArrayView<const FAssetData> AssetDataList, TArrayV
 	SyncToVirtualPaths(VirtualPathsToSync, bAllowImplicitSync);
 }
 
-bool SPathView::DoesItemExist(FName InVirtualPath) const
-{
-	return TreeData->FindTreeItem(InVirtualPath).IsValid();
-}
-
 void SPathView::ApplyHistoryData( const FHistoryData& History )
 {
 	// Prevent the selection changed delegate because it would add more history when we are just setting a state
@@ -2257,9 +2316,6 @@ void SPathView::RequestPopulation(FSimpleMulticastDelegate::FDelegate&& AfterPop
 	{
 		OnNextPopulate.Add(MoveTemp(AfterPopulation));
 	}
-
-	// Temporary turn off the delayed population to investigate navigation issue when the view is not populated
-	Populate();
 }
 
 
@@ -2494,9 +2550,6 @@ void SPathView::HandleItemDataUpdated(TArrayView<const FContentBrowserItemDataUp
 void SPathView::HandleItemDataRefreshed()
 {
 	bPopulationRequested = true;
-
-	// Temporary turn off the delayed population to investigate navigation issue when the view is not populated
-	Populate();
 }
 
 void SPathView::HandleItemDataDiscoveryComplete()
@@ -2527,9 +2580,6 @@ void SPathView::HandleSettingChanged(FName PropertyName)
 
 		// Update our path view so that it can include/exclude the permitted folder
 		bPopulationForSettingChangeRequested = true;
-
-		// Temporary turn off the delayed population to investigate navigation issue when the view is not populated
-		Populate();
 	}
 }
 
diff --git a/Engine/Source/Editor/ContentBrowser/Private/SPathView.h b/Engine/Source/Editor/ContentBrowser/Private/SPathView.h
--- a/Engine/Source/Editor/ContentBrowser/Private/SPathView.h
+++ b/Engine/Source/Editor/ContentBrowser/Private/SPathView.h
@@ -232,9 +232,6 @@ public:
 	 */
 	void SyncToLegacy( TArrayView<const FAssetData> AssetDataList, TArrayView<const FString> FolderList, const bool bAllowImplicitSync = false );
 
-	/** Returns whether the tree contains an item with the given virtual path. */
-	bool DoesItemExist(FName InVirtualPath) const;
-
 	/** Sets the state of the path view to the one described by the history data */
 	void ApplyHistoryData( const FHistoryData& History );
 
