Fixed syncing to aliased assets by FAssetData not selecting them in the Content Browser.

#rnx

diff --git a/Engine/Source/Editor/ContentBrowser/Private/SContentBrowser.cpp b/Engine/Source/Editor/ContentBrowser/Private/SContentBrowser.cpp
--- a/Engine/Source/Editor/ContentBrowser/Private/SContentBrowser.cpp
+++ b/Engine/Source/Editor/ContentBrowser/Private/SContentBrowser.cpp
@@ -1669,9 +1669,13 @@ void SContentBrowser::SyncToLegacy(TArrayView<const FAssetData> AssetDataList, T
 				const FName InternalPath = Alias.GetInternalPathName();
 				if (FolderPermissions->PassesStartsWithFilter(InternalPath))
 				{
-					FAssetData AliasAssetData(InternalPath, *FPackageName::GetLongPackagePath(InternalPath.ToString()), FPackageName::GetShortFName(InternalPath), AssetDataList[i].AssetClassPath, AssetDataList[i].TagsAndValues.CopyMap());
-					NewItemsToSync.Add(MoveTemp(AliasAssetData));
-					break;
+					FSoftObjectPath ObjectPath{ FNameBuilder(InternalPath) };
+					if (ObjectPath.IsAsset())
+					{
+						FAssetData AliasAssetData(ObjectPath.GetLongPackageFName(), FName(*FPackageName::GetLongPackagePath(ObjectPath.GetLongPackageName())), ObjectPath.GetAssetFName(), AssetDataList[i].AssetClassPath, AssetDataList[i].TagsAndValues.CopyMap());
+						NewItemsToSync.Add(MoveTemp(AliasAssetData));
+						break;
+					}
 				}
 			}
 		}
