Don't do just in time population for the SFavoritePathView.

diff --git a/Engine/Source/Editor/ContentBrowser/Private/SPathView.cpp b/Engine/Source/Editor/ContentBrowser/Private/SPathView.cpp
--- a/Engine/Source/Editor/ContentBrowser/Private/SPathView.cpp
+++ b/Engine/Source/Editor/ContentBrowser/Private/SPathView.cpp
@@ -1443,6 +1443,7 @@ void SPathView::SetSelectedPaths(const TArray<FString>& Paths)
 	LastSelectedPaths.Empty();
 	TreeViewPtr->ClearSelection();
 
+	if (ShouldDoJustInTimePopulation())
 	{
 		TArray<const FName> InPathToPopulate;
 		for (const FString& Path : Paths)
@@ -1799,6 +1800,8 @@ void SPathView::SyncToVirtualPaths(TArrayView<const FName> VirtualPathsToSync, c
 	{
 		TArray<const FName> VirtualPathsNotFound;
 
+		bool bShouldDoJustInTimePopulation = ShouldDoJustInTimePopulation();
+
 		for (const FName& VirtualPathToSync : VirtualPathsToSync)
 		{
 			TSharedPtr<FTreeItem> Item = TreeData->FindTreeItem(VirtualPathToSync);
@@ -1806,7 +1809,7 @@ void SPathView::SyncToVirtualPaths(TArrayView<const FName> VirtualPathsToSync, c
 			{
 				SyncTreeItems.Add(Item.ToSharedRef());
 			}
-			else
+			else if (bShouldDoJustInTimePopulation)
 			{
 				VirtualPathsNotFound.Add(VirtualPathToSync);
 			}
@@ -1815,7 +1818,7 @@ void SPathView::SyncToVirtualPaths(TArrayView<const FName> VirtualPathsToSync, c
 		if (!VirtualPathsNotFound.IsEmpty())
 		{
 			TreeData->ProcessVirtualPaths(VirtualPathsNotFound,
-		[this]() { return CreateCompiledFolderFilter(); },
+			[this]() { return CreateCompiledFolderFilter(); },
 				SyncTreeItems);
 		}
 	}
diff --git a/Engine/Source/Editor/ContentBrowser/Private/SPathView.h b/Engine/Source/Editor/ContentBrowser/Private/SPathView.h
--- a/Engine/Source/Editor/ContentBrowser/Private/SPathView.h
+++ b/Engine/Source/Editor/ContentBrowser/Private/SPathView.h
@@ -356,6 +356,11 @@ protected:
 	/** Should be call after a populate */
 	void SharedPostPopulate();
 
+	constexpr virtual bool ShouldDoJustInTimePopulation() const
+	{
+		return true;
+	}
+
 private:
 	/** Selects the given path only if it exists. Returns true if selected. */
 	bool ExplicitlyAddPathToSelection(const FName Path);
@@ -639,6 +644,12 @@ public:
 	DECLARE_DELEGATE_OneParam(FOnFolderFavoriteAdd, const TArray<FString>& /*FoldersToAdd*/)
 	void SetOnFolderFavoriteAdd(const FOnFolderFavoriteAdd& InOnFolderFavoriteAdd);
 
+protected:
+	constexpr virtual bool ShouldDoJustInTimePopulation() const override
+	{
+		return false;
+	}
+
 private:
 	virtual TSharedRef<ITableRow> GenerateTreeRow(TSharedPtr<FTreeItem> TreeItem, const TSharedRef<STableViewBase>& OwnerTable) override;
 
