Refactored the Content Browser sorting to handle secondary sorting in a easier to understand way, then complicated things again by adding some column value caching while sorting, which brought the sort time of my 70k test case down from around several seconds to pretty instant (eyeballed).
Always use the display name then virtual path as tie breakers so you don't end up with random and inconsistent sort orders when sorting on columns with many of the same values.

diff --git a/Engine/Source/Editor/ContentBrowser/Private/AssetViewSortManager.cpp b/Engine/Source/Editor/ContentBrowser/Private/AssetViewSortManager.cpp
--- a/Engine/Source/Editor/ContentBrowser/Private/AssetViewSortManager.cpp
+++ b/Engine/Source/Editor/ContentBrowser/Private/AssetViewSortManager.cpp
@@ -22,380 +22,401 @@ struct FCompareFAssetItemBase
 {
 public:
 	/** Default constructor */
-	FCompareFAssetItemBase(bool bInAscending, const FName& InTag) : bAscending(bInAscending), Tag(InTag) {}
+	FCompareFAssetItemBase() = default;
 
-	virtual ~FCompareFAssetItemBase() = default;
-
-	/** Sort function */
-	FORCEINLINE bool operator()(const TSharedPtr<FAssetViewItem>& A, const TSharedPtr<FAssetViewItem>& B) const
-	{
-		bool bFolderComparisonResult = false;
-		if (CompareFolderFirst(A, B, bFolderComparisonResult))
-		{
-			return bFolderComparisonResult;
-		}
-		return Compare(A, B);
-	}
-
-	/** Get the next comparisons array */
-	FORCEINLINE TArray<TUniquePtr<FCompareFAssetItemBase>>& GetNextComparisons() const
-	{
-		return NextComparisons;
-	}
+	FCompareFAssetItemBase(const FCompareFAssetItemBase&) = delete;
 
-	/** Set the next comparisons array */
-	FORCEINLINE void SetNextComparisons(TArray<TUniquePtr<FCompareFAssetItemBase>>& InNextComparisons) const
-	{
-		check(NextComparisons.Num() == 0);
-		NextComparisons = MoveTemp(InNextComparisons);
-	}
+	virtual ~FCompareFAssetItemBase() = default;
 
-protected:
-	/** Comparison function, needs to be implemented by inheriting structs - by default just forwards to the next comparitor */
-	FORCEINLINE virtual bool Compare(const TSharedPtr<FAssetViewItem>& A, const TSharedPtr<FAssetViewItem>& B) const { return CompareNext(A, B); }
+	FCompareFAssetItemBase& operator=(const FCompareFAssetItemBase&) = delete;
 
-private:
-	/** Compare the folders of both assets */
-	FORCEINLINE bool CompareFolderFirst(const TSharedPtr<FAssetViewItem>& A, const TSharedPtr<FAssetViewItem>& B, bool& bComparisonResult) const
+	/** Returns true if the struct has cached all item values and the index overload of Compare should be called. */
+	virtual bool CacheValues(const TArray<TSharedPtr<FAssetViewItem>>& AssetItems)
 	{
-		// TODO: Have a view option to sort folders first or not (eg, Explorer vs Finder behavior)
-		const FContentBrowserItemData* AItemData = A->GetItem().GetPrimaryInternalItem();
-		const FContentBrowserItemData* BItemData = B->GetItem().GetPrimaryInternalItem();
-		if (AItemData && AItemData->IsFolder())
-		{
-			if (BItemData && BItemData->IsFolder())
-			{
-				const int32 Result = AItemData->GetDisplayName().CompareTo(BItemData->GetDisplayName());
-				bComparisonResult = bAscending ? Result < 0 : Result > 0;
-			}
-			else
-			{
-				bComparisonResult = bAscending;
-			}
-			return true;
-		}
-		else if (BItemData && BItemData->IsFolder())
-		{
-			bComparisonResult = !bAscending;
-			return true;
-		}
 		return false;
 	}
 
-	/** Comparison function in the event of a tie*/
-	FORCEINLINE bool CompareNext(const TSharedPtr<FAssetViewItem>& A, const TSharedPtr<FAssetViewItem>& B) const
+	/** Compare cached items. */
+	virtual int32 Compare(int32 IndexA, int32 IndexB) const
 	{
-		if (NextComparisons.Num() > 0)
-		{
-			TUniquePtr<FCompareFAssetItemBase> CompareNext = MoveTemp(NextComparisons[0]);
-			check(CompareNext);
-			NextComparisons.RemoveAt(0);
-
-			// Move all the remaining comparisons to the next comparitor
-			CompareNext->SetNextComparisons(GetNextComparisons());
-			const bool Result = CompareNext->Compare(A, B);
-
-			// Move the comparisons back for further ties
-			SetNextComparisons(CompareNext->GetNextComparisons());
-			NextComparisons.Insert(MoveTemp(CompareNext), 0);
-
-			return Result;
-		}
-		return true;	// default to true if a tie
+		checkNoEntry();
+		return 0;
 	}
 
-protected:
-	/** Whether to sort ascending or descending */
-	bool bAscending;
-
-	/** The tag type we need to compare next */
-	FName Tag;
-
-private:
-	/** The follow up sorting methods in the event of a tie */
-	mutable TArray<TUniquePtr<FCompareFAssetItemBase>> NextComparisons;
+	/** Compare uncached items. */
+	virtual int32 Compare(const FAssetViewItem& A, const FAssetViewItem& B) const = 0;
 };
 
-struct FCompareFAssetItemByName : public FCompareFAssetItemBase
+struct FCompareFAssetItemByName final : public FCompareFAssetItemBase
 {
 public:
-	FCompareFAssetItemByName(bool bInAscending, const FName& InTag) : FCompareFAssetItemBase(bInAscending, InTag) {}
-	virtual ~FCompareFAssetItemByName() override = default;
+	FCompareFAssetItemByName() = default;
 
-protected:
-	FORCEINLINE virtual bool Compare(const TSharedPtr<FAssetViewItem>& A, const TSharedPtr<FAssetViewItem>& B) const override
+	int32 Compare(const FAssetViewItem& A, const FAssetViewItem& B) const override
 	{
-		const FContentBrowserItemData* AItemData = A->GetItem().GetPrimaryInternalItem();
-		const FContentBrowserItemData* BItemData = B->GetItem().GetPrimaryInternalItem();
+		const FContentBrowserItemData* AItemData = A.GetItem().GetPrimaryInternalItem();
+		const FContentBrowserItemData* BItemData = B.GetItem().GetPrimaryInternalItem();
 		// TODO: Have an option to sort by display name? It's slower, but more correct for non-English languages
-		//const int32 Result = (AItemData && BItemData) ? AItemData->GetDisplayName().CompareTo(BItemData->GetDisplayName()) : 0;
-		//const int32 Result = (AItemData && BItemData) ? FAssetViewSortManager::CompareWithNumericSuffix(FNameBuilder(AItemData->GetItemName()).ToView(), FNameBuilder(BItemData->GetItemName()).ToView()) : 0;
-		const int32 Result = (AItemData && BItemData) ? UE::ComparisonUtility::CompareWithNumericSuffix(AItemData->GetDisplayName().ToString(), BItemData->GetDisplayName().ToString()) : 0;
-		if (Result < 0)
-		{
-			return bAscending;
-		}
-		else if (Result > 0)
-		{
-			return !bAscending;
-		}
-		return FCompareFAssetItemBase::Compare(A, B);
+		//return (AItemData && BItemData) ? AItemData->GetDisplayName().CompareTo(BItemData->GetDisplayName()) : 0;
+		//return (AItemData && BItemData) ? FAssetViewSortManager::CompareWithNumericSuffix(FNameBuilder(AItemData->GetItemName()).ToView(), FNameBuilder(BItemData->GetItemName()).ToView()) : 0;
+		return (AItemData && BItemData) ? UE::ComparisonUtility::CompareWithNumericSuffix(AItemData->GetDisplayName().ToString(), BItemData->GetDisplayName().ToString()) : 0;
 	}
 };
 
-struct FCompareFAssetItemByClass : public FCompareFAssetItemBase
+struct FCompareFAssetItemByClass final : public FCompareFAssetItemBase
 {
 public:
-	FCompareFAssetItemByClass(bool bInAscending, const FName& InTag) : FCompareFAssetItemBase(bInAscending, InTag) {}
-	virtual ~FCompareFAssetItemByClass() override = default;
+	FCompareFAssetItemByClass() = default;
 
-protected:
-	FORCEINLINE virtual bool Compare(const TSharedPtr<FAssetViewItem>& A, const TSharedPtr<FAssetViewItem>& B) const override
+	bool CacheValues(const TArray<TSharedPtr<FAssetViewItem>>& AssetItems) override
 	{
-		auto GetClassName = [](const FContentBrowserItem& Item)
-		{
-			FContentBrowserItemDataAttributeValue ClassValue = Item.GetItemAttribute(ContentBrowserItemAttributes::ItemTypeName);
-			return ClassValue.IsValid() ? ClassValue.GetValue<FName>() : FName();
-		};
+		CachedValues.Reset(AssetItems.Num());
 
-		const FName ValueA = GetClassName(A->GetItem());
-		const FName ValueB = GetClassName(B->GetItem());
-		const int32 Result = ValueA.Compare(ValueB);
-		if (Result < 0)
+		for (const TSharedPtr<FAssetViewItem>& AssetItem : AssetItems)
 		{
-			return bAscending;
+			CachedValues.Add(GetClassName(AssetItem->GetItem()));
 		}
-		else if (Result > 0)
-		{
-			return !bAscending;
-		}
-		return FCompareFAssetItemBase::Compare(A, B);
+
+		return true;
 	}
+
+	int32 Compare(int32 IndexA, int32 IndexB) const override
+	{
+		return CachedValues[IndexA].Compare(CachedValues[IndexB]);
+	}
+
+	int32 Compare(const FAssetViewItem& A, const FAssetViewItem& B) const override
+	{
+		return GetClassName(A.GetItem()).Compare(GetClassName(B.GetItem()));
+	}
+
+private:
+	static FName GetClassName(const FContentBrowserItem& Item)
+	{
+		FContentBrowserItemDataAttributeValue ClassValue = Item.GetItemAttribute(ContentBrowserItemAttributes::ItemTypeName);
+		return ClassValue.IsValid() ? ClassValue.GetValue<FName>() : FName();
+	}
+
+	TArray<FName> CachedValues;
 };
 
-struct FCompareFAssetItemByDiskSize : public FCompareFAssetItemBase
+struct FCompareFAssetItemByDiskSize final : public FCompareFAssetItemBase
 {
 public:
-	FCompareFAssetItemByDiskSize(bool bInAscending, const FName& InTag) : FCompareFAssetItemBase(bInAscending, InTag) {}
-	virtual ~FCompareFAssetItemByDiskSize() override = default;
+	FCompareFAssetItemByDiskSize() = default;
 
-protected:
-	FORCEINLINE virtual bool Compare(const TSharedPtr<FAssetViewItem>& A, const TSharedPtr<FAssetViewItem>& B) const override
+	bool CacheValues(const TArray<TSharedPtr<FAssetViewItem>>& AssetItems) override
 	{
-		auto GetDiskSize = [](const FContentBrowserItem& Item)
-		{
-			FContentBrowserItemDataAttributeValue DiskSizeValue = Item.GetItemAttribute(ContentBrowserItemAttributes::ItemDiskSize);
-			return DiskSizeValue.IsValid() ? DiskSizeValue.GetValue<int64>() : -1; // -1 is for items where size doesn't apply - ie. folders
-		};
-
-		const int64 ValueA = GetDiskSize(A->GetItem());
-		const int64 ValueB = GetDiskSize(B->GetItem());
+		CachedValues.Reset(AssetItems.Num());
 
-		if (ValueA < ValueB)
+		for (const TSharedPtr<FAssetViewItem>& AssetItem : AssetItems)
 		{
-			return bAscending;
+			CachedValues.Add(GetDiskSize(AssetItem->GetItem()));
 		}
 
-		if (ValueA > ValueB)
-		{
-			return !bAscending;
-		}
+		return true;
+	}
+
+	int32 Compare(int32 IndexA, int32 IndexB) const override
+	{
+		return Compare(CachedValues[IndexA], CachedValues[IndexB]);
+	}
+
+	int32 Compare(const FAssetViewItem& A, const FAssetViewItem& B) const override
+	{
+		return Compare(GetDiskSize(A.GetItem()), GetDiskSize(B.GetItem()));
+	}
+
+private:
+	static int32 Compare(int64 ValueA, int64 ValueB)
+	{
+		return ValueA < ValueB ? -1 : (ValueA > ValueB ? 1 : 0);
+	}
 
-		return FCompareFAssetItemBase::Compare(A, B);
+	static int64 GetDiskSize(const FContentBrowserItem& Item)
+	{
+		FContentBrowserItemDataAttributeValue DiskSizeValue = Item.GetItemAttribute(ContentBrowserItemAttributes::ItemDiskSize);
+		return DiskSizeValue.IsValid() ? DiskSizeValue.GetValue<int64>() : -1; // -1 is for items where size doesn't apply - ie. folders
 	}
+
+	TArray<int64> CachedValues;
 };
 
-struct FCompareFAssetItemByPath : public FCompareFAssetItemBase
+struct FCompareFAssetItemByPath final : public FCompareFAssetItemBase
 {
 public:
-	FCompareFAssetItemByPath(bool bInAscending, const FName& InTag) : FCompareFAssetItemBase(bInAscending, InTag) {}
-	virtual ~FCompareFAssetItemByPath() override = default;
+	FCompareFAssetItemByPath() = default;
 
-protected:
-	FORCEINLINE virtual bool Compare(const TSharedPtr<FAssetViewItem>& A, const TSharedPtr<FAssetViewItem>& B) const override
+	int32 Compare(const FAssetViewItem& A, const FAssetViewItem& B) const override
 	{
-		const FContentBrowserItemData* AItemData = A->GetItem().GetPrimaryInternalItem();
-		const FContentBrowserItemData* BItemData = B->GetItem().GetPrimaryInternalItem();
-		const int32 Result = (AItemData && BItemData) ? UE::ComparisonUtility::CompareWithNumericSuffix(AItemData->GetVirtualPath(), BItemData->GetVirtualPath()) : 0;
-		if (Result < 0)
-		{
-			return bAscending;
-		}
-		else if (Result > 0)
-		{
-			return !bAscending;
-		}
-		return FCompareFAssetItemBase::Compare(A, B);
+		const FContentBrowserItemData* AItemData = A.GetItem().GetPrimaryInternalItem();
+		const FContentBrowserItemData* BItemData = B.GetItem().GetPrimaryInternalItem();
+		return (AItemData && BItemData) ? UE::ComparisonUtility::CompareWithNumericSuffix(AItemData->GetVirtualPath(), BItemData->GetVirtualPath()) : 0;
 	}
 };
 
-struct FCompareFAssetItemByAssetAccessSpecifier : public FCompareFAssetItemBase
+struct FCompareFAssetItemByAssetAccessSpecifier final : public FCompareFAssetItemBase
 {
 public:
-	FCompareFAssetItemByAssetAccessSpecifier(bool bInAscending, const FName& InTag) : FCompareFAssetItemBase(bInAscending, InTag) {}
-	virtual ~FCompareFAssetItemByAssetAccessSpecifier() override = default;
+	FCompareFAssetItemByAssetAccessSpecifier() = default;
 
-protected:
-	FORCEINLINE virtual bool Compare(const TSharedPtr<FAssetViewItem>& A, const TSharedPtr<FAssetViewItem>& B) const override
+	bool CacheValues(const TArray<TSharedPtr<FAssetViewItem>>& AssetItems) override
 	{
-		FText Value1 = A->GetItemAssetAccessSpecifierText();
-		FText Value2 = B->GetItemAssetAccessSpecifierText();
+		CachedValues.Reset(AssetItems.Num());
 
-		// Depending if we're sorting ascending or descending it's quicker to flip the compares incase values are missing
-		const bool bFoundValue1 = bAscending ? !Value1.IsEmpty() : !Value2.IsEmpty();
-		if (!bFoundValue1)
+		for (const TSharedPtr<FAssetViewItem>& AssetItem : AssetItems)
 		{
-			return true;
+			CachedValues.Add(AssetItem->GetItemAssetAccessSpecifierText());
 		}
 
-		const bool bFoundValue2 = bAscending ? !Value2.IsEmpty() : !Value1.IsEmpty();
-		if (!bFoundValue2)
+		return true;
+	}
+
+	int32 Compare(int32 IndexA, int32 IndexB) const override
+	{
+		return Compare(CachedValues[IndexA], CachedValues[IndexB]);
+	}
+
+	int32 Compare(const FAssetViewItem& A, const FAssetViewItem& B) const override
+	{
+		return Compare(A.GetItemAssetAccessSpecifierText(), B.GetItemAssetAccessSpecifierText());
+	}
+
+private:
+	static int32 Compare(const FText& ValueA, const FText& ValueB)
+	{
+		const bool bValueAIsEmpty = ValueA.IsEmpty();
+		const bool bValueBIsEmpty = ValueB.IsEmpty();
+
+		if (bValueAIsEmpty && bValueBIsEmpty)
 		{
-			return false;
+			return 0;
 		}
 
-		const int32 Result = Value1.CompareTo(Value2);
-		if (Result < 0)
+		if (bValueAIsEmpty)
 		{
-			return bAscending;
+			return 1;
 		}
-		else if (Result > 0)
+
+		if (bValueBIsEmpty)
 		{
-			return !bAscending;
+			return -1;
 		}
-		return FCompareFAssetItemBase::Compare(A, B);
+
+		return ValueA.CompareTo(ValueB);
 	}
+
+	TArray<FText> CachedValues;
 };
 
-struct FCompareFAssetItemByTag : public FCompareFAssetItemBase
+struct FCompareFAssetItemByTag final : public FCompareFAssetItemBase
 {
 public:
-	FCompareFAssetItemByTag(bool bInAscending, const FName& InTag) : FCompareFAssetItemBase(bInAscending, InTag) {}
-	virtual ~FCompareFAssetItemByTag() override = default;
+	explicit FCompareFAssetItemByTag(FName InTag)
+		: Tag(InTag)
+	{
+	}
 
-protected:
-	FORCEINLINE virtual bool Compare(const TSharedPtr<FAssetViewItem>& A, const TSharedPtr<FAssetViewItem>& B) const override
+	bool CacheValues(const TArray<TSharedPtr<FAssetViewItem>>& AssetItems) override
 	{
-		// Depending if we're sorting ascending or descending it's quicker to flip the compares incase tags are missing
-		FString Value1;
-		const bool bFoundValue1 = bAscending ? A->GetTagValue(Tag, Value1) : B->GetTagValue(Tag, Value1);
-		if (!bFoundValue1)
-		{
-			return true;
-		}
-		
-		FString Value2;
-		const bool bFoundValue2 = bAscending ? B->GetTagValue(Tag, Value2) : A->GetTagValue(Tag, Value2);
-		if (!bFoundValue2)
+		CachedValues.Reset(AssetItems.Num());
+
+		for (const TSharedPtr<FAssetViewItem>& AssetItem : AssetItems)
 		{
-			return false;
+			FString Value;
+			if (AssetItem->GetTagValue(Tag, Value))
+			{
+				CachedValues.Emplace(MoveTemp(Value));
+			}
+			else
+			{
+				CachedValues.Emplace();
+			}
 		}
-		
-		const int32 Result = Value1.Compare(Value2, ESearchCase::IgnoreCase);
-		if (Result < 0)
+
+		return true;
+	}
+
+	int32 Compare(int32 IndexA, int32 IndexB) const override
+	{
+		const TOptional<FString>& ValueA = CachedValues[IndexA];
+		const TOptional<FString>& ValueB = CachedValues[IndexB];
+
+		if (ValueA.IsSet() && ValueB.IsSet())
 		{
-			return true;
+			return ValueA.GetValue().Compare(ValueB.GetValue(), ESearchCase::IgnoreCase);
 		}
-		else if (Result > 0)
+
+		return ValueA.IsSet() ? -1 : (ValueB.IsSet() ? 1 : 0);
+	}
+
+	int32 Compare(const FAssetViewItem& A, const FAssetViewItem& B) const override
+	{
+		FString ValueA;
+		const bool bFoundValueA = A.GetTagValue(Tag, ValueA);
+
+		FString ValueB;
+		const bool bFoundValueB = B.GetTagValue(Tag, ValueB);
+
+		if (bFoundValueA && bFoundValueB)
 		{
-			return false;
+			return ValueA.Compare(ValueB, ESearchCase::IgnoreCase);
 		}
-		return FCompareFAssetItemBase::Compare(A, B);
+
+		return bFoundValueA ? -1 : (bFoundValueB ? 1 : 0);
 	}
+
+private:
+	const FName Tag;
+	TArray<TOptional<FString>> CachedValues;
 };
 
-struct FCompareFAssetItemByTagNumerical : public FCompareFAssetItemBase
+struct FCompareFAssetItemByTagNumerical final : public FCompareFAssetItemBase
 {
 public:
-	FCompareFAssetItemByTagNumerical(bool bInAscending, const FName& InTag) : FCompareFAssetItemBase(bInAscending, InTag) {}
-	virtual ~FCompareFAssetItemByTagNumerical() override = default;
+	explicit FCompareFAssetItemByTagNumerical(FName InTag)
+		: Tag(InTag)
+	{
+	}
 
-protected:
-	FORCEINLINE virtual bool Compare(const TSharedPtr<FAssetViewItem>& A, const TSharedPtr<FAssetViewItem>& B) const override
+	bool CacheValues(const TArray<TSharedPtr<FAssetViewItem>>& AssetItems) override
 	{
-		// Depending if we're sorting ascending or descending it's quicker to flip the compares incase tags are missing
-		FString Value1;
-		const bool bFoundValue1 = bAscending ? A->GetTagValue(Tag, Value1) : B->GetTagValue(Tag, Value1);
-		if (!bFoundValue1)
-		{
-			return true;
-		}
-		
-		FString Value2;
-		const bool bFoundValue2 = bAscending ? B->GetTagValue(Tag, Value2) : A->GetTagValue(Tag, Value2);
-		if (!bFoundValue2)
+		CachedValues.Reset(AssetItems.Num());
+
+		FString Value;
+		for (const TSharedPtr<FAssetViewItem>& AssetItem : AssetItems)
 		{
-			return false;
+			if (AssetItem->GetTagValue(Tag, Value))
+			{
+				float FloatValue = 0.0f;
+				LexFromString(FloatValue, *Value);
+				CachedValues.Emplace(FloatValue);
+			}
+			else
+			{
+				CachedValues.Emplace();
+			}
 		}
-		
-		float FloatValue1 = 0.0f, FloatValue2 = 0.0f;
-		LexFromString(FloatValue1, *Value1);
-		LexFromString(FloatValue2, *Value2);
-		
-		if (FloatValue1 < FloatValue2)
+
+		return true;
+	}
+
+	int32 Compare(int32 IndexA, int32 IndexB) const override
+	{
+		const TOptional<float>& ValueA = CachedValues[IndexA];
+		const TOptional<float>& ValueB = CachedValues[IndexB];
+
+		if (ValueA.IsSet() && ValueB.IsSet())
 		{
-			return true;
+			return ValueA.GetValue() < ValueB.GetValue() ? -1 : (ValueA.GetValue() > ValueB.GetValue() ? 1 : 0);
 		}
-		else if (FloatValue1 > FloatValue2)
+
+		return ValueA.IsSet() ? -1 : (ValueB.IsSet() ? 1 : 0);
+	}
+
+	int32 Compare(const FAssetViewItem& A, const FAssetViewItem& B) const override
+	{
+		FString ValueA;
+		const bool bFoundValueA = A.GetTagValue(Tag, ValueA);
+
+		FString ValueB;
+		const bool bFoundValueB = B.GetTagValue(Tag, ValueB);
+
+		if (bFoundValueA && bFoundValueB)
 		{
-			return false;
+			float FloatValueA = 0.0f, FloatValueB = 0.0f;
+			LexFromString(FloatValueA, *ValueA);
+			LexFromString(FloatValueB, *ValueB);
+
+			return FloatValueA < FloatValueB ? -1 : (FloatValueA > FloatValueB ? 1 : 0);
 		}
-		return FCompareFAssetItemBase::Compare(A, B);
+
+		return bFoundValueA ? -1 : (bFoundValueB ? 1 : 0);
 	}
+
+private:
+	const FName Tag;
+	TArray<TOptional<float>> CachedValues;
 };
 
-struct FCompareFAssetItemByTagDimensional : public FCompareFAssetItemBase
+struct FCompareFAssetItemByTagDimensional final : public FCompareFAssetItemBase
 {
 public:
-	FCompareFAssetItemByTagDimensional(bool bInAscending, const FName& InTag) : FCompareFAssetItemBase(bInAscending, InTag) {}
-	virtual ~FCompareFAssetItemByTagDimensional() override = default;
+	explicit FCompareFAssetItemByTagDimensional(FName InTag)
+		: Tag(InTag)
+	{
+	}
 
-protected:
-	FORCEINLINE virtual bool Compare(const TSharedPtr<FAssetViewItem>& A, const TSharedPtr<FAssetViewItem>& B) const override
+	bool CacheValues(const TArray<TSharedPtr<FAssetViewItem>>& AssetItems) override
 	{
-		// Depending if we're sorting ascending or descending it's quicker to flip the compares incase tags are missing
-		FString Value1;
-		const bool bHasFoundValue1 = bAscending ? A->GetTagValue(Tag, Value1) : B->GetTagValue(Tag, Value1);
-		if (!bHasFoundValue1)
-		{
-			return true;
-		}
-		
-		FString Value2;
-		const bool bHasFoundValue2 = bAscending ? B->GetTagValue(Tag, Value2) : A->GetTagValue(Tag, Value2);
-		if (!bHasFoundValue2)
-		{
-			return false;
-		}
-		
-		float Num1 = 1.f;
+		CachedValues.Reset(AssetItems.Num());
+
+		FString Value;
+		TArray<FString> Tokens;
+		for (const TSharedPtr<FAssetViewItem>& AssetItem : AssetItems)
 		{
-			TArray<FString> Tokens1;
-			Value1.ParseIntoArray(Tokens1, TEXT("x"), true);
-			for (auto TokenIt1 = Tokens1.CreateConstIterator(); TokenIt1; ++TokenIt1)
+			if (AssetItem->GetTagValue(Tag, Value))
 			{
-				Num1 *= FCString::Atof(**TokenIt1);
+				CachedValues.Emplace(ComputeValue(Value, Tokens));
 			}
-		}
-		float Num2 = 1.f;
-		{
-			TArray<FString> Tokens2;
-			Value2.ParseIntoArray(Tokens2, TEXT("x"), true);
-			for (auto TokenIt2 = Tokens2.CreateConstIterator(); TokenIt2; ++TokenIt2)
+			else
 			{
-				Num2 *= FCString::Atof(**TokenIt2);
+				CachedValues.Emplace();
 			}
 		}
-		if (Num1 < Num2)
+
+		return true;
+	}
+
+	int32 Compare(int32 IndexA, int32 IndexB) const override
+	{
+		const TOptional<float>& ValueA = CachedValues[IndexA];
+		const TOptional<float>& ValueB = CachedValues[IndexB];
+
+		if (ValueA.IsSet() && ValueB.IsSet())
 		{
-			return true;
+			return ValueA.GetValue() < ValueB.GetValue() ? -1 : (ValueA.GetValue() > ValueB.GetValue() ? 1 : 0);
 		}
-		else if (Num1 > Num2)
+
+		return ValueA.IsSet() ? -1 : (ValueB.IsSet() ? 1 : 0);
+	}
+
+	int32 Compare(const FAssetViewItem& A, const FAssetViewItem& B) const override
+	{
+		FString ValueA;
+		const bool bFoundValueA = A.GetTagValue(Tag, ValueA);
+
+		FString ValueB;
+		const bool bFoundValueB = B.GetTagValue(Tag, ValueB);
+
+		if (bFoundValueA && bFoundValueB)
 		{
-			return false;
+			TArray<FString> Tokens;
+			const float FloatValueA = ComputeValue(ValueA, Tokens);
+			const float FloatValueB = ComputeValue(ValueB, Tokens);
+
+			return FloatValueA < FloatValueB ? -1 : (FloatValueA > FloatValueB ? 1 : 0);
 		}
-		return FCompareFAssetItemBase::Compare(A, B);
+
+		return bFoundValueA ? -1 : (bFoundValueB ? 1 : 0);
 	}
+
+private:
+	static float ComputeValue(const FString& Value, TArray<FString>& Tokens)
+	{
+		float FloatValue = 1.f;
+
+		Value.ParseIntoArray(Tokens, TEXT("x"), true);
+		for (const FString& Token : Tokens)
+		{
+			FloatValue *= FCString::Atof(*Token);
+		}
+
+		return FloatValue;
+	}
+
+	const FName Tag;
+	TArray<TOptional<float>> CachedValues;
 };
 
 const FName FAssetViewSortManager::NameColumnId = "Name";
@@ -447,38 +468,66 @@ bool FAssetViewSortManager::FindAndRefreshCustomColumn(const TArray<TSharedPtr<F
 
 void FAssetViewSortManager::SortList(TArray<TSharedPtr<FAssetViewItem>>& AssetItems, const FName& MajorityAssetType, const TArray<FAssetViewCustomColumn>& CustomColumns) const
 {
+	if (AssetItems.Num() <= 1)
+	{
+		return;
+	}
+
 	//double SortListStartTime = FPlatformTime::Seconds();
 
-	TArray<TUniquePtr<FCompareFAssetItemBase>> SortMethod;
+	int32 PrimarySign = 0;
+	TArray<TPair<TUniquePtr<FCompareFAssetItemBase>, int32>, TInlineAllocator<EColumnSortPriority::Max>> SortMethods;
+	bool bHasNameSortMethod = false;
 	for (int32 PriorityIdx = 0; PriorityIdx < EColumnSortPriority::Max; PriorityIdx++)
 	{
-		const bool bAscending(SortMode[PriorityIdx] == EColumnSortMode::Ascending);
-		const FName& Tag(SortColumnId[PriorityIdx]);
+		int32 Sign = 0;
+		switch (SortMode[PriorityIdx])
+		{
+		case EColumnSortMode::None:
+			break;
+		case EColumnSortMode::Ascending:
+			Sign = 1;
+			break;
+		case EColumnSortMode::Descending:
+			Sign = -1;
+			break;
+		default:
+			checkNoEntry();
+			break;
+		}
+		if (Sign == 0)
+		{
+			break;
+		}
 
+		const FName& Tag(SortColumnId[PriorityIdx]);
 		if (Tag == NAME_None)
 		{
 			break;
 		}
 
+		TUniquePtr<FCompareFAssetItemBase> SortMethod;
+
 		if (Tag == NameColumnId)
 		{
-			SortMethod.Add(MakeUnique<FCompareFAssetItemByName>(bAscending, Tag));
+			SortMethod = MakeUnique<FCompareFAssetItemByName>();
+			bHasNameSortMethod = true;
 		}
 		else if (Tag == ClassColumnId)
 		{
-			SortMethod.Add(MakeUnique<FCompareFAssetItemByClass>(bAscending, Tag));
+			SortMethod = MakeUnique<FCompareFAssetItemByClass>();
 		}
 		else if (Tag == PathColumnId)
 		{
-			SortMethod.Add(MakeUnique<FCompareFAssetItemByPath>(bAscending, Tag));
+			SortMethod = MakeUnique<FCompareFAssetItemByPath>();
 		}
 		else if (Tag == DiskSizeColumnId)
 		{
-			SortMethod.Add(MakeUnique<FCompareFAssetItemByDiskSize>(bAscending, Tag));
+			SortMethod = MakeUnique<FCompareFAssetItemByDiskSize>();
 		}
 		else if (Tag == AssetAccessSpecifierColumnId)
 		{
-			SortMethod.Add(MakeUnique<FCompareFAssetItemByAssetAccessSpecifier>(bAscending, Tag));
+			SortMethod = MakeUnique<FCompareFAssetItemByAssetAccessSpecifier>();
 		}
 		else
 		{
@@ -509,43 +558,146 @@ void FAssetViewSortManager::SortList(TArray<TSharedPtr<FAssetViewItem>>& AssetIt
 			if (TagType == UObject::FAssetRegistryTag::TT_Numerical)
 			{
 				// The property is a Number, compare using atof
-				SortMethod.Add(MakeUnique<FCompareFAssetItemByTagNumerical>(bAscending, Tag));
+				SortMethod = MakeUnique<FCompareFAssetItemByTagNumerical>(Tag);
 			}
 			else if (TagType == UObject::FAssetRegistryTag::TT_Dimensional)
 			{
 				// The property is a series of Numbers representing dimensions, compare by using atof for each Number, delimited by an "x"
-				SortMethod.Add(MakeUnique<FCompareFAssetItemByTagDimensional>(bAscending, Tag));
+				SortMethod = MakeUnique<FCompareFAssetItemByTagDimensional>(Tag);
 			}
 			else if (TagType != UObject::FAssetRegistryTag::ETagType::TT_Hidden)
 			{
 				// Unknown or alphabetical, sort alphabetically either way
-				SortMethod.Add(MakeUnique<FCompareFAssetItemByTag>(bAscending, Tag));
+				SortMethod = MakeUnique<FCompareFAssetItemByTag>(Tag);
+			}
+		}
+
+		if (SortMethod)
+		{
+			if (SortMethods.IsEmpty())
+			{
+				PrimarySign = Sign;
 			}
+			SortMethods.Emplace(MoveTemp(SortMethod), Sign);
 		}
 	}
 
 	// Sort the list...
-	if (SortMethod.Num() > 0)
+	if (SortMethods.Num() > 0)
 	{
-		TUniquePtr<FCompareFAssetItemBase> PrimarySortMethod = MoveTemp(SortMethod[EColumnSortPriority::Primary]);
-		check(PrimarySortMethod);
-		SortMethod.RemoveAt(0);
+		// Use the display name as the primary tie breaker.
+		if (!bHasNameSortMethod)
+		{
+			SortMethods.Emplace(MakeUnique<FCompareFAssetItemByName>(), PrimarySign);
+		}
+
+		auto CompareType = [](const FAssetViewItem& AssetViewItemA, const FAssetViewItem& AssetViewItemB) -> int32
+		{
+			// TODO: Have a view option to sort folders first or not (eg, Explorer vs Finder behavior)
+			const FContentBrowserItemData* AItemData = AssetViewItemA.GetItem().GetPrimaryInternalItem();
+			const FContentBrowserItemData* BItemData = AssetViewItemB.GetItem().GetPrimaryInternalItem();
+			if (AItemData && AItemData->IsFolder())
+			{
+				if (BItemData && BItemData->IsFolder())
+				{
+					return AItemData->GetDisplayName().CompareTo(BItemData->GetDisplayName());
+				}
+				return -1;
+			}
+			else if (BItemData && BItemData->IsFolder())
+			{
+				return 1;
+			}
+			return 0;
+		};
 
-		// Move all the comparisons to the primary sort method
-		PrimarySortMethod->SetNextComparisons(SortMethod);
-		AssetItems.Sort(*(PrimarySortMethod.Get()));
+		// Check if the primary sort method supports a cached fast path.
+		if (SortMethods[0].Key->CacheValues(AssetItems))
+		{
+			TArray<TPair<int32, TSharedPtr<FAssetViewItem>>> IndexedAssetItems;
+			IndexedAssetItems.Reserve(AssetItems.Num());
+			for (int32 Index = 0; Index < AssetItems.Num(); ++Index)
+			{
+				IndexedAssetItems.Emplace(Index, MoveTemp(AssetItems[Index]));
+			}
 
-		// Move the comparisons back for ease of cleanup
-		SortMethod = MoveTemp(PrimarySortMethod->GetNextComparisons());
-		SortMethod.Insert(MoveTemp(PrimarySortMethod), 0);
-	}
+			IndexedAssetItems.Sort([PrimarySign, &CompareType, &SortMethods](const TPair<int32, TSharedPtr<FAssetViewItem>>& A, const TPair<int32, TSharedPtr<FAssetViewItem>>& B)
+			{
+				const FAssetViewItem& AssetViewItemA = *A.Value;
+				const FAssetViewItem& AssetViewItemB = *B.Value;
 
-	// Cleanup the methods we no longer need.
-	for (int32 PriorityIdx = 0; PriorityIdx < SortMethod.Num(); PriorityIdx++)
-	{
-		SortMethod[PriorityIdx].Reset();
+				// Group folders and files first.
+				{
+					const int32 Result = CompareType(AssetViewItemA, AssetViewItemB);
+					if (Result != 0)
+					{
+						return PrimarySign * Result < 0;
+					}
+				}
+
+				// Cached sort by primary column.
+				{
+					const int32 IndexA = A.Key;
+					const int32 IndexB = B.Key;
+
+					const int32 Result = SortMethods[0].Key->Compare(IndexA, IndexB);
+					if (Result != 0)
+					{
+						return SortMethods[0].Value * Result < 0;
+					}
+				}
+
+				// Sort by secondary columns.
+				for (int32 Index = 1; Index < SortMethods.Num(); ++Index)
+				{
+					const TPair<TUniquePtr<FCompareFAssetItemBase>, int32>& SortMethod = SortMethods[Index];
+					const int32 Result = SortMethod.Key->Compare(AssetViewItemA, AssetViewItemB);
+					if (Result != 0)
+					{
+						return SortMethod.Value * Result < 0;
+					}
+				}
+
+				// Use virtual paths as the final tie breaker.
+				return PrimarySign * AssetViewItemA.GetItem().GetVirtualPath().Compare(AssetViewItemB.GetItem().GetVirtualPath()) < 0;
+			});
+
+			for (int32 Index = 0; Index < AssetItems.Num(); ++Index)
+			{
+				AssetItems[Index] = MoveTemp(IndexedAssetItems[Index].Value);
+			}
+		}
+		else
+		{
+			AssetItems.Sort([PrimarySign, &CompareType, &SortMethods](const TSharedPtr<FAssetViewItem>& A, const TSharedPtr<FAssetViewItem>& B)
+			{
+				const FAssetViewItem& AssetViewItemA = *A;
+				const FAssetViewItem& AssetViewItemB = *B;
+
+				// Group folders and files first.
+				{
+					const int32 Result = CompareType(AssetViewItemA, AssetViewItemB);
+					if (Result != 0)
+					{
+						return PrimarySign * Result < 0;
+					}
+				}
+
+				// Sort by columns.
+				for (const TPair<TUniquePtr<FCompareFAssetItemBase>, int32>& SortMethod : SortMethods)
+				{
+					const int32 Result = SortMethod.Key->Compare(AssetViewItemA, AssetViewItemB);
+					if (Result != 0)
+					{
+						return SortMethod.Value * Result < 0;
+					}
+				}
+
+				// Use virtual paths as the final tie breaker.
+				return PrimarySign * AssetViewItemA.GetItem().GetVirtualPath().Compare(AssetViewItemB.GetItem().GetVirtualPath()) < 0;
+			});
+		}
 	}
-	SortMethod.Empty();
 
 	//UE_LOG(LogContentBrowser, Warning/*VeryVerbose*/, TEXT("FAssetViewSortManager Sort Time: %0.4f seconds."), FPlatformTime::Seconds() - SortListStartTime);
 }
