Fixed columns appearing as sorted after you click on their headers in the Content Browser column view even when they are no longer the active primary or secondary sort columns.

diff --git a/Engine/Source/Editor/ContentBrowser/Private/SAssetView.cpp b/Engine/Source/Editor/ContentBrowser/Private/SAssetView.cpp
--- a/Engine/Source/Editor/ContentBrowser/Private/SAssetView.cpp
+++ b/Engine/Source/Editor/ContentBrowser/Private/SAssetView.cpp
@@ -6050,7 +6050,7 @@ EColumnSortMode::Type SAssetView::GetColumnSortMode(const FName ColumnId) const
 		}
 	}
 
-	static constexpr EColumnSortMode::Type DefaultSortMode = EColumnSortMode::Ascending;
+	static constexpr EColumnSortMode::Type DefaultSortMode = EColumnSortMode::None;
 	return DefaultSortMode;
 }
 
@@ -6100,7 +6100,12 @@ void SAssetView::SetSortParameters(
 	}
 
 	// Use specified sort mode OR get from the ColumnId
-	const EColumnSortMode::Type SortMode = InNewSortMode.Get(GetColumnSortMode(ColumnId));
+	EColumnSortMode::Type DefaultSortMode = GetColumnSortMode(ColumnId);
+	if (DefaultSortMode == EColumnSortMode::None)
+	{
+		DefaultSortMode = EColumnSortMode::Ascending;
+	}
+	const EColumnSortMode::Type SortMode = InNewSortMode.Get(DefaultSortMode);
 
 	OnSortColumnHeader(SortPriority, ColumnId, SortMode);
 }
