Fixed a case where clicking a Content Browser column header can take many seconds if you have a lot of items.  This is because we are figuring out which thumbnails to update based on the range of visible items from the previous sorting of the list, and those items may be very far apart in the new sorting of the list.

diff --git a/Engine/Source/Editor/ContentBrowser/Private/SAssetView.cpp b/Engine/Source/Editor/ContentBrowser/Private/SAssetView.cpp
--- a/Engine/Source/Editor/ContentBrowser/Private/SAssetView.cpp
+++ b/Engine/Source/Editor/ContentBrowser/Private/SAssetView.cpp
@@ -2215,7 +2215,9 @@ void SAssetView::Tick( const FGeometry& AllottedGeometry, const double InCurrent
 
 	CalculateThumbnailHintColorAndOpacity();
 
-	if (bPendingUpdateThumbnails)
+	// Wait for any pending refresh to complete, otherwise VisibleItems will be out of date.
+	// If we sorted a long list that can cause us to update thousands of thumbnails.
+	if (bPendingUpdateThumbnails && !IsPendingRefresh())
 	{
 		UpdateThumbnails();
 		bPendingUpdateThumbnails = false;
@@ -5014,6 +5016,21 @@ void SAssetView::FocusList() const
 	}
 }
 
+bool SAssetView::IsPendingRefresh()
+{
+	switch (GetCurrentViewType())
+	{
+	case EAssetViewType::List:
+		return ListView->IsPendingRefresh();
+	case EAssetViewType::Tile:
+		return TileView->IsPendingRefresh();
+	case EAssetViewType::Column:
+		return ColumnView->IsPendingRefresh();
+	default:
+		return false;
+	}
+}
+
 void SAssetView::RefreshList()
 {
 	switch ( GetCurrentViewType() )
diff --git a/Engine/Source/Editor/ContentBrowser/Public/SAssetView.h b/Engine/Source/Editor/ContentBrowser/Public/SAssetView.h
--- a/Engine/Source/Editor/ContentBrowser/Public/SAssetView.h
+++ b/Engine/Source/Editor/ContentBrowser/Public/SAssetView.h
@@ -761,6 +761,9 @@ private:
 	/** Set the keyboard focus to the correct list view that should be active */
 	UE_API void FocusList() const;
 
+	/** Return true if there is currently a refresh pending, false otherwise */
+	UE_API bool IsPendingRefresh();
+
 	/** Refreshes the list view to display any changes made to the non-filtered assets */
 	UE_API void RefreshList();
 
