Updated the asset view type column sort method to match what is displayed in the column.

diff --git a/Engine/Source/Editor/ContentBrowser/Private/AssetViewSortManager.cpp b/Engine/Source/Editor/ContentBrowser/Private/AssetViewSortManager.cpp
--- a/Engine/Source/Editor/ContentBrowser/Private/AssetViewSortManager.cpp
+++ b/Engine/Source/Editor/ContentBrowser/Private/AssetViewSortManager.cpp
@@ -74,7 +74,7 @@ public:
 
 		for (const TSharedPtr<FAssetViewItem>& AssetItem : AssetItems)
 		{
-			CachedValues.Add(GetClassName(AssetItem->GetItem()));
+			CachedValues.Emplace(AssetItem->GetItem());
 		}
 
 		return true;
@@ -82,22 +82,117 @@ public:
 
 	int32 Compare(int32 IndexA, int32 IndexB) const override
 	{
-		return CachedValues[IndexA].Compare(CachedValues[IndexB]);
+		return CachedValues[IndexA].CompareTo(CachedValues[IndexB]);
 	}
 
 	int32 Compare(const FAssetViewItem& A, const FAssetViewItem& B) const override
 	{
-		return GetClassName(A.GetItem()).Compare(GetClassName(B.GetItem()));
+		FContentBrowserItemDataAttributeValue AttributeValue;
+
+		FText ItemTypeDisplayNameA;
+		AttributeValue = A.GetItem().GetItemAttribute(ContentBrowserItemAttributes::ItemTypeDisplayName);
+		if (AttributeValue.IsValid())
+		{
+			ItemTypeDisplayNameA = AttributeValue.GetValue<FText>();
+		}
+
+		FText ItemTypeDisplayNameB;
+		AttributeValue = B.GetItem().GetItemAttribute(ContentBrowserItemAttributes::ItemTypeDisplayName);
+		if (AttributeValue.IsValid())
+		{
+			ItemTypeDisplayNameB = AttributeValue.GetValue<FText>();
+		}
+
+		FString ItemTypeNameA;
+		auto GetItemTypeNameA = [&A, &AttributeValue, &ItemTypeNameA]() -> const FString&
+		{
+			if (ItemTypeNameA.IsEmpty())
+			{
+				AttributeValue = A.GetItem().GetItemAttribute(ContentBrowserItemAttributes::ItemTypeName);
+				if (AttributeValue.IsValid())
+				{
+					ItemTypeNameA = AttributeValue.GetValue<FString>();
+				}
+			}
+			return ItemTypeNameA;
+		};
+
+		FString ItemTypeNameB;
+		auto GetItemTypeNameB = [&B, &AttributeValue, &ItemTypeNameB]() -> const FString&
+		{
+			if (ItemTypeNameB.IsEmpty())
+			{
+				AttributeValue = B.GetItem().GetItemAttribute(ContentBrowserItemAttributes::ItemTypeName);
+				if (AttributeValue.IsValid())
+				{
+					ItemTypeNameB = AttributeValue.GetValue<FString>();
+				}
+			}
+			return ItemTypeNameB;
+		};
+
+		return Compare(ItemTypeDisplayNameA, ItemTypeDisplayNameB, GetItemTypeNameA, GetItemTypeNameB);
 	}
 
 private:
-	static FName GetClassName(const FContentBrowserItem& Item)
+	template <typename GetItemTypeNameAType, typename GetItemTypeNameBType>
+	static int32 Compare(const FText& ItemTypeDisplayNameA, const FText& ItemTypeDisplayNameB, GetItemTypeNameAType GetItemTypeNameA, GetItemTypeNameBType GetItemTypeNameB)
 	{
-		FContentBrowserItemDataAttributeValue ClassValue = Item.GetItemAttribute(ContentBrowserItemAttributes::ItemTypeName);
-		return ClassValue.IsValid() ? ClassValue.GetValue<FName>() : FName();
+		int32 Result = 0;
+
+		if (!ItemTypeDisplayNameA.IsEmpty())
+		{
+			if (!ItemTypeDisplayNameB.IsEmpty())
+			{
+				Result = ItemTypeDisplayNameA.CompareToCaseIgnored(ItemTypeDisplayNameB);
+			}
+			else
+			{
+				Result = FTextComparison::CompareToCaseIgnored(ItemTypeDisplayNameA.ToString(), GetItemTypeNameB());
+			}
+		}
+		else if (!ItemTypeDisplayNameB.IsEmpty())
+		{
+			Result = FTextComparison::CompareToCaseIgnored(GetItemTypeNameA(), ItemTypeDisplayNameB.ToString());
+		}
+
+		if (Result != 0)
+		{
+			return Result;
+		}
+
+		return FTextComparison::CompareToCaseIgnored(GetItemTypeNameA(), GetItemTypeNameB());
 	}
 
-	TArray<FName> CachedValues;
+	struct FCachedValue
+	{
+	public:
+		explicit FCachedValue(const FContentBrowserItem& Item)
+		{
+			FContentBrowserItemDataAttributeValue AttributeValue = Item.GetItemAttribute(ContentBrowserItemAttributes::ItemTypeDisplayName);
+			if (AttributeValue.IsValid())
+			{
+				ItemTypeDisplayName = AttributeValue.GetValue<FText>();
+			}
+
+			AttributeValue = Item.GetItemAttribute(ContentBrowserItemAttributes::ItemTypeName);
+			if (AttributeValue.IsValid())
+			{
+				ItemTypeName = AttributeValue.GetValue<FString>();
+			}
+		}
+
+		int32 CompareTo(const FCachedValue& Other) const
+		{
+			return Compare(ItemTypeDisplayName, Other.ItemTypeDisplayName, [this]() -> const FString& { return ItemTypeName; }, [&Other]() -> const FString& { return Other.ItemTypeName; });
+		}
+
+	private:
+		FText ItemTypeDisplayName;
+		FString ItemTypeName;
+	};
+
+	TArray<FCachedValue> CachedValues;
 };
 
 struct FCompareFAssetItemByDiskSize final : public FCompareFAssetItemBase
