Added shadow styling options to SPathView, STreeView and SContentBrowserSourceTree. Removed shadows from file path section of the content browser

#rnx

diff --git a/Engine/Source/Editor/ContentBrowser/Private/SContentBrowser.cpp b/Engine/Source/Editor/ContentBrowser/Private/SContentBrowser.cpp
--- a/Engine/Source/Editor/ContentBrowser/Private/SContentBrowser.cpp
+++ b/Engine/Source/Editor/ContentBrowser/Private/SContentBrowser.cpp
@@ -1170,7 +1170,9 @@ TSharedRef<SWidget> SContentBrowser::CreatePathView(const FContentBrowserConfig*
 		.PluginPathFilters(PluginPathFilters)
 		.OwningContentBrowserName(InstanceName)
 		.AllowClassesFolder(Config != nullptr ? Config->bCanShowClasses : true)
-		.CanShowDevelopersFolder(Config != nullptr ? Config->bCanShowDevelopersFolder : true);
+		.CanShowDevelopersFolder(Config != nullptr ? Config->bCanShowDevelopersFolder : true)
+		.ShadowBoxStyle(&FAppStyle::Get().GetWidgetStyle<FScrollBoxStyle>("ScrollBoxNoShadow"))
+		.bEnableShadowBoxStyle(true);
 
 	TSharedRef<UE::Editor::ContentBrowser::Private::SContentBrowserSourceTreeArea> ViewWidget =
 		SAssignNew(
@@ -1181,7 +1183,9 @@ TSharedRef<SWidget> SContentBrowser::CreatePathView(const FContentBrowserConfig*
 			LegacyContentSourceWidgets->PathViewPtr.ToSharedRef())
 			.Label(FText::FromString(FApp::GetProjectName()))
 			.ExpandedByDefault(true) // The path area, unlike the favorites and collection areas, are expanded by default (unless overridden)
-			.OnExpansionChanged(this, &SContentBrowser::SetPathViewExpanded);
+			.OnExpansionChanged(this, &SContentBrowser::SetPathViewExpanded)
+			.ShadowBoxStyle(&FAppStyle::Get().GetWidgetStyle<FScrollBorderStyle>("ScrollBorderNoShadow"))
+			.bEnableShadowBoxStyle(true);
 
 	return ViewWidget;
 }
diff --git a/Engine/Source/Editor/ContentBrowser/Private/SPathView.cpp b/Engine/Source/Editor/ContentBrowser/Private/SPathView.cpp
--- a/Engine/Source/Editor/ContentBrowser/Private/SPathView.cpp
+++ b/Engine/Source/Editor/ContentBrowser/Private/SPathView.cpp
@@ -1023,7 +1023,9 @@ void SPathView::Construct( const FArguments& InArgs )
 					  .AllowInvisibleItemSelection(true)
 					  .OnSelectionChanged(this, &SPathView::TreeSelectionChanged)
 					  .OnContextMenuOpening(this, &SPathView::MakePathViewContextMenu)
-					  .ClearSelectionOnClick(false);
+					  .ClearSelectionOnClick(false)
+					  .ShadowBoxStyle(InArgs._ShadowBoxStyle)
+					  .bEnableShadowBoxStyle(InArgs._bEnableShadowBoxStyle);
 
 	SearchPtr = InArgs._ExternalSearch;
 	if (!SearchPtr)
diff --git a/Engine/Source/Editor/ContentBrowser/Private/SPathView.h b/Engine/Source/Editor/ContentBrowser/Private/SPathView.h
--- a/Engine/Source/Editor/ContentBrowser/Private/SPathView.h
+++ b/Engine/Source/Editor/ContentBrowser/Private/SPathView.h
@@ -72,7 +72,8 @@ public:
 	DECLARE_DELEGATE( FOnFrontendPluginFilterChanged );
 
 	SLATE_BEGIN_ARGS( SPathView )
-		: _InitialCategoryFilter(EContentBrowserItemCategoryFilter::IncludeAll)
+		: _bEnableShadowBoxStyle(false)
+		, _InitialCategoryFilter(EContentBrowserItemCategoryFilter::IncludeAll)
 		, _FocusSearchBoxWhenOpened(true)
 		, _ShowTreeTitle(false)
 		, _SearchBarVisibility(EVisibility::Visible)
@@ -91,6 +92,12 @@ public:
 		/** Content displayed to the left of the search bar */
 		SLATE_NAMED_SLOT( FArguments, SearchContent )
 
+		/** Shadow box style used for the STreeView if bEnableShadowBoxStyle is true */
+		SLATE_ARGUMENT(const FScrollBoxStyle*, ShadowBoxStyle)
+		
+		/** If true, overriding shadow box style of STreeView is enabled */
+		SLATE_ARGUMENT(bool, bEnableShadowBoxStyle)
+
 		/** Called when a tree paths was selected */
 		SLATE_EVENT( FOnContentBrowserItemSelectionChanged, OnItemSelectionChanged )
 
diff --git a/Engine/Source/Editor/ContentBrowser/Private/Widgets/SContentBrowserSourceTree.cpp b/Engine/Source/Editor/ContentBrowser/Private/Widgets/SContentBrowserSourceTree.cpp
--- a/Engine/Source/Editor/ContentBrowser/Private/Widgets/SContentBrowserSourceTree.cpp
+++ b/Engine/Source/Editor/ContentBrowser/Private/Widgets/SContentBrowserSourceTree.cpp
@@ -422,6 +422,7 @@ namespace UE::Editor::ContentBrowser::Private
 						[
 							BodyScrollableWidget->GetScrollWidget()
 						]
+						.Style(InArgs._bEnableShadowBoxStyle ? InArgs._ShadowBoxStyle : &FCoreStyle::Get().GetWidgetStyle<FScrollBorderStyle>("ScrollBorder"))
 					]
 				]
 
diff --git a/Engine/Source/Editor/ContentBrowser/Private/Widgets/SContentBrowserSourceTree.h b/Engine/Source/Editor/ContentBrowser/Private/Widgets/SContentBrowserSourceTree.h
--- a/Engine/Source/Editor/ContentBrowser/Private/Widgets/SContentBrowserSourceTree.h
+++ b/Engine/Source/Editor/ContentBrowser/Private/Widgets/SContentBrowserSourceTree.h
@@ -137,6 +137,7 @@ namespace UE::Editor::ContentBrowser::Private
 	public:
 		SLATE_BEGIN_ARGS(SContentBrowserSourceTreeArea)
 			: _IsEmpty(false)
+			, _bEnableShadowBoxStyle(false)
 		{ }
 			/** Label to display in the header, ie. "Favorites" */
 			SLATE_ATTRIBUTE(FText, Label)
@@ -153,6 +154,12 @@ namespace UE::Editor::ContentBrowser::Private
 			/** (Optional) default state for area expansion. */
 			SLATE_ARGUMENT(bool, ExpandedByDefault)
 
+			/** (Optional) override for STreeView ScrollBorder style */
+			SLATE_STYLE_ARGUMENT(FScrollBorderStyle, ShadowBoxStyle)
+			
+			/** (Optional) If true, ScrollBorder will use the style set in BodyScrollBorderStyle */
+			SLATE_ARGUMENT(bool, bEnableShadowBoxStyle)
+
 			/** Called when the area is expanded or collapsed */
 			SLATE_EVENT(FOnBooleanValueChanged, OnExpansionChanged)
 		SLATE_END_ARGS()
diff --git a/Engine/Source/Runtime/Slate/Public/Widgets/Views/STreeView.h b/Engine/Source/Runtime/Slate/Public/Widgets/Views/STreeView.h
--- a/Engine/Source/Runtime/Slate/Public/Widgets/Views/STreeView.h
+++ b/Engine/Source/Runtime/Slate/Public/Widgets/Views/STreeView.h
@@ -139,6 +139,8 @@ public:
 		, _ScrollBarStyle(&FAppStyle::Get().GetWidgetStyle<FScrollBarStyle>("ScrollBar"))
 		, _PreventThrottling(false)
 		, _WheelScrollMultiplier(GetGlobalScrollAmount())
+		, _ShadowBoxStyle()
+		, _bEnableShadowBoxStyle(false)
 		, _OnItemToString_Debug()
 		, _OnEnteredBadState()
 		, _HandleGamepadEvents(true)
@@ -215,6 +217,10 @@ public:
 
 		SLATE_ARGUMENT( float, WheelScrollMultiplier );
 
+		SLATE_STYLE_ARGUMENT(FScrollBoxStyle, ShadowBoxStyle)
+
+		SLATE_ARGUMENT(bool, bEnableShadowBoxStyle)
+
 		/** Assign this to get more diagnostics from the list view. */
 		SLATE_EVENT(FOnItemToString_Debug, OnItemToString_Debug)
 
@@ -293,6 +299,9 @@ public:
 
 		this->SetStyle(InArgs._TreeViewStyle);
 
+		this->SetShadowBoxStyle(InArgs._ShadowBoxStyle);
+		this->bShouldUseShadowBoxStyle = InArgs._bEnableShadowBoxStyle;
+
 		this->MaxPinnedItems = InArgs._MaxPinnedItems;
 		this->DefaultMaxPinnedItems = InArgs._MaxPinnedItems;
 
diff --git a/Engine/Source/Runtime/SlateCore/Private/Styling/StarshipCoreStyle.cpp b/Engine/Source/Runtime/SlateCore/Private/Styling/StarshipCoreStyle.cpp
--- a/Engine/Source/Runtime/SlateCore/Private/Styling/StarshipCoreStyle.cpp
+++ b/Engine/Source/Runtime/SlateCore/Private/Styling/StarshipCoreStyle.cpp
@@ -756,6 +756,11 @@ TSharedRef<ISlateStyle> FStarshipCoreStyle::Create()
 			.SetTopShadowBrush(BOX_BRUSH("Common/ScrollBorderShadowTop", FVector2f(16.f, 8.f), FMargin(0.5f, 1.f, 0.5f, 0.f)))
 			.SetBottomShadowBrush(BOX_BRUSH("Common/ScrollBorderShadowBottom", FVector2f(16.f, 8.f), FMargin(0.5f, 0.f, 0.5f, 1.f)))
 		);
+
+		Style->Set("ScrollBorderNoShadow", FScrollBorderStyle()
+			.SetTopShadowBrush(FSlateNoResource())
+			.SetBottomShadowBrush(FSlateNoResource())
+		);
 	}
 
 	// SWindow defaults...
