EditroFunctionalTests refactoring
- moved common methods to utils
- moved content/tests into appropriate folders
- reduced test class variables amount
- unified CB open/close in tests
- reduced AutomationDriver usage

#rnx
#virtualized

diff --git a/Engine/Source/Developer/ToolWidgets/Private/Filters/SFilterSearchBox.cpp b/Engine/Source/Developer/ToolWidgets/Private/Filters/SFilterSearchBox.cpp
--- a/Engine/Source/Developer/ToolWidgets/Private/Filters/SFilterSearchBox.cpp
+++ b/Engine/Source/Developer/ToolWidgets/Private/Filters/SFilterSearchBox.cpp
@@ -2,7 +2,6 @@
 
 #include "Filters/SFilterSearchBox.h"
 #include "Framework/Text/TextLayout.h"
-#include "Framework/MetaData/DriverMetaData.h"
 #include "Widgets/Input/SSearchBox.h"
 #include "Widgets/Input/SMenuAnchor.h"
 #include "Widgets/Views/SListView.h"
@@ -98,7 +97,7 @@ void SFilterSearchBoxImpl::Construct(const FArguments& InArgs)
 		[
 			// Button to save the currently occuring search 
 			SNew(SButton)
-			.AddMetaData(FDriverMetaData::Id(TEXT("CreateFilterFromCurrentSearchButton")))
+			.AddMetaData<FTagMetaData>(FTagMetaData(TEXT("CreateFilterFromCurrentSearchButton")))
 			.ContentPadding(0.f)
 			.Visibility_Lambda([this]()
 			{
diff --git a/Engine/Source/Editor/ContentBrowser/Private/AssetViewWidgets.cpp b/Engine/Source/Editor/ContentBrowser/Private/AssetViewWidgets.cpp
--- a/Engine/Source/Editor/ContentBrowser/Private/AssetViewWidgets.cpp
+++ b/Engine/Source/Editor/ContentBrowser/Private/AssetViewWidgets.cpp
@@ -34,7 +34,6 @@
 #include "EditorFramework/AssetImportData.h"
 #include "Fonts/FontMeasure.h"
 #include "Framework/Application/SlateApplication.h"
-#include "Framework/MetaData/DriverMetaData.h"
 #include "Framework/Text/TextLayout.h"
 #include "HAL/IConsoleManager.h"
 #include "HAL/PlatformApplicationMisc.h"
@@ -2739,7 +2738,7 @@ void SAssetTileItem::Construct( const FArguments& InArgs )
 										.OverflowPolicy(ETextOverflowPolicy::Ellipsis)
 										.Text(this, &SAssetTileItem::GetAssetClassText)
 										.ColorAndOpacity(this, &SAssetTileItem::GetAssetClassLabelTextColor)
-										.AddMetaData(FDriverMetaData::Id(TEXT("AssetTileClassTextBlock")))
+										.AddMetaData<FTagMetaData>(FTagMetaData(TEXT("AssetTileClassTextBlock")))
 									]
 								]
 							]
@@ -2818,7 +2817,7 @@ void SAssetTileItem::Construct( const FArguments& InArgs )
 										.OverflowPolicy(ETextOverflowPolicy::Ellipsis)
 										.Text(this, &SAssetTileItem::GetAssetClassText)
 										.ColorAndOpacity(this, &SAssetTileItem::GetAssetClassLabelTextColor)
-										.AddMetaData(FDriverMetaData::Id(TEXT("AssetTileClassTextBlock")))
+										.AddMetaData<FTagMetaData>(FTagMetaData(TEXT("AssetTileClassTextBlock")))
 									]
 
 									+SHorizontalBox::Slot()
