The asset view will now display, filter and sort by Verse paths, object paths or package paths in the path column if they are available for the current item.
Did some minor filtering optimizations to mitigate the added costs.

#rnx

diff --git a/Engine/Plugins/Editor/ContentBrowser/ContentBrowserAssetDataSource/Source/ContentBrowserAssetDataSource/Private/ContentBrowserAssetDataCore.cpp b/Engine/Plugins/Editor/ContentBrowser/ContentBrowserAssetDataSource/Source/ContentBrowserAssetDataSource/Private/ContentBrowserAssetDataCore.cpp
--- a/Engine/Plugins/Editor/ContentBrowser/ContentBrowserAssetDataSource/Source/ContentBrowserAssetDataSource/Private/ContentBrowserAssetDataCore.cpp
+++ b/Engine/Plugins/Editor/ContentBrowser/ContentBrowserAssetDataSource/Source/ContentBrowserAssetDataSource/Private/ContentBrowserAssetDataCore.cpp
@@ -1453,18 +1453,11 @@ void AppendAssetExportText(const FAssetData& AssetData, FString& InOutStr)
 
 void AppendAssetObjectPath(const FAssetData& AssetData, FString& InOutStr)
 {
-	if (InOutStr.IsEmpty())
-	{
-		InOutStr = AssetData.GetObjectPathString();
-	}
-	else
+	if (!InOutStr.IsEmpty())
 	{
-		if (InOutStr.Len() > 0)
-		{
-			InOutStr += LINE_TERMINATOR;
-		}
-		InOutStr += AssetData.GetObjectPathString();
+		InOutStr += LINE_TERMINATOR;
 	}
+	AssetData.AppendObjectPath(InOutStr);
 }
 
 void AppendAssetPackageName(const FAssetData& AssetData, FString& InOutStr)
@@ -1844,10 +1837,10 @@ bool GetAssetFileItemAttribute(const FContentBrowserAssetFileItemDataPayload& In
 		{
 			if (TSharedPtr<IAssetTypeActions> AssetTypeActions = InAssetPayload.GetAssetTypeActions())
 			{
-				const FText AssetDisplayName = AssetTypeActions->GetDisplayNameFromAssetData(InAssetPayload.GetAssetData());
+				FText AssetDisplayName = AssetTypeActions->GetDisplayNameFromAssetData(InAssetPayload.GetAssetData());
 				if (!AssetDisplayName.IsEmpty())
 				{
-					OutAttributeValue.SetValue(AssetDisplayName);
+					OutAttributeValue.SetValue(MoveTemp(AssetDisplayName));
 				}
 				else
 				{
@@ -1864,10 +1857,10 @@ bool GetAssetFileItemAttribute(const FContentBrowserAssetFileItemDataPayload& In
 			{
 				FWarnIfAssetsLoadedInScope WarnIfAssetLoadedInScope({ InAssetPayload.GetAssetData() });
 				
-				const FText AssetDescription = AssetDefinition->GetAssetDescription(InAssetPayload.GetAssetData());
+				FText AssetDescription = AssetDefinition->GetAssetDescription(InAssetPayload.GetAssetData());
 				if (!AssetDescription.IsEmpty())
 				{
-					OutAttributeValue.SetValue(AssetDescription);
+					OutAttributeValue.SetValue(MoveTemp(AssetDescription));
 					return true;
 				}
 			}
diff --git a/Engine/Plugins/Editor/ContentBrowser/ContentBrowserClassDataSource/Source/ContentBrowserClassDataSource/Private/ContentBrowserClassDataCore.cpp b/Engine/Plugins/Editor/ContentBrowser/ContentBrowserClassDataSource/Source/ContentBrowserClassDataSource/Private/ContentBrowserClassDataCore.cpp
--- a/Engine/Plugins/Editor/ContentBrowser/ContentBrowserClassDataSource/Source/ContentBrowserClassDataSource/Private/ContentBrowserClassDataCore.cpp
+++ b/Engine/Plugins/Editor/ContentBrowser/ContentBrowserClassDataSource/Source/ContentBrowserClassDataSource/Private/ContentBrowserClassDataCore.cpp
@@ -334,22 +334,18 @@ bool AppendObjectPathFileItemReference(const FContentBrowserClassFileItemDataPay
 	{
 		InOutStr += LINE_TERMINATOR;
 	}
-	InOutStr += InClassPayload.GetAssetData().GetObjectPathString();
+	InClassPayload.GetAssetData().AppendObjectPath(InOutStr);
 	return true;
 }
 
 bool AppendPackageNameItemReference(const FContentBrowserClassFileItemDataPayload& InClassPayload, FString& InOutStr)
 {
-	if (const UPackage* Package = InClassPayload.GetAssetData().GetPackage())
+	if (InOutStr.Len() > 0)
 	{
-		if (InOutStr.Len() > 0)
-		{
-			InOutStr += LINE_TERMINATOR;
-		}
-		InOutStr += Package->GetPathName();
-		return true;
+		InOutStr += LINE_TERMINATOR;
 	}
-	return false;
+	InClassPayload.GetAssetData().PackageName.AppendString(InOutStr);
+	return true;
 }
 
 bool GetItemPhysicalPath(const UContentBrowserDataSource* InOwnerDataSource, const FContentBrowserItemData& InItem, FString& OutDiskPath)
@@ -453,7 +449,7 @@ void GetGenericItemAttribute(const FName InTagKey, const FString& InTagValue, co
 		FText TmpText;
 		if (FTextStringHelper::ReadFromBuffer(*InTagValue, TmpText))
 		{
-			OutAttributeValue.SetValue(TmpText);
+			OutAttributeValue.SetValue(MoveTemp(TmpText));
 		}
 	}
 	if (!OutAttributeValue.IsValid())
diff --git a/Engine/Source/Editor/ContentBrowser/Private/AssetTextFilter.cpp b/Engine/Source/Editor/ContentBrowser/Private/AssetTextFilter.cpp
--- a/Engine/Source/Editor/ContentBrowser/Private/AssetTextFilter.cpp
+++ b/Engine/Source/Editor/ContentBrowser/Private/AssetTextFilter.cpp
@@ -355,6 +355,11 @@ FCompiledAssetTextFilter FCompiledAssetTextFilter::CloneForThreading()
 }
 
 bool FCompiledAssetTextFilter::PassesFilter(FAssetFilterType InItem)
+{
+	return PassesFilter(InItem, {});
+}
+
+bool FCompiledAssetTextFilter::PassesFilter(FAssetFilterType InItem, FStringView InItemVersePath)
 {
 	using namespace AssetTextFilter;
 	FReadScopeLock Guard(CustomHandlerLock);
@@ -363,9 +368,13 @@ bool FCompiledAssetTextFilter::PassesFilter(FAssetFilterType InItem)
     TextBuffer.Reset();
 	AssetCollectionNames.Reset();
 
+	bIsFile = AssetPtr->IsFile();
+
 	int32 DisplayNameLen = 0;
 	int32 AssetPathLen = 0;
+	int32 AssetVirtualPathLen = 0;
 	int32 ExportTextPathLen = 0;
+	int32 VersePathLen = 0;
 
     TextBuffer.Append(AssetPtr->GetDisplayName().ToString());
 	DisplayNameLen = TextBuffer.Len();
@@ -373,23 +382,38 @@ bool FCompiledAssetTextFilter::PassesFilter(FAssetFilterType InItem)
 	if (bIncludeAssetPath)
 	{
 		int32 TextBufferLenBefore = TextBuffer.Len();
-		// Get the full asset path, and also split it so we can compare each part in the filter
-		AssetPtr->GetVirtualPath().AppendString(TextBuffer);
-        
-        // chop off item name if necessary
+
+		if (bIsFile)
 		{
-			int32 LastSlashIndex = INDEX_NONE;
-			int32 LastDotIndex = INDEX_NONE;
-			if (TextBuffer.FindLastChar(TEXT('/'), LastSlashIndex) && TextBuffer.FindLastChar(TEXT('.'), LastDotIndex)
-				&& LastDotIndex > LastSlashIndex && LastSlashIndex >= DisplayNameLen)
+			if (AssetPtr->AppendItemObjectPath(TextBuffer))
+			{
+				AssetPathLen = TextBuffer.Len() - TextBufferLenBefore;
+
+				TextBufferLenBefore = TextBuffer.Len();
+			}
+
+			AssetPtr->GetVirtualPath().AppendString(TextBuffer);
+
+			AssetVirtualPathLen = TextBuffer.Len() - TextBufferLenBefore;
+		}
+		else
+		{
+			FName PackagePath;
+			if (AssetPtr->Legacy_TryGetPackagePath(PackagePath))
 			{
-				TextBuffer.LeftInline(LastDotIndex, EAllowShrinking::No);
-				checkf(TextBuffer.Len() >= TextBufferLenBefore, TEXT("If TextBuffer is smaller than it used to be, then DisplayNameLen is now invalid. That should never happen."));
+				PackagePath.AppendString(TextBuffer);
+
+				AssetPathLen = TextBuffer.Len() - TextBufferLenBefore;
+
+				TextBufferLenBefore = TextBuffer.Len();
 			}
+
+			AssetPtr->GetVirtualPath().AppendString(TextBuffer);
+
+			AssetVirtualPathLen = TextBuffer.Len() - TextBufferLenBefore;
 		}
-		AssetPathLen = TextBuffer.Len() - TextBufferLenBefore;
 
-		if (bIncludeClassName && !AssetPtr->IsFolder())
+		if (bIncludeClassName && bIsFile)
 		{
 			TextBufferLenBefore = TextBuffer.Len();
 			// Get the full export text path as people sometimes search by copying this (requires class and asset path
@@ -400,6 +424,12 @@ bool FCompiledAssetTextFilter::PassesFilter(FAssetFilterType InItem)
 		}
 	}
 
+	if (!InItemVersePath.IsEmpty())
+	{
+		int32 TextBufferLenBefore = TextBuffer.Len();
+		TextBuffer += InItemVersePath;
+		VersePathLen = TextBuffer.Len() - TextBufferLenBefore;
+	}
 
     // Convert entire text buffer to uppercase
     // Assumes that for a TCHAR buffer this won't change number of code points 
@@ -408,11 +438,15 @@ bool FCompiledAssetTextFilter::PassesFilter(FAssetFilterType InItem)
 	const TCHAR* Cursor = *TextBuffer;
 	AssetDisplayName = FStringView(Cursor, DisplayNameLen);
 	Cursor += DisplayNameLen;
-	AssetFullPath = FStringView(Cursor, AssetPathLen);
+	AssetFullPath = FStringView(Cursor, AssetPathLen).TrimStartAndEnd();
 	Cursor += AssetPathLen;
+	AssetVirtualPath = FStringView(Cursor, AssetVirtualPathLen);
+	Cursor += AssetVirtualPathLen;
 	AssetExportTextPath  = FStringView(Cursor, ExportTextPathLen);
 	AssetExportTextPath.TrimStartAndEndInline(); // Backends try to separate export text paths with newlines 
 	Cursor += ExportTextPathLen;
+	AssetVersePath = FStringView(Cursor, VersePathLen);
+	Cursor += VersePathLen;
 
 	// Executed last because we pass *this to DynamicQuery and it would try to access StringView before initialization.
 	if ((CollectionContainers.IsValid() || ReferencedDynamicCollections.IsValid()) && bIncludeCollectionNames)
@@ -480,13 +514,16 @@ bool FCompiledAssetTextFilter::TestBasicStringExpression(
 
 	if (bIncludeAssetPath)
 	{
-		if (InValue.CompareFStringView(AssetFullPath, InTextComparisonMode))
+		// Mimic how we select the value to display for the path column.
+		const FStringView AssetPath = !AssetVersePath.IsEmpty() ? AssetVersePath : (!AssetFullPath.IsEmpty() ? AssetFullPath : AssetVirtualPath);
+
+		if (InValue.CompareFStringView(AssetPath, InTextComparisonMode))
 		{
 			return true;
 		}
 
 		bool bSuccess = false;
-		UE::String::ParseTokens(AssetFullPath, TEXTVIEW("/"), [&bSuccess, &InValue, InTextComparisonMode](FStringView Element){
+		UE::String::ParseTokens(AssetPath, TEXTVIEW("/"), [&bSuccess, &InValue, InTextComparisonMode](FStringView Element){
 			bSuccess = bSuccess || InValue.CompareFStringView(Element, InTextComparisonMode);
 		});
 		if (bSuccess)
@@ -541,6 +578,12 @@ bool FCompiledAssetTextFilter::TestComplexExpression(
 		}
 	}
 
+	auto ResetTextFilterStringTo = [this](const auto& Value) -> const FTextFilterString&
+	{
+		TextFilterString = Value;
+		return TextFilterString;
+	};
+
 	// Special case for the asset name, as this isn't contained within the asset registry meta-data
 	if (InKey == NameKeyName)
 	{
@@ -551,7 +594,7 @@ bool FCompiledAssetTextFilter::TestComplexExpression(
 			return false;
 		}
 
-		const bool bIsMatch = TextFilterUtils::TestBasicStringExpression(AssetPtr->GetItemName(), InValue, InTextComparisonMode);
+		const bool bIsMatch = TextFilterUtils::TestBasicStringExpression(ResetTextFilterStringTo(AssetPtr->GetItemName()), InValue, InTextComparisonMode);
 		return (InComparisonOperation == ETextFilterComparisonOperation::Equal) ? bIsMatch : !bIsMatch;
 	}
 
@@ -565,21 +608,91 @@ bool FCompiledAssetTextFilter::TestComplexExpression(
 			return false;
 		}
 
-		// If the comparison mode is partial, then we only need to test the ObjectPath as that contains the other two as
-		// sub-strings
+		FString ObjectPathBuffer;
+		FStringView ObjectPath;
+
+		// For the path key we check against all possible path types, not just the value displayed in the path column.
 		bool bIsMatch = false;
-		if (InTextComparisonMode == ETextFilterTextComparisonMode::Partial)
+
+		// AssetVersePath is always set if available.
+		if (!AssetVersePath.IsEmpty() && TextFilterUtils::TestBasicStringExpression(ResetTextFilterStringTo(AssetVersePath), InValue, InTextComparisonMode))
 		{
-			bIsMatch =
-				TextFilterUtils::TestBasicStringExpression(AssetPtr->GetVirtualPath(), InValue, InTextComparisonMode);
+			bIsMatch = true;
+		}
+		else if (bIncludeAssetPath)
+		{
+			// Other cached paths are available.
+
+			if (TextFilterUtils::TestBasicStringExpression(ResetTextFilterStringTo(AssetVirtualPath), InValue, InTextComparisonMode))
+			{
+				bIsMatch = true;
+			}
+			else if (!AssetFullPath.IsEmpty())
+			{
+				if (TextFilterUtils::TestBasicStringExpression(ResetTextFilterStringTo(AssetFullPath), InValue, InTextComparisonMode))
+				{
+					bIsMatch = true;
+				}
+				else if (bIsFile)
+				{
+					ObjectPath = AssetFullPath;
+				}
+			}
 		}
 		else
 		{
-			bIsMatch =
-				TextFilterUtils::TestBasicStringExpression(AssetPtr->GetVirtualPath(), InValue, InTextComparisonMode)
-				|| (!AssetFullPath.IsEmpty()
-					&& TextFilterUtils::TestBasicStringExpression(FString(AssetFullPath), InValue, InTextComparisonMode));
+			// Other cached paths are NOT available.
+
+			if (TextFilterUtils::TestBasicStringExpression(ResetTextFilterStringTo(AssetPtr->GetVirtualPath()), InValue, InTextComparisonMode))
+			{
+				bIsMatch = true;
+			}
+			else
+			{
+				if (bIsFile)
+				{
+					if (AssetPtr->AppendItemObjectPath(ObjectPathBuffer))
+					{
+						if (TextFilterUtils::TestBasicStringExpression(ResetTextFilterStringTo(ObjectPathBuffer), InValue, InTextComparisonMode))
+						{
+							bIsMatch = true;
+						}
+						else
+						{
+							ObjectPath = ObjectPathBuffer;
+						}
+					}
+				}
+				else
+				{
+					FName PackagePath;
+					if (AssetPtr->Legacy_TryGetPackagePath(PackagePath))
+					{
+						if (TextFilterUtils::TestBasicStringExpression(ResetTextFilterStringTo(PackagePath), InValue, InTextComparisonMode))
+						{
+							bIsMatch = true;
+						}
+					}
+				}
+			}
+		}
+
+		// If there is no match and the path has an object name, trim that off and try again with the package name, but only for comparison modes where it makes a difference.
+		if (!bIsMatch &&
+			!ObjectPath.IsEmpty() &&
+			(InTextComparisonMode == ETextFilterTextComparisonMode::Exact || InTextComparisonMode == ETextFilterTextComparisonMode::EndsWith))
+		{
+			int32 LastSlashIndex = INDEX_NONE;
+			int32 LastDotIndex = INDEX_NONE;
+			if (ObjectPath.FindLastChar(TEXT('/'), LastSlashIndex) &&
+				ObjectPath.FindLastChar(TEXT('.'), LastDotIndex) &&
+				LastDotIndex > LastSlashIndex &&
+				TextFilterUtils::TestBasicStringExpression(ResetTextFilterStringTo(ObjectPath.Left(LastDotIndex)), InValue, InTextComparisonMode))
+			{
+				bIsMatch = true;
+			}
 		}
+
 		return (InComparisonOperation == ETextFilterComparisonOperation::Equal) ? bIsMatch : !bIsMatch;
 	}
 
@@ -596,7 +709,7 @@ bool FCompiledAssetTextFilter::TestComplexExpression(
 		const FContentBrowserItemDataAttributeValue ClassValue = AssetPtr->GetItemAttribute(NAME_Class);
 		const bool bIsMatch =
 			ClassValue.IsValid()
-			&& TextFilterUtils::TestBasicStringExpression(ClassValue.GetValue<FName>(), InValue, InTextComparisonMode);
+			&& TextFilterUtils::TestBasicStringExpression(ResetTextFilterStringTo(ClassValue.GetValue<FName>()), InValue, InTextComparisonMode);
 		return (InComparisonOperation == ETextFilterComparisonOperation::Equal) ? bIsMatch : !bIsMatch;
 	}
 
@@ -613,7 +726,7 @@ bool FCompiledAssetTextFilter::TestComplexExpression(
 		bool bFoundMatch = false;
 		for (const FName& AssetCollectionName : AssetCollectionNames)
 		{
-			if (TextFilterUtils::TestBasicStringExpression(AssetCollectionName, InValue, InTextComparisonMode))
+			if (TextFilterUtils::TestBasicStringExpression(ResetTextFilterStringTo(AssetCollectionName), InValue, InTextComparisonMode))
 			{
 				bFoundMatch = true;
 				break;
diff --git a/Engine/Source/Editor/ContentBrowser/Private/AssetViewSortManager.cpp b/Engine/Source/Editor/ContentBrowser/Private/AssetViewSortManager.cpp
--- a/Engine/Source/Editor/ContentBrowser/Private/AssetViewSortManager.cpp
+++ b/Engine/Source/Editor/ContentBrowser/Private/AssetViewSortManager.cpp
@@ -54,12 +54,9 @@ public:
 
 	int32 Compare(const FAssetViewItem& A, const FAssetViewItem& B) const override
 	{
-		const FContentBrowserItemData* AItemData = A.GetItem().GetPrimaryInternalItem();
-		const FContentBrowserItemData* BItemData = B.GetItem().GetPrimaryInternalItem();
 		// TODO: Have an option to sort by display name? It's slower, but more correct for non-English languages
-		//return (AItemData && BItemData) ? AItemData->GetDisplayName().CompareTo(BItemData->GetDisplayName()) : 0;
-		//return (AItemData && BItemData) ? FAssetViewSortManager::CompareWithNumericSuffix(FNameBuilder(AItemData->GetItemName()).ToView(), FNameBuilder(BItemData->GetItemName()).ToView()) : 0;
-		return (AItemData && BItemData) ? UE::ComparisonUtility::CompareWithNumericSuffix(AItemData->GetDisplayName().ToString(), BItemData->GetDisplayName().ToString()) : 0;
+		//return UE::ComparisonUtility::CompareWithNumericSuffix(FNameBuilder(A.GetItem().GetItemName()).ToView(), FNameBuilder(B.GetItem().GetItemName()).ToView());
+		return UE::ComparisonUtility::CompareWithNumericSuffix(A.GetItem().GetDisplayName().ToString(), B.GetItem().GetDisplayName().ToString());
 	}
 };
 
@@ -244,9 +241,7 @@ public:
 
 	int32 Compare(const FAssetViewItem& A, const FAssetViewItem& B) const override
 	{
-		const FContentBrowserItemData* AItemData = A.GetItem().GetPrimaryInternalItem();
-		const FContentBrowserItemData* BItemData = B.GetItem().GetPrimaryInternalItem();
-		return (AItemData && BItemData) ? UE::ComparisonUtility::CompareWithNumericSuffix(AItemData->GetVirtualPath(), BItemData->GetVirtualPath()) : 0;
+		return UE::ComparisonUtility::CompareWithNumericSuffix(A.GetAssetPathText().ToString(), B.GetAssetPathText().ToString());
 	}
 };
 
@@ -833,7 +828,7 @@ void FAssetViewSortManager::ExportColumnsToCSV(TArray<TSharedPtr<FAssetViewItem>
 			}
 			else if (Column == PathColumnId)
 			{
-				ValueString = AssetItemData->GetVirtualPath().ToString();
+				ValueString = AssetItem->GetAssetPathText().ToString();
 			}
 			else if (Column == DiskSizeColumnId)
 			{
diff --git a/Engine/Source/Editor/ContentBrowser/Private/AssetViewTypes.cpp b/Engine/Source/Editor/ContentBrowser/Private/AssetViewTypes.cpp
--- a/Engine/Source/Editor/ContentBrowser/Private/AssetViewTypes.cpp
+++ b/Engine/Source/Editor/ContentBrowser/Private/AssetViewTypes.cpp
@@ -4,47 +4,67 @@
 
 #include "AssetRegistry/AssetData.h"
 #include "Containers/UnrealString.h"
+#include "Containers/VersePath.h"
 #include "ContentBrowserDelegates.h"
 #include "ContentBrowserItemData.h"
 #include "HAL/Platform.h"
 #include "IAssetTools.h"
 #include "Internationalization/Text.h"
 #include "Misc/AssertionMacros.h"
+#include "Misc/PackageName.h"
 #include "Templates/UnrealTemplate.h"
 
 #define LOCTEXT_NAMESPACE "ContentBrowser"
 
 FAssetViewItem::FAssetViewItem(int32 InIndex)
 	: Index(InIndex)
+	, bShowingContentVersePath(false)
+	, bHasCachedVersePathText(false)
 {
 }
 
 FAssetViewItem::FAssetViewItem(int32 InIndex, FContentBrowserItem&& InItem)
 	: Item(MoveTemp(InItem))
 	, Index(InIndex)
+	, bShowingContentVersePath(false)
+	, bHasCachedVersePathText(false)
 {
 	checkf(Item.IsValid(), TEXT("FAssetViewItem was constructed from an invalid item!"));
+
+	CacheAssetPathText();
 }
 
 FAssetViewItem::FAssetViewItem(int32 InIndex, const FContentBrowserItem& InItem)
 	: Item(InItem)
 	, Index(InIndex)
+	, bShowingContentVersePath(false)
+	, bHasCachedVersePathText(false)
 {
 	checkf(Item.IsValid(), TEXT("FAssetViewItem was constructed from an invalid item!"));
+
+	CacheAssetPathText();
 }
 
 FAssetViewItem::FAssetViewItem(int32 InIndex, FContentBrowserItemData&& InItemData)
 	: Item(MoveTemp(InItemData))
 	, Index(InIndex)
+	, bShowingContentVersePath(false)
+	, bHasCachedVersePathText(false)
 {
 	checkf(Item.IsValid(), TEXT("FAssetViewItem was constructed from an invalid item!"));
+
+	CacheAssetPathText();
 }
 
 FAssetViewItem::FAssetViewItem(int32 InIndex, const FContentBrowserItemData& InItemData)
 	: Item(InItemData)
 	, Index(InIndex)
+	, bShowingContentVersePath(false)
+	, bHasCachedVersePathText(false)
 {
 	checkf(Item.IsValid(), TEXT("FAssetViewItem was constructed from an invalid item!"));
+
+	CacheAssetPathText();
 }
 
 void FAssetViewItem::ResetItemData(int32 OldIndex, int32 NewIndex, FContentBrowserItemData InItemData)
@@ -65,6 +85,15 @@ void FAssetViewItem::ResetItemData(int32 OldIndex, int32 NewIndex, FContentBrows
 
 void FAssetViewItem::BroadcastItemDataChanged()
 {
+	// Item has changed, update the asset path.
+	CacheAssetPathText();
+
+	// If CacheVersePathTextForTextFiltering has been called, update the Verse path too.
+	if (bHasCachedVersePathText)
+	{
+		CacheVersePathText();
+	}
+
 	ItemDataChangedEvent.Broadcast();
 }
 
@@ -85,7 +114,7 @@ void FAssetViewItem::RemoveItemData(const FContentBrowserItem& InItem)
 	Item.Remove(InItem);
 	if (Item.IsValid())
 	{
-		ItemDataChangedEvent.Broadcast();
+		BroadcastItemDataChanged();
 	}
 }
 
@@ -94,7 +123,7 @@ void FAssetViewItem::RemoveItemData(const FContentBrowserItemData& InItemData)
 	Item.Remove(InItemData);
 	if (Item.IsValid())
 	{
-		ItemDataChangedEvent.Broadcast();
+		BroadcastItemDataChanged();
 	}
 }
 
@@ -103,10 +132,56 @@ void FAssetViewItem::RemoveItemData(const FContentBrowserMinimalItemData& InItem
 	Item.TryRemove(InItemKey);
 	if (Item.IsValid())
 	{
-		ItemDataChangedEvent.Broadcast();
+		BroadcastItemDataChanged();
 	}	
 }
 
+void FAssetViewItem::ConfigureAssetPathText(bool bInShowingContentVersePath)
+{
+	bShowingContentVersePath = bInShowingContentVersePath;
+}
+
+void FAssetViewItem::CacheVersePathTextForTextFiltering()
+{
+	checkf(IsInGameThread(), TEXT("Verse paths can only be computed on the main thread."));
+
+	if (!bHasCachedVersePathText)
+	{
+		CacheVersePathText();
+		bHasCachedVersePathText = true;
+	}
+}
+
+const FText& FAssetViewItem::GetAssetPathText() const
+{
+	checkf(IsInGameThread(), TEXT("Verse paths can only be computed on the main thread."));
+
+	// Compute the Verse path on demand since it is somewhat expensive and not parallelizable.
+	// This function is called when viewing the path column if it is avalible, or sorting on the path column.
+	if (bShowingContentVersePath)
+	{
+		if (!bHasCachedVersePathText)
+		{
+			CacheVersePathText();
+			bHasCachedVersePathText = true;
+		}
+
+		if (!CachedVersePathText.IsEmpty())
+		{
+			return CachedVersePathText;
+		}
+	}
+
+	return CachedAssetPathText;
+}
+
+const FText& FAssetViewItem::GetVersePathTextForTextFiltering() const
+{
+	checkf(bHasCachedVersePathText, TEXT("Need to call CacheVersePathTextForTextFiltering first."));
+
+	return CachedVersePathText;
+}
+
 void FAssetViewItem::ClearCachedCustomColumns()
 {
 	check(IsInGameThread());
@@ -346,4 +421,62 @@ FString FAssetViewItem::ItemToString_Debug(TSharedPtr<FAssetViewItem> AssetItem)
 	return TEXT("nullptr");
 }
 
+void FAssetViewItem::CacheAssetPathText()
+{
+	if (IsFile())
+	{
+		FString ObjectPath;
+		if (Item.AppendItemObjectPath(ObjectPath))
+		{
+			CachedAssetPathText = FText::AsCultureInvariant(MoveTemp(ObjectPath));
+			return;
+		}
+	}
+	else
+	{
+		FName PackagePath;
+		if (Item.Legacy_TryGetPackagePath(PackagePath))
+		{
+			CachedAssetPathText = FText::FromName(PackagePath);
+			return;
+		}
+	}
+
+	CachedAssetPathText = FText::FromName(Item.GetVirtualPath());
+}
+
+void FAssetViewItem::CacheVersePathText() const
+{
+	checkf(IsInGameThread(), TEXT("Verse paths can only be computed on the main thread."));
+
+	if (IsFile())
+	{
+		FAssetData AssetData;
+		if (Item.Legacy_TryGetAssetData(AssetData))
+		{
+			UE::Core::FVersePath VersePath = AssetData.GetVersePath();
+			if (VersePath.IsValid())
+			{
+				CachedVersePathText = MoveTemp(VersePath).AsText();
+				return;
+			}
+		}
+	}
+	else
+	{
+		FName PackagePath;
+		if (Item.Legacy_TryGetPackagePath(PackagePath))
+		{
+			UE::Core::FVersePath VersePath = FPackageName::LongPackagePathToVersePath(FNameBuilder(PackagePath));
+			if (VersePath.IsValid())
+			{
+				CachedVersePathText = MoveTemp(VersePath).AsText();
+				return;
+			}
+		}
+	}
+
+	CachedVersePathText = {};
+}
+
 #undef LOCTEXT_NAMESPACE
diff --git a/Engine/Source/Editor/ContentBrowser/Private/AssetViewWidgets.cpp b/Engine/Source/Editor/ContentBrowser/Private/AssetViewWidgets.cpp
--- a/Engine/Source/Editor/ContentBrowser/Private/AssetViewWidgets.cpp
+++ b/Engine/Source/Editor/ContentBrowser/Private/AssetViewWidgets.cpp
@@ -2091,9 +2091,7 @@ FSlateColor SAssetListItem::GetColumnTextColor(FIsSelected InIsSelected) const
 
 FText SAssetListItem::GetAssetPathText() const
 {
-	return AssetItem
-		? FText::AsCultureInvariant(AssetItem->GetItem().GetVirtualPath().ToString())
-		: FText();
+	return AssetItem ? AssetItem->GetAssetPathText() : FText::GetEmpty();
 }
 
 FText SAssetListItem::GetAssetTagText(FName InTag) const
@@ -3559,9 +3557,7 @@ void SAssetColumnItem::OnAssetDataChanged()
 
 FText SAssetColumnItem::GetAssetPathText() const
 {
-	return AssetItem
-		? FText::AsCultureInvariant(AssetItem->GetItem().GetVirtualPath().ToString())
-		: FText();
+	return AssetItem ? AssetItem->GetAssetPathText() : FText::GetEmpty();
 }
 
 FText SAssetColumnItem::GetAssetTagText(FName AssetTag) const
diff --git a/Engine/Source/Editor/ContentBrowser/Private/SAssetView.cpp b/Engine/Source/Editor/ContentBrowser/Private/SAssetView.cpp
--- a/Engine/Source/Editor/ContentBrowser/Private/SAssetView.cpp
+++ b/Engine/Source/Editor/ContentBrowser/Private/SAssetView.cpp
@@ -370,6 +370,14 @@ class FAssetViewItemCollection
 public: 
 	FAssetViewItemCollection()
 	{
+		bShowingContentVersePath = IAssetTools::Get().ShowingContentVersePath();
+
+		IPluginManager::Get().OnPluginEdited().AddRaw(this, &FAssetViewItemCollection::OnPluginEdited);
+	}
+
+	~FAssetViewItemCollection()
+	{
+		IPluginManager::Get().OnPluginEdited().RemoveAll(this);
 	}
 
 	/** Returns the number of items which were fetched and have not been removed. */
@@ -463,6 +471,17 @@ public:
 	bool PerformPriorityFiltering(FAssetViewFrontendFilterHelper& Helper,
 		TArray<TSharedPtr<FAssetViewItem>>& FilteredAssetItems);
 
+	bool GetShowingContentVersePath() const
+	{
+		return bShowingContentVersePath;
+	}
+
+	/**
+	 * Check if we need to update the path text of any items.
+	 * Returns true if the path text may have changed to indicate items need to be sorted and filtered again.
+	 */
+	bool UpdateShowingContentVersePath();
+
 private:
 	int32 CreateItem_Locked(FContentBrowserItemData&& InItem)
 	{
@@ -473,6 +492,11 @@ private:
 	{
 		uint32 Hash = HashItem(InItem);
 		TSharedPtr<FAssetViewItem> NewItem = MakeShared<FAssetViewItem>(Items.Num(), MoveTemp(InItem));
+		NewItem->ConfigureAssetPathText(bShowingContentVersePath);
+		if (TextFilterTask.IsValid() && bCompiledTextFilterRequiresVersePaths)
+		{
+			NewItem->CacheVersePathTextForTextFiltering();
+		}
 		Items.Add(MoveTemp(NewItem));
 		FilterState.AddZeroed(1);
 		++NumValidItems;
@@ -538,6 +562,12 @@ private:
 	 */
 	bool RefreshLookup();
 
+	void OnPluginEdited(IPlugin& InPlugin)
+	{
+		// The plugin's Verse path may have changed.  Recompute all Verse paths.
+		bShouldInvalidateVersePaths = true;
+	}
+
 private:
 	// Lock for access to the size and contents of Items - e.g. when text filtering is operating on a batch of items
 	// The size of Items may need to change or the ItemData object within items may need modification as data scanning progresses
@@ -566,6 +596,8 @@ private:
 
 	// Cached compiled text filter for the current filtering pass
 	TSharedPtr<FCompiledAssetTextFilter> CompiledTextFilter;
+	// True if the compiled text filter might need Verse path info.
+	bool bCompiledTextFilterRequiresVersePaths = false;
 
 	// Handle to ongoing task filtering Items by a text query. 
 	// Modifications during filtering are protected by the Lock in this class and otherwise yet-to-be-filtered items should not be
@@ -584,6 +616,11 @@ private:
 
 	// Some items have been marked for removal but their pointers have not been cleared yet because we need to compare against them
 	std::atomic<bool> bItemsPendingRemove = false;
+
+	// True if we should show Verse paths in the path column or include them when text filtering.
+	bool bShowingContentVersePath;
+	// True if a plugin descriptor might have changed in a way that affects Verse paths.
+	bool bShouldInvalidateVersePaths = false;
 };
 
 bool FAssetViewItemCollection::RefreshLookup()
@@ -794,15 +831,29 @@ void FAssetViewItemCollection::StartTextFiltering(TSharedPtr<FAssetTextFilter> T
 	if (!TextFilter.IsValid() || TextFilter->IsEmpty())
 	{
 		CompiledTextFilter.Reset();
+		bCompiledTextFilterRequiresVersePaths = false;
 		bAllItemsPassedTextFilter = true;
 		return;
 	}
 
 	CompiledTextFilter = TextFilter->Compile();
+	bCompiledTextFilterRequiresVersePaths = bShowingContentVersePath;
 	bAllItemsPassedTextFilter = false;
 
 	if (UE::AssetView::AllowAsync)
 	{
+		// Verse paths need to be computed on the game thread at the moment.
+		if (bCompiledTextFilterRequiresVersePaths)
+		{
+			for (const TSharedPtr<FAssetViewItem>& Item : Items)
+			{
+				if (Item.IsValid())
+				{
+					Item->CacheVersePathTextForTextFiltering();
+				}
+			}
+		}
+
 		const int32 MaxItemsPerTask = UE::AssetView::GetMaxTextFilterItemBatch();
 		TextFilterTask = UE::Tasks::Launch(UE_SOURCE_LOCATION, [this, MaxItemsPerTask]() {
 			return AsyncFilterText(0, MaxItemsPerTask);
@@ -832,7 +883,10 @@ FAssetViewItemCollection::FTextFilterResult FAssetViewItemCollection::AsyncFilte
 	};
 	auto DoWork = [&MergedResult, StartIndex, this](FCompiledAssetTextFilter& Filter, int32 TaskIndex) {
 		const TSharedPtr<FAssetViewItem>& Item = Items[StartIndex + TaskIndex];
-		bool bPasses = Item.IsValid() && Filter.PassesFilter(Item->GetItem());
+		bool bPasses = Item.IsValid() &&
+			(bCompiledTextFilterRequiresVersePaths ?
+			Filter.PassesFilter(Item->GetItem(), Item->GetVersePathTextForTextFiltering().ToString()) :
+			Filter.PassesFilter(Item->GetItem()));
 		if (bPasses)
 		{
 			MergedResult[TaskIndex].AtomicSet(true);
@@ -849,7 +903,7 @@ FAssetViewItemCollection::FTextFilterResult FAssetViewItemCollection::AsyncFilte
 	{
 		return FTextFilterResult{ StartIndex, MoveTemp(MergedResult),
 			UE::Tasks::Launch(
-				UE_SOURCE_LOCATION, [this, NextStartIndex, InMaxItems]() mutable { return AsyncFilterText(NextStartIndex, InMaxItems); }) };
+				UE_SOURCE_LOCATION, [this, NextStartIndex, InMaxItems]() { return AsyncFilterText(NextStartIndex, InMaxItems); }) };
 	}
 	else
 	{
@@ -923,8 +977,25 @@ void FAssetViewItemCollection::UpdateItemFiltering(
 				MergeTextFilterResult(MoveTemp(TextFilterTask.GetResult()));
 			}
 
+			auto CacheVersePathTextForTextFiltering = [this]()
+			{
+				// Verse paths need to be computed on the game thread at the moment.
+				if (bCompiledTextFilterRequiresVersePaths)
+				{
+					for (int32 Index = TextFilterProgress; Index < Items.Num(); ++Index)
+					{
+						if (Items[Index].IsValid())
+						{
+							Items[Index]->CacheVersePathTextForTextFiltering();
+						}
+					}
+				}
+			};
+
 			if (!TextFilterTask.IsValid() && TextFilterProgress < Items.Num() && UE::AssetView::AllowAsync)
 			{
+				CacheVersePathTextForTextFiltering();
+
 				// New elements were added after the text filter attempted to launch a continuation task
 				TextFilterTask = UE::Tasks::Launch(UE_SOURCE_LOCATION,
 					[this, StartIndex = TextFilterProgress]() {
@@ -935,6 +1006,8 @@ void FAssetViewItemCollection::UpdateItemFiltering(
 			// In case flag was flipped while filtering was running, wait til task ends before performing text filtering on the game thread.
 			if (!UE::AssetView::AllowAsync && !TextFilterTask.IsValid() && TextFilterProgress < Items.Num())
 			{
+				CacheVersePathTextForTextFiltering();
+
 				FTextFilterResult Result = AsyncFilterText(TextFilterProgress, UE::AssetView::GetMaxTextFilterItemBatch());
 				MergeTextFilterResult(Result);
 			}
@@ -1021,7 +1094,9 @@ bool FAssetViewItemCollection::PerformPriorityFiltering(FAssetViewFrontendFilter
 					continue;
 				}
 
-				bool bPassed = CompiledTextFilter->PassesFilter(Items[Index]->GetItem());
+				bool bPassed = bCompiledTextFilterRequiresVersePaths ?
+					CompiledTextFilter->PassesFilter(Items[Index]->GetItem(), Items[Index]->GetVersePathTextForTextFiltering().ToString()) :
+					CompiledTextFilter->PassesFilter(Items[Index]->GetItem());
 				FilterState[Index].PassedTextFilter = bPassed;
 			}
 		}
@@ -1074,6 +1149,46 @@ bool FAssetViewItemCollection::PerformPriorityFiltering(FAssetViewFrontendFilter
 	return PrevNum != FilteredAssetItems.Num();
 }
 
+bool FAssetViewItemCollection::UpdateShowingContentVersePath()
+{
+	bool bInvalidateAssetPathText = false;
+
+	const bool bNewShowingContentVersePath = IAssetTools::Get().ShowingContentVersePath();
+	if (bShowingContentVersePath != bNewShowingContentVersePath)
+	{
+		bShowingContentVersePath = bNewShowingContentVersePath;
+
+		bInvalidateAssetPathText = true;
+	}
+
+	if (bShouldInvalidateVersePaths)
+	{
+		bShouldInvalidateVersePaths = false;
+
+		if (bShowingContentVersePath)
+		{
+			bInvalidateAssetPathText = true;
+		}
+	}
+
+	if (bInvalidateAssetPathText)
+	{
+		// Lock since we may update a Verse path that the background text search is reading.
+		FWriteScopeLock Guard(Lock);
+		for (const TSharedPtr<FAssetViewItem>& Item : Items)
+		{
+			if (Item.IsValid())
+			{
+				Item->ConfigureAssetPathText(bShowingContentVersePath);
+				Item->BroadcastItemDataChanged();
+			}
+		}
+		return true;
+	}
+
+	return false;
+}
+
 /** 
  * SAssetView
  */
@@ -1771,6 +1886,7 @@ void SAssetView::BeginCreateDeferredItem()
 	if (DeferredItemToCreate.IsValid() && !DeferredItemToCreate->bWasAddedToView)
 	{
 		TSharedPtr<FAssetViewItem> NewItem = MakeShared<FAssetViewItem>(INDEX_NONE, DeferredItemToCreate->ItemContext.GetItem());
+		NewItem->ConfigureAssetPathText(Items->GetShowingContentVersePath());
 		AwaitingScrollIntoViewForRename = NewItem;
 		DeferredItemToCreate->bWasAddedToView = true;
 
@@ -2192,6 +2308,26 @@ void SAssetView::Tick( const FGeometry& AllottedGeometry, const double InCurrent
 {
 	using namespace UE::AssetView;
 
+	if (Items->UpdateShowingContentVersePath())
+	{
+		// A path may have changed, if sorting on the path column we need to resort.
+		if (bShowPathInColumnView)
+		{
+			for (int32 PriorityIdx = 0; PriorityIdx < EColumnSortPriority::Max; PriorityIdx++)
+			{
+				if (SortManager->GetSortColumnId(static_cast<EColumnSortPriority::Type>(PriorityIdx)) == FAssetViewSortManager::PathColumnId)
+				{
+					// Don't sync to selection if we are just going to do it below
+					SortList(!PendingSyncItems.Num());
+					break;
+				}
+			}
+		}
+
+		// A path may have changed, if filtering on a path we need to refilter.
+		bQuickFrontendListRefreshRequested = true;
+	}
+
 	// Adjust min and max thumbnail scale based on dpi
 	MinThumbnailScale = (0.2f * ThumbnailScaleRangeScalar)/AllottedGeometry.Scale;
 	MaxThumbnailScale = (1.9f * ThumbnailScaleRangeScalar)/AllottedGeometry.Scale;
@@ -3401,7 +3537,9 @@ void FAssetViewItemCollection::RefreshItemsFromBackend(const FAssetViewContentSo
 				}
 
 				// Was not able to recycle an old item
-				Items[Index] = MakeShared<FAssetViewItem>(Index, MoveTemp(ItemData));
+				TSharedPtr<FAssetViewItem> NewItem = MakeShared<FAssetViewItem>(Index, MoveTemp(ItemData));
+				NewItem->ConfigureAssetPathText(bShowingContentVersePath);
+				Items[Index] = MoveTemp(NewItem);
 				Lookup.Add_Concurrent(Hash, Index);
 			}, UE::AssetView::AllowParallelism ? EParallelForFlags::None : EParallelForFlags::ForceSingleThread);
 	}
@@ -3501,6 +3639,7 @@ void FAssetViewItemCollection::RefreshItemsFromBackend(const FAssetViewContentSo
 	// Until we start filtering and get a compiled filter, initialize to no filtering
 	bAllItemsPassedTextFilter = true;
 	CompiledTextFilter.Reset();
+	bCompiledTextFilterRequiresVersePaths = false;
 }
 
 bool SAssetView::IsFilteringRecursively() const
diff --git a/Engine/Source/Editor/ContentBrowser/Public/AssetTextFilter.h b/Engine/Source/Editor/ContentBrowser/Public/AssetTextFilter.h
--- a/Engine/Source/Editor/ContentBrowser/Public/AssetTextFilter.h
+++ b/Engine/Source/Editor/ContentBrowser/Public/AssetTextFilter.h
@@ -170,6 +170,7 @@ public:
 	 * many such checks 
      */
 	bool PassesFilter(FAssetFilterType InItem);
+	bool PassesFilter(FAssetFilterType InItem, FStringView InItemVersePath);
 
 	/** Test the given value against the strings extracted from the current item */
 	virtual bool TestBasicStringExpression(
@@ -202,18 +203,27 @@ private:
     
 	// State for PassesFilter
 
+	bool bIsFile = false;
+
 	/** Pointer to the asset we're currently filtering */
 	FAssetFilterTypePtr AssetPtr;
     
     /** Shared buffer for text keys to search */
     FString TextBuffer;
 
+	/** Recyclable text filter string so reduce allocations  */
+	mutable FTextFilterString TextFilterString;
+
 	/** Display name of the current asset */
 	FStringView AssetDisplayName;
 	/** Full path of the current asset */
 	FStringView AssetFullPath;
+	/** Virtual path of the current asset */
+	FStringView AssetVirtualPath;
 	/** The export text name of the current asset */
 	FStringView AssetExportTextPath;
+	/** Verse path of the current asset */
+	FStringView AssetVersePath;
 
 	/** Names of the collections that the current asset is in */
 	TArray<FName> AssetCollectionNames;
diff --git a/Engine/Source/Editor/ContentBrowser/Public/AssetViewTypes.h b/Engine/Source/Editor/ContentBrowser/Public/AssetViewTypes.h
--- a/Engine/Source/Editor/ContentBrowser/Public/AssetViewTypes.h
+++ b/Engine/Source/Editor/ContentBrowser/Public/AssetViewTypes.h
@@ -55,6 +55,18 @@ public:
 
 	void RemoveItemData(const FContentBrowserMinimalItemData& InItemKey);
 
+	/** Configures the type of path GetAssetPathText should return. */
+	UE_INTERNAL void ConfigureAssetPathText(bool bInShowingContentVersePath);
+
+	/** Caches the Verse path to use when filtering items. */
+	UE_INTERNAL void CacheVersePathTextForTextFiltering();
+
+	/** Gets the asset path to use when viewing or sorting the path column. */
+	UE_INTERNAL const FText& GetAssetPathText() const;
+
+	/** Gets the Verse path to use when filtering items. Requires a call to CacheVersePathTextForTextFiltering first. */
+	UE_INTERNAL const FText& GetVersePathTextForTextFiltering() const;
+
 	/** Clear cached custom column data */
 	void ClearCachedCustomColumns();
 
@@ -123,6 +135,7 @@ public:
 
 	// Called when the view explicitly wants to notify widgets of changes
 	// Not called during bulk rebuilds when the view will be re-populated even if items are being recycled
+	// Updates cached paths then raises the OnItemDataChanged event.
 	void BroadcastItemDataChanged();
 
 	/** Get the event fired when the data for this item changes */
@@ -138,6 +151,9 @@ public:
 	static FString ItemToString_Debug(TSharedPtr<FAssetViewItem> AssetItem);
 
 private:
+	void CacheAssetPathText();
+	void CacheVersePathText() const;
+
 	/** Underlying Content Browser item data */
 	FContentBrowserItem Item;
 
@@ -147,6 +163,11 @@ private:
 	 */
 	std::atomic<int32> Index;
 
+	bool bShowingContentVersePath;
+	mutable bool bHasCachedVersePathText;
+	FText CachedAssetPathText;
+	mutable FText CachedVersePathText;
+
 	/** An event to fire when the data for this item changes */
 	FSimpleMulticastDelegate ItemDataChangedEvent;
 
diff --git a/Engine/Source/Runtime/Core/Private/Misc/TextFilterUtils.cpp b/Engine/Source/Runtime/Core/Private/Misc/TextFilterUtils.cpp
--- a/Engine/Source/Runtime/Core/Private/Misc/TextFilterUtils.cpp
+++ b/Engine/Source/Runtime/Core/Private/Misc/TextFilterUtils.cpp
@@ -225,64 +225,72 @@ namespace FastToUpper
 	}
 } // namespace FastToUpper
 
-FTextFilterString::FTextFilterString()
-	: InternalString(),
-	InternalStringAnsi()
+FTextFilterString::FTextFilterString() = default;
+
+FTextFilterString::FTextFilterString(const FTextFilterString& Other) = default;
+
+FTextFilterString::FTextFilterString(FTextFilterString&& Other) = default;
+
+FTextFilterString::FTextFilterString(const FString& InString)
+	: InternalString(InString)
 {
+	UppercaseInternalString();
 }
 
-FTextFilterString::FTextFilterString(const FTextFilterString& Other)
-	: InternalString(Other.InternalString),
-	InternalStringAnsi(Other.InternalStringAnsi)
+FTextFilterString::FTextFilterString(FString&& InString)
+	: InternalString(MoveTemp(InString))
 {
+	UppercaseInternalString();
 }
 
-FTextFilterString::FTextFilterString(FTextFilterString&& Other)
-	: InternalString(MoveTemp(Other.InternalString)),
-	InternalStringAnsi(MoveTemp(Other.InternalStringAnsi))
+FTextFilterString::FTextFilterString(FStringView InString)
+	: InternalString(InString)
 {
+	UppercaseInternalString();
 }
 
-FTextFilterString::FTextFilterString(const FString& InString)
-	: InternalString(InString),
-	InternalStringAnsi()
+FTextFilterString::FTextFilterString(const TCHAR* InString)
+	: InternalString(InString)
 {
 	UppercaseInternalString();
 }
 
-FTextFilterString::FTextFilterString(FString&& InString)
-	: InternalString(MoveTemp(InString)),
-	InternalStringAnsi()
+FTextFilterString::FTextFilterString(FName InName)
 {
+	InName.AppendString(InternalString);
 	UppercaseInternalString();
 }
 
-FTextFilterString::FTextFilterString(const TCHAR* InString)
-	: InternalString(InString),
-	InternalStringAnsi()
+FTextFilterString& FTextFilterString::operator=(const FTextFilterString& Other) = default;
+
+FTextFilterString& FTextFilterString::operator=(FTextFilterString&& Other) = default;
+
+FTextFilterString& FTextFilterString::operator=(const FString& InString)
 {
+	InternalString = InString;
 	UppercaseInternalString();
+	return *this;
 }
 
-FTextFilterString::FTextFilterString(const FName& InName)
-	: InternalString(),
-	InternalStringAnsi()
+FTextFilterString& FTextFilterString::operator=(FStringView InString)
 {
-	InName.AppendString(InternalString);
+	InternalString = InString;
 	UppercaseInternalString();
+	return *this;
 }
 
-FTextFilterString& FTextFilterString::operator=(const FTextFilterString& Other)
+FTextFilterString& FTextFilterString::operator=(const TCHAR* InString)
 {
-	InternalString = Other.InternalString;
-	InternalStringAnsi = Other.InternalStringAnsi;
+	InternalString = InString;
+	UppercaseInternalString();
 	return *this;
 }
 
-FTextFilterString& FTextFilterString::operator=(FTextFilterString&& Other)
+FTextFilterString& FTextFilterString::operator=(FName InName)
 {
-	InternalString = MoveTemp(Other.InternalString);
-	InternalStringAnsi = MoveTemp(Other.InternalStringAnsi);
+	InternalString.Reset();
+	InName.AppendString(InternalString);
+	UppercaseInternalString();
 	return *this;
 }
 
diff --git a/Engine/Source/Runtime/Core/Public/Misc/TextFilterUtils.h b/Engine/Source/Runtime/Core/Public/Misc/TextFilterUtils.h
--- a/Engine/Source/Runtime/Core/Public/Misc/TextFilterUtils.h
+++ b/Engine/Source/Runtime/Core/Public/Misc/TextFilterUtils.h
@@ -47,15 +47,21 @@ public:
 	/** Construct from a string */
 	CORE_API FTextFilterString(const FString& InString);
 	CORE_API FTextFilterString(FString&& InString);
+	CORE_API FTextFilterString(FStringView InString);
 	CORE_API FTextFilterString(const TCHAR* InString);
 
 	/** Construct from a name */
-	CORE_API FTextFilterString(const FName& InName);
+	CORE_API FTextFilterString(FName InName);
 
 	/** Move and copy assignment */
 	CORE_API FTextFilterString& operator=(const FTextFilterString& Other);
 	CORE_API FTextFilterString& operator=(FTextFilterString&& Other);
 
+	CORE_API FTextFilterString& operator=(const FString& InString);
+	CORE_API FTextFilterString& operator=(FStringView InString);
+	CORE_API FTextFilterString& operator=(const TCHAR* InString);
+	CORE_API FTextFilterString& operator=(FName InName);
+
 	/** Compare this string against the other, using the text comparison mode provided */
 	CORE_API bool CompareText(const FTextFilterString& InOther, const ETextFilterTextComparisonMode InTextComparisonMode) const;
 
