Added UE_EXPERIMENTAL around ItemDataAttributeDelegates definitions.

#rnx

diff --git a/Engine/Source/Editor/ContentBrowser/Public/ContentBrowserDelegates.h b/Engine/Source/Editor/ContentBrowser/Public/ContentBrowserDelegates.h
--- a/Engine/Source/Editor/ContentBrowser/Public/ContentBrowserDelegates.h
+++ b/Engine/Source/Editor/ContentBrowser/Public/ContentBrowserDelegates.h
@@ -119,6 +119,7 @@ DECLARE_DELEGATE_RetVal_OneParam(TSharedRef<SWidget>, FOnGenerateAssetViewExtraS
 DECLARE_DELEGATE_RetVal_OneParam(TSharedPtr<SWidget>, FPathViewStateIconGenerator, const FContentBrowserItem& /*ContentBrowserItem*/);
 
 /** Called to get ContentBrowserItemData attribute values. */
+UE_EXPERIMENTAL(5.7, "FGetItemDataAttribute is experimental while a newer system is in active development.")
 DECLARE_DELEGATE_RetVal_FourParams(bool, FGetItemDataAttribute, const FContentBrowserItemData& /*InItem*/, const bool /*InIncludeMetaData*/, const FName /*InAttributeKey*/, FContentBrowserItemDataAttributeValue& /*OutAttributeValue*/);
 
 /** Called when an asset item visualizes its tooltip */
diff --git a/Engine/Source/Editor/ContentBrowser/Public/ContentBrowserModule.h b/Engine/Source/Editor/ContentBrowser/Public/ContentBrowserModule.h
--- a/Engine/Source/Editor/ContentBrowser/Public/ContentBrowserModule.h
+++ b/Engine/Source/Editor/ContentBrowser/Public/ContentBrowserModule.h
@@ -178,7 +178,10 @@ public:
 	virtual TArray<FAddPathViewPluginFilters>& GetAddPathViewPluginFilters() { return PathViewPluginFilters; }
 
 	/** Delegates to call to provide the value for a ContentBrowserItemData's attribute */
+	PRAGMA_DISABLE_EXPERIMENTAL_WARNINGS
+	UE_EXPERIMENTAL(5.7, "GetAllItemDataAttributeDelegates() is experimental while a newer system is in active development.")
 	virtual TArray<FGetItemDataAttribute>& GetAllItemDataAttributeDelegates() { return ItemDataAttributeDelegates; }
+	PRAGMA_ENABLE_EXPERIMENTAL_WARNINGS
 
 	/** Register a delegate that creates a custom view that can show up inside SAssetView in the Content Browser */
 	virtual void SetContentBrowserViewExtender(const FCreateViewExtender& InCreateViewExtender);
@@ -275,7 +278,9 @@ private:
 	TArray<FAddPathViewPluginFilters> PathViewPluginFilters;
 
 	/** All delegates that return values for ContentItemData attributes */
+	PRAGMA_DISABLE_EXPERIMENTAL_WARNINGS
 	TArray<FGetItemDataAttribute> ItemDataAttributeDelegates;
+	PRAGMA_ENABLE_EXPERIMENTAL_WARNINGS
 
 	FOnFilterChanged OnFilterChanged;
 	FOnSearchBoxChanged OnSearchBoxChanged;
