Slate Dynamic Invalidation - Content Browser - Widgets already supporting Dynamic Invalidation

This CL adds a Dynamic Invalidation Panel for the ContentBrowser, and adds a `TWidgetTypeTraits` to widgets that are already invalidation compatible.
The DynamicInvalidation has also been turned off by default at this stage, and requires setting `Slate.DynamicInvalidation.Options 2` for it to be enabled.

diff --git a/Engine/Source/Developer/ToolWidgets/Public/SPrimaryButton.h b/Engine/Source/Developer/ToolWidgets/Public/SPrimaryButton.h
--- a/Engine/Source/Developer/ToolWidgets/Public/SPrimaryButton.h
+++ b/Engine/Source/Developer/ToolWidgets/Public/SPrimaryButton.h
@@ -8,6 +8,12 @@
 
 #define UE_API TOOLWIDGETS_API
 
+template <>
+struct TWidgetTypeTraits<class SPrimaryButton>
+{
+	static constexpr bool SupportsInvalidation() { return true; }
+};
+
 class SPrimaryButton
 	: public SButton
 {
diff --git a/Engine/Source/Editor/ContentBrowser/Private/SContentBrowser.cpp b/Engine/Source/Editor/ContentBrowser/Private/SContentBrowser.cpp
--- a/Engine/Source/Editor/ContentBrowser/Private/SContentBrowser.cpp
+++ b/Engine/Source/Editor/ContentBrowser/Private/SContentBrowser.cpp
@@ -162,6 +162,7 @@
 #include "Widgets/Notifications/SNotificationList.h"
 #include "Widgets/SBoxPanel.h"
 #include "Widgets/SContentBrowserSourceTree.h"
+#include "Widgets/SInvalidationPanel.h"
 #include "Widgets/SNullWidget.h"
 #include "Widgets/SWidget.h"
 #include "Widgets/Text/STextBlock.h"
@@ -1014,13 +1015,17 @@ TSharedRef<SWidget> SContentBrowser::CreateAssetView(const FContentBrowserConfig
 			.BorderImage(FAppStyle::Get().GetBrush("Brushes.Panel"))
 			.Padding(AssetViewPadding)
 			[
-				SNew(SVerticalBox)
-
-				+ SVerticalBox::Slot()
-				.FillHeight(1.0f)
-				.Padding(0, 0)
+				SNew(SInvalidationPanel)
+				.UseDynamicInvalidation(true)
 				[
-					LegacyContentSourceWidgets->AssetViewPtr.ToSharedRef()
+					SNew(SVerticalBox)
+
+					+ SVerticalBox::Slot()
+					.FillHeight(1.0f)
+					.Padding(0, 0)
+					[
+						LegacyContentSourceWidgets->AssetViewPtr.ToSharedRef()
+					]
 				]
 			];
 
diff --git a/Engine/Source/Editor/ContentBrowser/Public/SAssetView.h b/Engine/Source/Editor/ContentBrowser/Public/SAssetView.h
--- a/Engine/Source/Editor/ContentBrowser/Public/SAssetView.h
+++ b/Engine/Source/Editor/ContentBrowser/Public/SAssetView.h
@@ -107,6 +107,12 @@ enum class EAssetViewCopyType
 	PackageName
 };
 
+template <>
+struct TWidgetTypeTraits<class SAssetView>
+{
+	static constexpr bool SupportsInvalidation() { return true; }
+};
+
 /**
  * A widget to display a list of filtered assets
  */
diff --git a/Engine/Source/Editor/EditorWidgets/Private/SAssetDiscoveryIndicator.h b/Engine/Source/Editor/EditorWidgets/Private/SAssetDiscoveryIndicator.h
--- a/Engine/Source/Editor/EditorWidgets/Private/SAssetDiscoveryIndicator.h
+++ b/Engine/Source/Editor/EditorWidgets/Private/SAssetDiscoveryIndicator.h
@@ -19,6 +19,12 @@
 
 struct FGeometry;
 
+template <>
+struct TWidgetTypeTraits<class SAssetDiscoveryIndicator>
+{
+	static constexpr bool SupportsInvalidation() { return true; }
+};
+
 /** An indicator for the progress of the asset registry background search */
 class SAssetDiscoveryIndicator : public SCompoundWidget
 {
diff --git a/Engine/Source/Editor/UnrealEd/Private/AssetThumbnail.cpp b/Engine/Source/Editor/UnrealEd/Private/AssetThumbnail.cpp
--- a/Engine/Source/Editor/UnrealEd/Private/AssetThumbnail.cpp
+++ b/Engine/Source/Editor/UnrealEd/Private/AssetThumbnail.cpp
@@ -76,6 +76,12 @@ namespace AssetThumbnailPool
 
 FName FAssetThumbnailPool::CustomThumbnailTagName = "CustomThumbnail";
 
+template <>
+struct TWidgetTypeTraits<class SAssetThumbnail>
+{
+	static constexpr bool SupportsInvalidation() { return true; }
+};
+
 class SAssetThumbnail : public SCompoundWidget
 {
 public:
diff --git a/Engine/Source/Editor/UnrealEd/Private/SAssetThumbnailEditModeTools.h b/Engine/Source/Editor/UnrealEd/Private/SAssetThumbnailEditModeTools.h
--- a/Engine/Source/Editor/UnrealEd/Private/SAssetThumbnailEditModeTools.h
+++ b/Engine/Source/Editor/UnrealEd/Private/SAssetThumbnailEditModeTools.h
@@ -15,6 +15,12 @@ struct FSlateBrush;
 // This is the equivalent of the ContentBrowser one, with some Graphic changes.
 // The ContentBrowser one will be later removed when the new style will be enabled by default.
 
+template <>
+struct TWidgetTypeTraits<class SAssetThumbnailEditModeTools>
+{
+	static constexpr bool SupportsInvalidation() { return true; }
+};
+
 /**
  * AssetThumbnail EditMode widget, used for the Edit mode of the thumbnails (Ex: changing Shape/Angle)
  */
diff --git a/Engine/Source/Runtime/Slate/Private/Widgets/SInvalidationPanel.cpp b/Engine/Source/Runtime/Slate/Private/Widgets/SInvalidationPanel.cpp
--- a/Engine/Source/Runtime/Slate/Private/Widgets/SInvalidationPanel.cpp
+++ b/Engine/Source/Runtime/Slate/Private/Widgets/SInvalidationPanel.cpp
@@ -40,7 +40,7 @@ namespace DynamicInvalidationCVars
 {
 	void ConsoleVariableEnableDynamicInvalidationChanged(IConsoleVariable*);
 
-	static int32 DynamicInvalidationOptions = 2;
+	static int32 DynamicInvalidationOptions = 0;
 	static FAutoConsoleVariableRef CVarDynamicInvalidationOptions(
 		TEXT("Slate.DynamicInvalidation.Options"),
 		DynamicInvalidationOptions,
diff --git a/Engine/Source/Runtime/Slate/Private/Widgets/Views/SHeaderRow.cpp b/Engine/Source/Runtime/Slate/Private/Widgets/Views/SHeaderRow.cpp
--- a/Engine/Source/Runtime/Slate/Private/Widgets/Views/SHeaderRow.cpp
+++ b/Engine/Source/Runtime/Slate/Private/Widgets/Views/SHeaderRow.cpp
@@ -21,6 +21,12 @@
 
 #define LOCTEXT_NAMESPACE "SHeaderRow"
 
+template <>
+struct TWidgetTypeTraits<class STableColumnHeader>
+{
+	static constexpr bool SupportsInvalidation() { return true; }
+};
+
 class STableColumnHeader : public SCompoundWidget
 {
 public:
diff --git a/Engine/Source/Runtime/Slate/Public/Widgets/Colors/SColorBlock.h b/Engine/Source/Runtime/Slate/Public/Widgets/Colors/SColorBlock.h
--- a/Engine/Source/Runtime/Slate/Public/Widgets/Colors/SColorBlock.h
+++ b/Engine/Source/Runtime/Slate/Public/Widgets/Colors/SColorBlock.h
@@ -24,6 +24,13 @@ enum class EColorBlockAlphaDisplayMode : uint8
 	Ignore,
 };
 
+
+template <>
+struct TWidgetTypeTraits<class SColorBlock>
+{
+	static constexpr bool SupportsInvalidation() { return true; }
+};
+
 class SColorBlock : public SLeafWidget
 {
 	SLATE_DECLARE_WIDGET_API(SColorBlock, SLeafWidget, SLATE_API)
diff --git a/Engine/Source/Runtime/Slate/Public/Widgets/Images/SThrobber.h b/Engine/Source/Runtime/Slate/Public/Widgets/Images/SThrobber.h
--- a/Engine/Source/Runtime/Slate/Public/Widgets/Images/SThrobber.h
+++ b/Engine/Source/Runtime/Slate/Public/Widgets/Images/SThrobber.h
@@ -15,6 +15,12 @@ class FSlateWindowElementList;
 class SHorizontalBox;
 
 
+template <>
+struct TWidgetTypeTraits<class SThrobber>
+{
+	static constexpr bool SupportsInvalidation() { return true; }
+};
+
 /**
  * A throbber widget that uses 5 zooming circles in a row.
  */
@@ -100,6 +106,12 @@ private:
 };
 
 
+template <>
+struct TWidgetTypeTraits<class SCircularThrobber>
+{
+	static constexpr bool SupportsInvalidation() { return true; }
+};
+
 /**
  * A throbber widget that orients images in a spinning circle.
  */
diff --git a/Engine/Source/Runtime/Slate/Public/Widgets/Input/SButton.h b/Engine/Source/Runtime/Slate/Public/Widgets/Input/SButton.h
--- a/Engine/Source/Runtime/Slate/Public/Widgets/Input/SButton.h
+++ b/Engine/Source/Runtime/Slate/Public/Widgets/Input/SButton.h
@@ -20,6 +20,12 @@ class FSlateWindowElementList;
 enum class ETextFlowDirection : uint8;
 enum class ETextShapingMethod : uint8;
 
+template <>
+struct TWidgetTypeTraits<class SButton>
+{
+	static constexpr bool SupportsInvalidation() { return true; }
+};
+
 /**
  * Slate's Buttons are clickable Widgets that can contain arbitrary widgets as its Content().
  */
diff --git a/Engine/Source/Runtime/Slate/Public/Widgets/Layout/SBorder.h b/Engine/Source/Runtime/Slate/Public/Widgets/Layout/SBorder.h
--- a/Engine/Source/Runtime/Slate/Public/Widgets/Layout/SBorder.h
+++ b/Engine/Source/Runtime/Slate/Public/Widgets/Layout/SBorder.h
@@ -15,6 +15,12 @@
 class FPaintArgs;
 class FSlateWindowElementList;
 
+template <>
+struct TWidgetTypeTraits<class SBorder>
+{
+	static constexpr bool SupportsInvalidation() { return true; }
+};
+
 /**
  * A border is an widget that can be used to contain other widgets.
  * It has a BorderImage property, which allows it to take on different appearances.
diff --git a/Engine/Source/Runtime/Slate/Public/Widgets/Layout/SBox.h b/Engine/Source/Runtime/Slate/Public/Widgets/Layout/SBox.h
--- a/Engine/Source/Runtime/Slate/Public/Widgets/Layout/SBox.h
+++ b/Engine/Source/Runtime/Slate/Public/Widgets/Layout/SBox.h
@@ -16,6 +16,12 @@ class FArrangedChildren;
 class FPaintArgs;
 class FSlateWindowElementList;
 
+template <>
+struct TWidgetTypeTraits<class SBox>
+{
+	static constexpr bool SupportsInvalidation() { return true; }
+};
+
 /**
  * SBox is the simplest layout element.
  */
diff --git a/Engine/Source/Runtime/Slate/Public/Widgets/Layout/SScrollBarTrack.h b/Engine/Source/Runtime/Slate/Public/Widgets/Layout/SScrollBarTrack.h
--- a/Engine/Source/Runtime/Slate/Public/Widgets/Layout/SScrollBarTrack.h
+++ b/Engine/Source/Runtime/Slate/Public/Widgets/Layout/SScrollBarTrack.h
@@ -11,6 +11,12 @@
 
 class FArrangedChildren;
 
+template <>
+struct TWidgetTypeTraits<class SScrollBarTrack>
+{
+	static constexpr bool SupportsInvalidation() { return true; }
+};
+
 /** Arrange 3 widgets: the top track, bottom track, and thumb. */
 class SScrollBarTrack : public SPanel
 {
diff --git a/Engine/Source/Runtime/Slate/Public/Widgets/Layout/SSeparator.h b/Engine/Source/Runtime/Slate/Public/Widgets/Layout/SSeparator.h
--- a/Engine/Source/Runtime/Slate/Public/Widgets/Layout/SSeparator.h
+++ b/Engine/Source/Runtime/Slate/Public/Widgets/Layout/SSeparator.h
@@ -8,6 +8,12 @@
 #include "Styling/CoreStyle.h"
 #include "Widgets/Layout/SBorder.h"
 
+template <>
+struct TWidgetTypeTraits<class SSeparator>
+{
+	static constexpr bool SupportsInvalidation() { return true; }
+};
+
 class SSeparator : public SBorder
 {
 public:
diff --git a/Engine/Source/Runtime/Slate/Public/Widgets/Layout/SSpacer.h b/Engine/Source/Runtime/Slate/Public/Widgets/Layout/SSpacer.h
--- a/Engine/Source/Runtime/Slate/Public/Widgets/Layout/SSpacer.h
+++ b/Engine/Source/Runtime/Slate/Public/Widgets/Layout/SSpacer.h
@@ -11,6 +11,12 @@
 class FPaintArgs;
 class FSlateWindowElementList;
 
+template <>
+struct TWidgetTypeTraits<class SSpacer>
+{
+	static constexpr bool SupportsInvalidation() { return true; }
+};
+
 class SSpacer : public SLeafWidget
 {
 	SLATE_DECLARE_WIDGET_API(SSpacer, SLeafWidget, SLATE_API)
diff --git a/Engine/Source/Runtime/Slate/Public/Widgets/SViewport.h b/Engine/Source/Runtime/Slate/Public/Widgets/SViewport.h
--- a/Engine/Source/Runtime/Slate/Public/Widgets/SViewport.h
+++ b/Engine/Source/Runtime/Slate/Public/Widgets/SViewport.h
@@ -19,6 +19,13 @@ class FPaintArgs;
 class FSlateWindowElementList;
 class ICustomHitTestPath;
 
+
+template <>
+struct TWidgetTypeTraits<class SViewport>
+{
+	static constexpr bool SupportsInvalidation() { return true; }
+};
+
 class SViewport
 	: public SCompoundWidget
 {
diff --git a/Engine/Source/Runtime/Slate/Public/Widgets/Text/STextBlock.h b/Engine/Source/Runtime/Slate/Public/Widgets/Text/STextBlock.h
--- a/Engine/Source/Runtime/Slate/Public/Widgets/Text/STextBlock.h
+++ b/Engine/Source/Runtime/Slate/Public/Widgets/Text/STextBlock.h
@@ -32,6 +32,11 @@ namespace ETextRole
 	};
 }
 
+template <>
+struct TWidgetTypeTraits<class STextBlock>
+{
+	static constexpr bool SupportsInvalidation() { return true; }
+};
 
 /**
  * A simple static text widget
diff --git a/Engine/Source/Runtime/Slate/Public/Widgets/Views/SHeaderRow.h b/Engine/Source/Runtime/Slate/Public/Widgets/Views/SHeaderRow.h
--- a/Engine/Source/Runtime/Slate/Public/Widgets/Views/SHeaderRow.h
+++ b/Engine/Source/Runtime/Slate/Public/Widgets/Views/SHeaderRow.h
@@ -90,6 +90,12 @@ DECLARE_DELEGATE_OneParam( FOnWidthChanged, float );
 /**	Callback to fetch the max row width for a specified column id */
 DECLARE_DELEGATE_RetVal_TwoParams(FVector2D, FOnGetMaxRowSizeForColumn, const FName&, EOrientation);
 
+template <>
+struct TWidgetTypeTraits<class SHeaderRow>
+{
+	static constexpr bool SupportsInvalidation() { return true; }
+};
+
 /**
  * The header that appears above lists and trees when they are showing multiple columns.
  */
