[CrashReport][Assert] ContentBrowser!SContentBrowser::GetItemContextMenu() [SContentBrowser.cpp]

Unfortunately, could not reproduce.
It seems like PathView (the folders tree view, left of the Content Browser) can end up in a state where it asks for an asset context menu -  which should not be possible since that view should only show Folders.

Tentative fix is to wrap things differently, so that we never reach code which is supposed for AssetView when dealing with a PathView.

diff --git a/Engine/Source/Editor/ContentBrowser/Private/SContentBrowser.cpp b/Engine/Source/Editor/ContentBrowser/Private/SContentBrowser.cpp
--- a/Engine/Source/Editor/ContentBrowser/Private/SContentBrowser.cpp
+++ b/Engine/Source/Editor/ContentBrowser/Private/SContentBrowser.cpp
@@ -5029,6 +5029,7 @@ TSharedPtr<SWidget> SContentBrowser::GetItemContextMenu(TArrayView<const FConten
 	const bool bIsAssetViewContext = ViewContext == EContentBrowserViewContext::AssetView;
 	const bool bShouldForceAddMenu = bIsControlDown && bIsAssetViewContext;
 
+	// This can happen with both AssetView and PathView
 	if (SelectedFolders.Num() > 0 && !bShouldForceAddMenu)
 	{
 		// Folders selected - show the folder menu
@@ -5119,16 +5120,17 @@ TSharedPtr<SWidget> SContentBrowser::GetItemContextMenu(TArrayView<const FConten
 		Context->SelectedPackagePaths = MoveTemp(SelectedPackagePaths);
 		return UToolMenus::Get()->GenerateWidget("ContentBrowser.FolderContextMenu", MenuContext);
 	}
-	else if (SelectedItems.Num() > 0 && !bShouldForceAddMenu)
+	else if (ViewContext == EContentBrowserViewContext::AssetView)
 	{
-		// Files selected - show the file menu
-		checkf(ViewContext == EContentBrowserViewContext::AssetView, TEXT("File items were passed from a path view!"));
-		return LegacyContentSourceWidgets->AssetContextMenu->MakeContextMenu(SelectedItems, LegacyContentSourceWidgets->AssetViewPtr->GetContentSources(), Commands);
-	}
-	else if (ViewContext == EContentBrowserViewContext::AssetView || bShouldForceAddMenu)
-	{
-		// Nothing selected - show the new asset menu
-		return MakeAddNewContextMenu(EContentBrowserDataMenuContext_AddNewMenuDomain::AssetView, nullptr);
+		if (SelectedItems.IsEmpty() || bShouldForceAddMenu)
+		{
+			// Nothing selected, or forcing AddMenu: show the new asset menu
+			return MakeAddNewContextMenu(EContentBrowserDataMenuContext_AddNewMenuDomain::AssetView, nullptr);
+		}
+		else
+		{
+			return LegacyContentSourceWidgets->AssetContextMenu->MakeContextMenu(SelectedItems, LegacyContentSourceWidgets->AssetViewPtr->GetContentSources(), Commands);
+        }
 	}
 
 	return nullptr;
