[CrashReport][Crash] SlateCore!SWindow::NotifyWindowBeingDestroyed() [RefCounting.h]

- Uses Multicast delegate instead of unicast (which is also soft deprecated in `SWindow::OnIsActiveChanged(...)`

- Adds Unbind when deactivating window

diff --git a/Engine/Source/Editor/ContentBrowser/Private/SRenameWindow.cpp b/Engine/Source/Editor/ContentBrowser/Private/SRenameWindow.cpp
--- a/Engine/Source/Editor/ContentBrowser/Private/SRenameWindow.cpp
+++ b/Engine/Source/Editor/ContentBrowser/Private/SRenameWindow.cpp
@@ -121,11 +121,13 @@ void SRenameWindow::Construct(const FArguments& InArgs)
 	InArgs._OnBeginTextEdit.ExecuteIfBound(InArgs._InitialText);
 
 	SetWidgetToFocusOnActivate(RenameEditableTextBox);
-	OnWindowDeactivated.BindSP(this, &SRenameWindow::DeactivateWindow);
+	GetOnWindowDeactivatedEvent().AddSP(this, &SRenameWindow::DeactivateWindow);
 }
 
 void SRenameWindow::DeactivateWindow()
 {
+	GetOnWindowDeactivatedEvent().RemoveAll(this);
+
 	if (RenameEditableTextBox.IsValid())
 	{
 		// Always force it when deactivating the window
