[Content Browser] "Filter Recursively" option in the Content Browser does not function

Menu from CB new style was setting the flag, but flag was never used.

diff --git a/Engine/Source/Editor/ContentBrowser/Private/SAssetView.cpp b/Engine/Source/Editor/ContentBrowser/Private/SAssetView.cpp
--- a/Engine/Source/Editor/ContentBrowser/Private/SAssetView.cpp
+++ b/Engine/Source/Editor/ContentBrowser/Private/SAssetView.cpp
@@ -3649,6 +3649,15 @@ bool SAssetView::IsFilteringRecursively() const
 		return false;
 	}
 
+	// New CB Style has setting among asset filter options
+	if (UE::Editor::ContentBrowser::IsNewStyleEnabled())
+	{
+		if (TSharedPtr<SFilterList> FilterBarPinned = FilterBar.Pin())
+		{
+			return FilterBarPinned->IsFilteringPathsRecursively();
+		}
+	}
+
 	// In some cases we want to not filter recursively even if we have a backend filter (e.g. the open level window)
 	// Most of the time, bFilterRecursivelyWithBackendFilter is true
 	if (const FContentBrowserInstanceConfig* EditorConfig = GetContentBrowserConfig())
diff --git a/Engine/Source/Editor/ContentBrowser/Private/SFilterList.cpp b/Engine/Source/Editor/ContentBrowser/Private/SFilterList.cpp
--- a/Engine/Source/Editor/ContentBrowser/Private/SFilterList.cpp
+++ b/Engine/Source/Editor/ContentBrowser/Private/SFilterList.cpp
@@ -462,7 +462,7 @@ void SFilterList::PopulateAddFilterMenu(UToolMenu* Menu, TSharedPtr<FFilterCateg
 					FExecuteAction::CreateSP(this, &SFilterList::ToggleFilterPathsRecursively),
 					FCanExecuteAction(),
 					FIsActionChecked::CreateSP(this, &SFilterList::IsFilteringPathsRecursively)),
-				EUserInterfaceActionType::Check
+				EUserInterfaceActionType::ToggleButton
 			);
 		}
 
