Fix for crash when UE closes because when this `SAssetView` is destroyed, it destroys its `FAssetThumbnailPool` which accesses a `UThumbnailManager` via `IsValid(UThumbnailManager::TryGet())`. The crash happens in `IsValid`, likely because the pointer returned by `UThumbnailManager::TryGet()` is stale.
The crash was fixed by using a WeakPtr instead of a SharedPtr

diff --git a/Engine/Source/Editor/ContentBrowser/Private/Menus/AssetViewMenus.cpp b/Engine/Source/Editor/ContentBrowser/Private/Menus/AssetViewMenus.cpp
--- a/Engine/Source/Editor/ContentBrowser/Private/Menus/AssetViewMenus.cpp
+++ b/Engine/Source/Editor/ContentBrowser/Private/Menus/AssetViewMenus.cpp
@@ -611,10 +611,11 @@ FDelayedAutoRegisterHelper SAssetView::ToolBarMenuExtensionRegistration(
 							FToolUIAction SortAction;
 							SortAction.IsActionVisibleDelegate = FToolMenuIsActionButtonVisible::CreateSPLambda(
 								AssetView,
-								[AssetView](const FToolMenuContext&)
+								[WeakAssetView = AssetView.ToWeakPtr()](const FToolMenuContext&)
 								{
 									// Sort button only visible in Tile view mode - others have their own sorting UI representation.
-									return AssetView->IsCurrentViewType(EAssetViewType::Tile);
+									const TSharedPtr<SAssetView> PinnedAssetView = WeakAssetView.Pin();
+									return PinnedAssetView && PinnedAssetView->IsCurrentViewType(EAssetViewType::Tile);
 								});
 
 							FToolMenuEntry& SortEntry = InSection.AddEntry(
