Fixed Content Browser not syncing to items that have an a item in the alias data source, but not the asset data source.
Fixed Path Picker not selecting the current asset by default if it is an alias.

diff --git a/Engine/Plugins/Editor/ContentBrowser/ContentBrowserAliasDataSource/Source/ContentBrowserAliasDataSource/Private/ContentBrowserAliasDataSource.cpp b/Engine/Plugins/Editor/ContentBrowser/ContentBrowserAliasDataSource/Source/ContentBrowserAliasDataSource/Private/ContentBrowserAliasDataSource.cpp
--- a/Engine/Plugins/Editor/ContentBrowser/ContentBrowserAliasDataSource/Source/ContentBrowserAliasDataSource/Private/ContentBrowserAliasDataSource.cpp
+++ b/Engine/Plugins/Editor/ContentBrowser/ContentBrowserAliasDataSource/Source/ContentBrowserAliasDataSource/Private/ContentBrowserAliasDataSource.cpp
@@ -1096,10 +1096,7 @@ bool UContentBrowserAliasDataSource::Legacy_TryConvertPackagePathToVirtualPath(c
 
 bool UContentBrowserAliasDataSource::Legacy_TryConvertAssetDataToVirtualPath(const FAssetData& InAssetData, const bool InUseFolderPaths, FName& OutPath)
 {
-	return InAssetData.AssetClassPath != FTopLevelAssetPath(TEXT("/Script/CoreUObject"), NAME_Class) // Ignore legacy class items
-PRAGMA_DISABLE_DEPRECATION_WARNINGS
-		&& TryConvertInternalPathToVirtual(InUseFolderPaths ? InAssetData.PackagePath : InAssetData.ObjectPath, OutPath);
-PRAGMA_ENABLE_DEPRECATION_WARNINGS
+	return false;
 }
 
 void UContentBrowserAliasDataSource::RemoveUnusedCachedFilterData(const FContentBrowserDataFilterCacheIDOwner& IDOwner, TArrayView<const FName> InVirtualPathsInUse, const FContentBrowserDataFilter& DataFilter)
diff --git a/Engine/Plugins/Editor/ContentBrowser/ContentBrowserAssetDataSource/Source/ContentBrowserAssetDataSource/Private/ContentBrowserAssetDataSource.cpp b/Engine/Plugins/Editor/ContentBrowser/ContentBrowserAssetDataSource/Source/ContentBrowserAssetDataSource/Private/ContentBrowserAssetDataSource.cpp
--- a/Engine/Plugins/Editor/ContentBrowser/ContentBrowserAssetDataSource/Source/ContentBrowserAssetDataSource/Private/ContentBrowserAssetDataSource.cpp
+++ b/Engine/Plugins/Editor/ContentBrowser/ContentBrowserAssetDataSource/Source/ContentBrowserAssetDataSource/Private/ContentBrowserAssetDataSource.cpp
@@ -3249,6 +3249,7 @@ bool UContentBrowserAssetDataSource::Legacy_TryConvertPackagePathToVirtualPath(c
 bool UContentBrowserAssetDataSource::Legacy_TryConvertAssetDataToVirtualPath(const FAssetData& InAssetData, const bool InUseFolderPaths, FName& OutPath)
 {
 	return InAssetData.AssetClassPath != FTopLevelAssetPath(TEXT("/Script/CoreUObject"), TEXT("Class")) // Ignore legacy class items
+		&& ContentBrowserAssetData::IsPrimaryAsset(InAssetData)
 PRAGMA_DISABLE_DEPRECATION_WARNINGS
 		&& TryConvertInternalPathToVirtual(InUseFolderPaths ? InAssetData.PackagePath : InAssetData.ObjectPath, OutPath);
 PRAGMA_ENABLE_DEPRECATION_WARNINGS
diff --git a/Engine/Source/Editor/ContentBrowser/Private/ContentBrowserUtils.cpp b/Engine/Source/Editor/ContentBrowser/Private/ContentBrowserUtils.cpp
--- a/Engine/Source/Editor/ContentBrowser/Private/ContentBrowserUtils.cpp
+++ b/Engine/Source/Editor/ContentBrowser/Private/ContentBrowserUtils.cpp
@@ -866,42 +866,87 @@ bool ContentBrowserUtils::CanExploreFolders(const TArray<FContentBrowserItem>& I
 	return false;
 }
 
-template <typename OutputContainerType>
-void ConvertLegacySelectionToVirtualPathsImpl(TArrayView<const FAssetData> InAssets, TArrayView<const FString> InFolders, const bool InUseFolderPaths, OutputContainerType& OutVirtualPaths)
+TArray<FContentBrowserItem> ContentBrowserUtils::ConvertLegacySelectionToItems(TArrayView<const FAssetData> InAssets, TArrayView<const FString> InFolders)
 {
-	OutVirtualPaths.Reset();
-	if (InAssets.Num() == 0 && InFolders.Num() == 0)
-	{
-		return;
-	}
-
 	UContentBrowserDataSubsystem* ContentBrowserData = IContentBrowserDataModule::Get().GetSubsystem();
 
-	auto AppendVirtualPath = [&OutVirtualPaths](FName InPath)
-	{
-		OutVirtualPaths.Add(InPath);
-		return true;
-	};
+	TArray<FContentBrowserItem> Items;
+	Items.Reserve(InAssets.Num() + InFolders.Num());
 
-	for (const FAssetData& Asset : InAssets)
+	if (!InAssets.IsEmpty())
 	{
-		ContentBrowserData->Legacy_TryConvertAssetDataToVirtualPaths(Asset, InUseFolderPaths, AppendVirtualPath);
+		const TSharedRef<FPathPermissionList>& FolderPermissions = FAssetToolsModule::GetModule().Get().GetFolderPermissionList();
+
+		TArray<FName> VirtualPaths;
+		VirtualPaths.Reserve(InAssets.Num());
+
+		auto AppendVirtualPath = [&VirtualPaths](FName VirtualPath)
+		{
+			VirtualPaths.Add(VirtualPath);
+			return true;
+		};
+
+		for (const FAssetData& AssetData : InAssets)
+		{
+			if (FolderPermissions->PassesStartsWithFilter(AssetData.GetObjectPathString()))
+			{
+				const int32 PreviousNumVirtualPaths = VirtualPaths.Num();
+				ContentBrowserData->Legacy_TryConvertAssetDataToVirtualPaths(AssetData, /*UseFolderPaths*/false, AppendVirtualPath);
+				if (VirtualPaths.Num() > PreviousNumVirtualPaths)
+				{
+					// If we found a virtual path, continue, otherwise the item may only exist as an alias, so try those next.
+					continue;
+				}
+			}
+
+			TArray<FContentBrowserItemPath> Aliases = ContentBrowserData->GetAliasesForPath(*AssetData.GetObjectPathString());
+			for (const FContentBrowserItemPath& Alias : Aliases)
+			{
+				if (FolderPermissions->PassesStartsWithFilter(Alias.GetInternalPathName()))
+				{
+					VirtualPaths.Add(Alias.GetVirtualPathName());
+					break;
+				}
+			}
+		}
+
+		for (const FName& VirtualPath : VirtualPaths)
+		{
+			FContentBrowserItem ItemToSync = ContentBrowserData->GetItemAtPath(VirtualPath, EContentBrowserItemTypeFilter::IncludeFiles);
+			if (ItemToSync.IsValid())
+			{
+				Items.Add(MoveTemp(ItemToSync));
+			}
+		}
 	}
 
-	for (const FString& Folder : InFolders)
+	if (!InFolders.IsEmpty())
 	{
-		ContentBrowserData->Legacy_TryConvertPackagePathToVirtualPaths(*Folder, AppendVirtualPath);
-	}
-}
+		TArray<FName> VirtualPaths;
+		VirtualPaths.Reserve(InFolders.Num());
 
-void ContentBrowserUtils::ConvertLegacySelectionToVirtualPaths(TArrayView<const FAssetData> InAssets, TArrayView<const FString> InFolders, const bool InUseFolderPaths, TArray<FName>& OutVirtualPaths)
-{
-	ConvertLegacySelectionToVirtualPathsImpl(InAssets, InFolders, InUseFolderPaths, OutVirtualPaths);
-}
+		auto AppendVirtualPath = [&VirtualPaths](FName VirtualPath)
+		{
+			VirtualPaths.Add(VirtualPath);
+			return true;
+		};
 
-void ContentBrowserUtils::ConvertLegacySelectionToVirtualPaths(TArrayView<const FAssetData> InAssets, TArrayView<const FString> InFolders, const bool InUseFolderPaths, TSet<FName>& OutVirtualPaths)
-{
-	ConvertLegacySelectionToVirtualPathsImpl(InAssets, InFolders, InUseFolderPaths, OutVirtualPaths);
+		for (const FString& Folder : InFolders)
+		{
+			ContentBrowserData->Legacy_TryConvertPackagePathToVirtualPaths(*Folder, AppendVirtualPath);
+		}
+
+		for (const FName& VirtualPath : VirtualPaths)
+		{
+			FContentBrowserItem ItemToSync = ContentBrowserData->GetItemAtPath(VirtualPath, EContentBrowserItemTypeFilter::IncludeFolders);
+			if (ItemToSync.IsValid())
+			{
+				Items.Add(MoveTemp(ItemToSync));
+			}
+		}
+	}
+
+	return Items;
 }
 
 void ContentBrowserUtils::AppendAssetFilterToContentBrowserFilter(const FARFilter& InAssetFilter, const TSharedPtr<FPathPermissionList>& InAssetClassPermissionList, const TSharedPtr<FPathPermissionList>& InFolderPermissionList, FContentBrowserDataFilter& OutDataFilter)
diff --git a/Engine/Source/Editor/ContentBrowser/Private/ContentBrowserUtils.h b/Engine/Source/Editor/ContentBrowser/Private/ContentBrowserUtils.h
--- a/Engine/Source/Editor/ContentBrowser/Private/ContentBrowserUtils.h
+++ b/Engine/Source/Editor/ContentBrowser/Private/ContentBrowserUtils.h
@@ -127,9 +127,8 @@ namespace ContentBrowserUtils
 	/** Returns if can perform a batched "explore" operation on the specified file and/or folder paths */
 	bool CanExploreFolders(const TArray<FContentBrowserItem>& InItems);
 
-	/** Convert a legacy asset and path selection to their corresponding virtual paths for content browser data items */
-	void ConvertLegacySelectionToVirtualPaths(TArrayView<const FAssetData> InAssets, TArrayView<const FString> InFolders, const bool InUseFolderPaths, TArray<FName>& OutVirtualPaths);
-	void ConvertLegacySelectionToVirtualPaths(TArrayView<const FAssetData> InAssets, TArrayView<const FString> InFolders, const bool InUseFolderPaths, TSet<FName>& OutVirtualPaths);
+	/** Convert a legacy asset and path selection to their corresponding content browser data items */
+	TArray<FContentBrowserItem> ConvertLegacySelectionToItems(TArrayView<const FAssetData> InAssets, TArrayView<const FString> InFolders);
 
 	/** Append the asset registry filter and permission lists to the content browser data filter */
 	void AppendAssetFilterToContentBrowserFilter(const FARFilter& InAssetFilter, const TSharedPtr<FPathPermissionList>& InAssetClassPermissionList, const TSharedPtr<FPathPermissionList>& InFolderPermissionList, FContentBrowserDataFilter& OutDataFilter);
diff --git a/Engine/Source/Editor/ContentBrowser/Private/SAssetView.cpp b/Engine/Source/Editor/ContentBrowser/Private/SAssetView.cpp
--- a/Engine/Source/Editor/ContentBrowser/Private/SAssetView.cpp
+++ b/Engine/Source/Editor/ContentBrowser/Private/SAssetView.cpp
@@ -1976,10 +1976,7 @@ void SAssetView::SyncToVirtualPaths(TArrayView<const FName> VirtualPathsToSync,
 
 void SAssetView::SyncToLegacy(TArrayView<const FAssetData> AssetDataList, TArrayView<const FString> FolderList, const bool bFocusOnSync)
 {
-	PendingSyncItems.Reset();
-	ContentBrowserUtils::ConvertLegacySelectionToVirtualPaths(AssetDataList, FolderList, /*UseFolderPaths*/false, PendingSyncItems.SelectedVirtualPaths);
-	InitDeferredPendingSyncItems();
-	bPendingFocusOnSync = bFocusOnSync;
+	SyncToItems(ContentBrowserUtils::ConvertLegacySelectionToItems(AssetDataList, FolderList), bFocusOnSync);
 }
 
 void SAssetView::InitDeferredPendingSyncItems()
diff --git a/Engine/Source/Editor/ContentBrowser/Private/SContentBrowser.cpp b/Engine/Source/Editor/ContentBrowser/Private/SContentBrowser.cpp
--- a/Engine/Source/Editor/ContentBrowser/Private/SContentBrowser.cpp
+++ b/Engine/Source/Editor/ContentBrowser/Private/SContentBrowser.cpp
@@ -1576,17 +1576,6 @@ void SContentBrowser::PrepareToSyncVirtualPaths(TArrayView<const FName> VirtualP
 	PrepareToSyncItems(ItemsToSync, bDisableFiltersThatHideAssets);
 }
 
-void SContentBrowser::PrepareToSyncLegacy(TArrayView<const FAssetData> AssetDataList, TArrayView<const FString> FolderPaths, const bool bDisableFiltersThatHideAssets)
-{
-	// For now we just forcefully enable the legacy content source when a sync is requested so it succeeds
-	ContentSourcesContainer->ActivateLegacyContentSource();
-	
-	TArray<FName> VirtualPathsToSync;
-	ContentBrowserUtils::ConvertLegacySelectionToVirtualPaths(AssetDataList, FolderPaths, /*UseFolderPaths*/false, VirtualPathsToSync);
-
-	PrepareToSyncVirtualPaths(VirtualPathsToSync, bDisableFiltersThatHideAssets);
-}
-
 void SContentBrowser::SyncToAssets(TArrayView<const FAssetData> AssetDataList, const bool bAllowImplicitSync, const bool bDisableFiltersThatHideAssets)
 {
 	// For now we just forcefully enable the legacy content source when a sync is requested so it succeeds
@@ -1662,44 +1651,14 @@ void SContentBrowser::SyncToLegacy(TArrayView<const FAssetData> AssetDataList, T
 	// For now we just forcefully enable the legacy content source when a sync is requested so it succeeds
 	ContentSourcesContainer->ActivateLegacyContentSource();
 	
-	UContentBrowserDataSubsystem* ContentBrowserData = IContentBrowserDataModule::Get().GetSubsystem();
-	const TSharedRef<FPathPermissionList>& FolderPermissions = FAssetToolsModule::GetModule().Get().GetFolderPermissionList();
-
-	// If any of the items to sync don't pass the permission filter, try to map the item to a different one that might be visible
-	TArray<FAssetData> NewItemsToSync;
-	for (int32 i = 0; i < AssetDataList.Num(); ++i)
-	{
-		if (FolderPermissions->PassesStartsWithFilter(AssetDataList[i].GetObjectPathString()))
-		{
-			NewItemsToSync.Add(AssetDataList[i]);
-		}
-		else
-		{
-			TArray<FContentBrowserItemPath> Aliases = ContentBrowserData->GetAliasesForPath(*AssetDataList[i].GetObjectPathString());
-			for (const FContentBrowserItemPath& Alias : Aliases)
-			{
-				const FName InternalPath = Alias.GetInternalPathName();
-				if (FolderPermissions->PassesStartsWithFilter(InternalPath))
-				{
-					FSoftObjectPath ObjectPath{ FNameBuilder(InternalPath) };
-					if (ObjectPath.IsAsset())
-					{
-						FAssetData AliasAssetData(ObjectPath.GetLongPackageFName(), FName(*FPackageName::GetLongPackagePath(ObjectPath.GetLongPackageName())), ObjectPath.GetAssetFName(), AssetDataList[i].AssetClassPath, AssetDataList[i].TagsAndValues.CopyMap());
-						NewItemsToSync.Add(MoveTemp(AliasAssetData));
-						break;
-					}
-				}
-			}
-		}
-	}
+	TArray<FContentBrowserItem> ItemsToSync = ContentBrowserUtils::ConvertLegacySelectionToItems(AssetDataList, FolderList);
 
-	AssetDataList = NewItemsToSync;
-	PrepareToSyncLegacy(AssetDataList, FolderList, bDisableFiltersThatHideAssets);
+	PrepareToSyncItems(ItemsToSync, bDisableFiltersThatHideAssets);
 
 	// Tell the sources view first so the asset view will be up to date by the time we request the sync
-	LegacyContentSourceWidgets->PathViewPtr->SyncToLegacy(AssetDataList, FolderList, bAllowImplicitSync);
-	LegacyContentSourceWidgets->FavoritePathViewPtr->SyncToLegacy(AssetDataList, FolderList, bAllowImplicitSync);
-	LegacyContentSourceWidgets->AssetViewPtr->SyncToLegacy(AssetDataList, FolderList);
+	LegacyContentSourceWidgets->PathViewPtr->SyncToItems(ItemsToSync, bAllowImplicitSync);
+	LegacyContentSourceWidgets->FavoritePathViewPtr->SyncToItems(ItemsToSync, bAllowImplicitSync);
+	LegacyContentSourceWidgets->AssetViewPtr->SyncToItems(ItemsToSync);
 }
 
 void SContentBrowser::SyncTo( const FContentBrowserSelection& ItemSelection, const bool bAllowImplicitSync, const bool bDisableFiltersThatHideAssets )
diff --git a/Engine/Source/Editor/ContentBrowser/Private/SContentBrowser.h b/Engine/Source/Editor/ContentBrowser/Private/SContentBrowser.h
--- a/Engine/Source/Editor/ContentBrowser/Private/SContentBrowser.h
+++ b/Engine/Source/Editor/ContentBrowser/Private/SContentBrowser.h
@@ -290,9 +290,6 @@ private:
 	/** Called prior to syncing the selection in this Content Browser */
 	void PrepareToSyncVirtualPaths(TArrayView<const FName> VirtualPathsToSync, const bool bDisableFiltersThatHideAssets);
 
-	/** Called prior to syncing the selection in this Content Browser */
-	void PrepareToSyncLegacy(TArrayView<const FAssetData> AssetDataList, TArrayView<const FString> FolderPaths, const bool bDisableFiltersThatHideAssets);
-
 	/** Called to retrieve the text that should be highlighted on assets */
 	FText GetHighlightedText() const;
 
diff --git a/Engine/Source/Editor/ContentBrowser/Private/SPathView.cpp b/Engine/Source/Editor/ContentBrowser/Private/SPathView.cpp
--- a/Engine/Source/Editor/ContentBrowser/Private/SPathView.cpp
+++ b/Engine/Source/Editor/ContentBrowser/Private/SPathView.cpp
@@ -1817,14 +1817,6 @@ void SPathView::SyncToVirtualPaths(TArrayView<const FName> VirtualPathsToSync, c
 	}
 }
 
-void SPathView::SyncToLegacy(TArrayView<const FAssetData> AssetDataList, TArrayView<const FString> FolderList, const bool bAllowImplicitSync)
-{
-	TArray<FName> VirtualPathsToSync;
-	ContentBrowserUtils::ConvertLegacySelectionToVirtualPaths(AssetDataList, FolderList, /*UseFolderPaths*/true, VirtualPathsToSync);
-
-	SyncToVirtualPaths(VirtualPathsToSync, bAllowImplicitSync);
-}
-
 bool SPathView::DoesItemExist(FName InVirtualPath) const
 {
 	return TreeData->FindTreeItem(InVirtualPath).IsValid();
diff --git a/Engine/Source/Editor/ContentBrowser/Private/SPathView.h b/Engine/Source/Editor/ContentBrowser/Private/SPathView.h
--- a/Engine/Source/Editor/ContentBrowser/Private/SPathView.h
+++ b/Engine/Source/Editor/ContentBrowser/Private/SPathView.h
@@ -228,17 +228,6 @@ public:
 	 */
 	void SyncToVirtualPaths( TArrayView<const FName> VirtualPathsToSync, const bool bAllowImplicitSync = false );
 
-	/**
-	 * Selects the paths containing the specified assets and paths.
-	 *
-	 *	@param AssetDataList		- A list of assets to sync the view to
-	 *	@param FolderList			- A list of folders to sync the view to
-	 *
-	 *	@param bAllowImplicitSync	- true to allow the view to sync to parent folders if they are already selected,
-	 *								  false to force the view to select the explicit Parent folders of each asset
-	 */
-	void SyncToLegacy( TArrayView<const FAssetData> AssetDataList, TArrayView<const FString> FolderList, const bool bAllowImplicitSync = false );
-
 	/** Returns whether the tree contains an item with the given virtual path. */
 	bool DoesItemExist(FName InVirtualPath) const;
 
