Removed unused SContentBrowser::SyncToVirtualPaths and SPathView::SyncToVirtualPaths.  The filtering was inconsistent with SyncToItems and we didn't handle syncing the path view to a file's folder correctly.
Deprecated SAssetView::SyncToVirtualPaths.

diff --git a/Engine/Source/Editor/ContentBrowser/Private/SContentBrowser.cpp b/Engine/Source/Editor/ContentBrowser/Private/SContentBrowser.cpp
--- a/Engine/Source/Editor/ContentBrowser/Private/SContentBrowser.cpp
+++ b/Engine/Source/Editor/ContentBrowser/Private/SContentBrowser.cpp
@@ -1552,30 +1552,6 @@ void SContentBrowser::PrepareToSyncItems(TArrayView<const FContentBrowserItem> I
 	LegacyContentSourceWidgets->SearchBoxPtr->SetError(FText::GetEmpty());
 }
 
-void SContentBrowser::PrepareToSyncVirtualPaths(TArrayView<const FName> VirtualPathsToSync, const bool bDisableFiltersThatHideAssets)
-{
-	// For now we just forcefully enable the legacy content source when a sync is requested so it succeeds
-	ContentSourcesContainer->ActivateLegacyContentSource();
-	
-	// We need to try and resolve these paths back to items in order to query their attributes
-	// This will only work for items that have already been discovered
-	TArray<FContentBrowserItem> ItemsToSync;
-	{
-		UContentBrowserDataSubsystem* ContentBrowserData = IContentBrowserDataModule::Get().GetSubsystem();
-
-		for (const FName& VirtualPathToSync : VirtualPathsToSync)
-		{
-			FContentBrowserItem ItemToSync = ContentBrowserData->GetItemAtPath(VirtualPathToSync, EContentBrowserItemTypeFilter::IncludeAll);
-			if (ItemToSync.IsValid())
-			{
-				ItemsToSync.Add(MoveTemp(ItemToSync));
-			}
-		}
-	}
-
-	PrepareToSyncItems(ItemsToSync, bDisableFiltersThatHideAssets);
-}
-
 void SContentBrowser::SyncToAssets(TArrayView<const FAssetData> AssetDataList, const bool bAllowImplicitSync, const bool bDisableFiltersThatHideAssets)
 {
 	// For now we just forcefully enable the legacy content source when a sync is requested so it succeeds
@@ -1607,45 +1583,6 @@ void SContentBrowser::SyncToItems(TArrayView<const FContentBrowserItem> ItemsToS
 	LegacyContentSourceWidgets->AssetViewPtr->SyncToItems(ItemsToSync);
 }
 
-void SContentBrowser::SyncToVirtualPaths(TArrayView<const FName> VirtualPathsToSync, const bool bAllowImplicitSync, const bool bDisableFiltersThatHideAssets)
-{
-	// For now we just forcefully enable the legacy content source when a sync is requested so it succeeds
-	ContentSourcesContainer->ActivateLegacyContentSource();
-	
-	UContentBrowserDataSubsystem* ContentBrowserData = IContentBrowserDataModule::Get().GetSubsystem();
-	const TSharedRef<FPathPermissionList>& FolderPermissions = FAssetToolsModule::GetModule().Get().GetFolderPermissionList();
-
-	// If any of the items to sync don't pass the permission filter, try to map the item to a different one that might be visible
-	TArray<FName> NewItemsToSync;
-	for (const FName VirtualPath : VirtualPathsToSync)
-	{
-		FName InternalPath;
-		ContentBrowserData->TryConvertVirtualPath(VirtualPath, InternalPath);
-		if (FolderPermissions->PassesStartsWithFilter(InternalPath))
-		{
-			NewItemsToSync.Add(VirtualPath);
-		}
-		else
-		{
-			TArray<FContentBrowserItemPath> Aliases = ContentBrowserData->GetAliasesForPath(InternalPath);
-			for (const FContentBrowserItemPath& Alias : Aliases)
-			{
-				if (FolderPermissions->PassesStartsWithFilter(Alias.GetInternalPathName()))
-				{
-					NewItemsToSync.Add(Alias.GetVirtualPathName());
-				}
-			}
-		}
-	}
-	VirtualPathsToSync = NewItemsToSync;
-	PrepareToSyncVirtualPaths(VirtualPathsToSync, bDisableFiltersThatHideAssets);
-
-	// Tell the sources view first so the asset view will be up to date by the time we request the sync
-	LegacyContentSourceWidgets->PathViewPtr->SyncToVirtualPaths(VirtualPathsToSync, bAllowImplicitSync);
-	LegacyContentSourceWidgets->FavoritePathViewPtr->SyncToVirtualPaths(VirtualPathsToSync, bAllowImplicitSync);
-	LegacyContentSourceWidgets->AssetViewPtr->SyncToVirtualPaths(VirtualPathsToSync);
-}
-
 void SContentBrowser::SyncToLegacy(TArrayView<const FAssetData> AssetDataList, TArrayView<const FString> FolderList, const bool bAllowImplicitSync, const bool bDisableFiltersThatHideAssets)
 {
 	// For now we just forcefully enable the legacy content source when a sync is requested so it succeeds
diff --git a/Engine/Source/Editor/ContentBrowser/Private/SContentBrowser.h b/Engine/Source/Editor/ContentBrowser/Private/SContentBrowser.h
--- a/Engine/Source/Editor/ContentBrowser/Private/SContentBrowser.h
+++ b/Engine/Source/Editor/ContentBrowser/Private/SContentBrowser.h
@@ -156,16 +156,6 @@ public:
 	 */
 	void SyncToItems( TArrayView<const FContentBrowserItem> ItemsToSync, const bool bAllowImplicitSync = false, const bool bDisableFiltersThatHideAssets = true );
 
-	/**
-	 * Changes sources to show the specified items and selects them in the asset view
-	 *
-	 *	@param VirtualPathsToSync	- A list of virtual paths to sync the view to
-	 *
-	 *	@param bAllowImplicitSync	- true to allow the view to sync to parent folders if they are already selected,
-	 *								  false to force the view to select the explicit Parent folders of each asset
-	 */
-	void SyncToVirtualPaths( TArrayView<const FName> VirtualPathsToSync, const bool bAllowImplicitSync = false, const bool bDisableFiltersThatHideAssets = true );
-
 	/**
 	 * Changes sources to show the specified assets and folders and selects them in the asset view
 	 *
@@ -287,9 +277,6 @@ private:
 	/** Called prior to syncing the selection in this Content Browser */
 	void PrepareToSyncItems(TArrayView<const FContentBrowserItem> ItemsToSync, const bool bDisableFiltersThatHideAssets);
 
-	/** Called prior to syncing the selection in this Content Browser */
-	void PrepareToSyncVirtualPaths(TArrayView<const FName> VirtualPathsToSync, const bool bDisableFiltersThatHideAssets);
-
 	/** Called to retrieve the text that should be highlighted on assets */
 	FText GetHighlightedText() const;
 
diff --git a/Engine/Source/Editor/ContentBrowser/Private/SPathView.cpp b/Engine/Source/Editor/ContentBrowser/Private/SPathView.cpp
--- a/Engine/Source/Editor/ContentBrowser/Private/SPathView.cpp
+++ b/Engine/Source/Editor/ContentBrowser/Private/SPathView.cpp
@@ -1729,28 +1729,20 @@ bool SPathView::InternalPathPassesBlockLists(const FStringView InInternalPath, c
 
 void SPathView::SyncToItems(TArrayView<const FContentBrowserItem> ItemsToSync, const bool bAllowImplicitSync)
 {
-	TArray<FName> VirtualPathsToSync;
-	for (const FContentBrowserItem& Item : ItemsToSync)
+	TSet<TSharedRef<FTreeItem>> SyncTreeItems;
+	for (const FContentBrowserItem& ItemToSync : ItemsToSync)
 	{
-		if (Item.IsFile())
+		FName VirtualPathToSync;
+		if (ItemToSync.IsFile())
 		{
 			// Files need to sync their parent folder in the tree, so chop off the end of their path
-			VirtualPathsToSync.Add(*FPaths::GetPath(Item.GetVirtualPath().ToString()));
+			VirtualPathToSync = FName(*FPaths::GetPath(ItemToSync.GetVirtualPath().ToString()));
 		}
 		else
 		{
-			VirtualPathsToSync.Add(Item.GetVirtualPath());
+			VirtualPathToSync = ItemToSync.GetVirtualPath();
 		}
-	}
-
-	SyncToVirtualPaths(VirtualPathsToSync, bAllowImplicitSync);
-}
 
-void SPathView::SyncToVirtualPaths(TArrayView<const FName> VirtualPathsToSync, const bool bAllowImplicitSync)
-{
-	TSet<TSharedRef<FTreeItem>> SyncTreeItems;
-	for (const FName& VirtualPathToSync : VirtualPathsToSync)
-	{
 		TSharedPtr<FTreeItem> Item = TreeData->FindTreeItem(VirtualPathToSync);
 		if (Item.IsValid())
 		{
diff --git a/Engine/Source/Editor/ContentBrowser/Private/SPathView.h b/Engine/Source/Editor/ContentBrowser/Private/SPathView.h
--- a/Engine/Source/Editor/ContentBrowser/Private/SPathView.h
+++ b/Engine/Source/Editor/ContentBrowser/Private/SPathView.h
@@ -218,16 +218,6 @@ public:
 	 */
 	void SyncToItems( TArrayView<const FContentBrowserItem> ItemsToSync, const bool bAllowImplicitSync = false );
 
-	/**
-	 * Selects the given virtual paths.
-	 *
-	 *	@param VirtualPathsToSync	- A list of virtual paths to sync the view to
-	 *
-	 *	@param bAllowImplicitSync	- true to allow the view to sync to parent folders if they are already selected,
-	 *								  false to force the view to select the explicit Parent folders of each asset
-	 */
-	void SyncToVirtualPaths( TArrayView<const FName> VirtualPathsToSync, const bool bAllowImplicitSync = false );
-
 	/** Returns whether the tree contains an item with the given virtual path. */
 	bool DoesItemExist(FName InVirtualPath) const;
 
diff --git a/Engine/Source/Editor/ContentBrowser/Public/SAssetView.h b/Engine/Source/Editor/ContentBrowser/Public/SAssetView.h
--- a/Engine/Source/Editor/ContentBrowser/Public/SAssetView.h
+++ b/Engine/Source/Editor/ContentBrowser/Public/SAssetView.h
@@ -381,6 +381,7 @@ public:
 	UE_API void SyncToItems( TArrayView<const FContentBrowserItem> ItemsToSync, const bool bFocusOnSync = true );
 
 	/** Selects the specified virtual paths. */
+	UE_DEPRECATED(5.8, "Use SyncToItems() instead.")
 	UE_API void SyncToVirtualPaths( TArrayView<const FName> VirtualPathsToSync, const bool bFocusOnSync = true );
 
 	/** Selects the specified assets and paths. */
