[Content Browser] Fix an issue on Mac where the traffic buttons showed up on the rename window.

[RN] Minor, Editor

diff --git a/Engine/Source/Editor/ContentBrowser/Private/SRenameWindow.cpp b/Engine/Source/Editor/ContentBrowser/Private/SRenameWindow.cpp
--- a/Engine/Source/Editor/ContentBrowser/Private/SRenameWindow.cpp
+++ b/Engine/Source/Editor/ContentBrowser/Private/SRenameWindow.cpp
@@ -101,7 +101,7 @@ void SRenameWindow::Construct(const FArguments& InArgs)
 		.SizingRule(ESizingRule::Autosized)
 		.AutoCenter(EAutoCenter::None)
 		.ScreenPosition(InArgs._ScreenPosition * (1.f/DPIScale))
-		.Type(EWindowType::Normal)
+		.Type(EWindowType::ToolTip)
 		.LayoutBorder(0)
 		.ActivationPolicy(EWindowActivationPolicy::FirstShown)
 		.bDragAnywhere(false)
