Fixed collection view not showing collections when a new collection is created procedurally and there were previously no existing collections.

diff --git a/Engine/Source/Editor/ContentBrowser/Private/SCollectionView.cpp b/Engine/Source/Editor/ContentBrowser/Private/SCollectionView.cpp
--- a/Engine/Source/Editor/ContentBrowser/Private/SCollectionView.cpp
+++ b/Engine/Source/Editor/ContentBrowser/Private/SCollectionView.cpp
@@ -269,8 +269,15 @@ void SCollectionView::Construct( const FArguments& InArgs )
 	UpdateCollectionItems();
 }
 
-bool SCollectionView::IsEmpty() const
+bool SCollectionView::IsEmpty()
 {
+	// If the SCollectionView is hidden because it is empty, it won't tick, so we need to make sure AvailableCollections is up to date.
+	if (bQueueCollectionItemsUpdate)
+	{
+		bQueueCollectionItemsUpdate = false;
+		UpdateCollectionItems();
+	}
+
 	return AvailableCollections.IsEmpty();
 }
 
diff --git a/Engine/Source/Editor/ContentBrowser/Private/SCollectionView.h b/Engine/Source/Editor/ContentBrowser/Private/SCollectionView.h
--- a/Engine/Source/Editor/ContentBrowser/Private/SCollectionView.h
+++ b/Engine/Source/Editor/ContentBrowser/Private/SCollectionView.h
@@ -99,7 +99,7 @@ public:
 	void Construct( const FArguments& InArgs );
 
 	/** Returns true if there are no collections being presented in this widget. */
-	virtual bool IsEmpty() const;
+	bool IsEmpty();
 
 	/** Selects the specified collections */
 	void SetSelectedCollections(const TArray<FCollectionNameType>& CollectionsToSelect, const bool bEnsureVisible = true);
