Slate Dynamic Invalidation - Content Browser Source Tree View

diff --git a/Engine/Source/Developer/ToolWidgets/Public/SSearchToggleButton.h b/Engine/Source/Developer/ToolWidgets/Public/SSearchToggleButton.h
--- a/Engine/Source/Developer/ToolWidgets/Public/SSearchToggleButton.h
+++ b/Engine/Source/Developer/ToolWidgets/Public/SSearchToggleButton.h
@@ -18,6 +18,12 @@ class SSearchBox;
 
 DECLARE_DELEGATE(FOnSearchBoxShown)
 
+template<>
+struct TWidgetTypeTraits<class SSearchToggleButton>
+{
+	static constexpr bool SupportsInvalidation() { return true; }
+};
+
 class SSearchToggleButton : public SCompoundWidget
 {
 public:
diff --git a/Engine/Source/Editor/AssetTagsEditor/Public/SAssetTagItem.h b/Engine/Source/Editor/AssetTagsEditor/Public/SAssetTagItem.h
--- a/Engine/Source/Editor/AssetTagsEditor/Public/SAssetTagItem.h
+++ b/Engine/Source/Editor/AssetTagsEditor/Public/SAssetTagItem.h
@@ -28,6 +28,16 @@
 class SWidget;
 struct FSlateBrush;
 
+
+template <typename ItemType>
+class SAssetTagItemTableRow;
+
+template <typename ItemType>
+struct TWidgetTypeTraits<class SAssetTagItemTableRow<ItemType>>
+{
+	static constexpr bool SupportsInvalidation() { return true; }
+};
+
 /** Custom table row used to contain asset tag items within list and tree views */
 template <typename ItemType>
 class SAssetTagItemTableRow : public STableRow<ItemType>
@@ -98,6 +108,12 @@ private:
 	TAttribute<bool> IsDropTarget;
 };
 
+template <>
+struct TWidgetTypeTraits<class SAssetTagItem>
+{
+	static constexpr bool SupportsInvalidation() { return true; }
+};
+
 /** A single asset tag item in the tags and collections tree. */
 class SAssetTagItem : public SCompoundWidget
 {
@@ -168,17 +184,17 @@ public:
 private:
 	FLinearColor GetAssetTagBrightColor() const
 	{
-		return BaseColor.Get().CopyWithNewOpacity(0.3f);
+		return BaseColor.Get(FLinearColor::White).CopyWithNewOpacity(0.3f);
 	}
 
 	FLinearColor GetAssetTagDullColor() const
 	{
-		return BaseColor.Get().CopyWithNewOpacity(0.1f);
+		return BaseColor.Get(FLinearColor::White).CopyWithNewOpacity(0.1f);
 	}
 
 	FLinearColor GetAssetTagDisabledColor() const
 	{
-		return BaseColor.Get().CopyWithNewOpacity(0.04f);
+		return BaseColor.Get(FLinearColor::White).CopyWithNewOpacity(0.04f);
 	}
 
 	FText GetCheckBoxTooltipText() const
diff --git a/Engine/Source/Editor/ContentBrowser/Private/SCollectionView.h b/Engine/Source/Editor/ContentBrowser/Private/SCollectionView.h
--- a/Engine/Source/Editor/ContentBrowser/Private/SCollectionView.h
+++ b/Engine/Source/Editor/ContentBrowser/Private/SCollectionView.h
@@ -48,6 +48,12 @@ struct FKeyEvent;
 struct FPointerEvent;
 struct FSlateBrush;
 
+template<>
+struct TWidgetTypeTraits<class SCollectionView>
+{
+	static constexpr bool SupportsInvalidation() { return true; }
+};
+
 /**
  * The list view of collections.
  */
diff --git a/Engine/Source/Editor/ContentBrowser/Private/SPathView.cpp b/Engine/Source/Editor/ContentBrowser/Private/SPathView.cpp
--- a/Engine/Source/Editor/ContentBrowser/Private/SPathView.cpp
+++ b/Engine/Source/Editor/ContentBrowser/Private/SPathView.cpp
@@ -297,7 +297,8 @@ struct FPathViewData
 {
 public:
 	FPathViewData(FName InContentBrowserName, bool InFlat)
-		: OwningContentBrowserName(InContentBrowserName)
+		: Version(0)
+		, OwningContentBrowserName(InContentBrowserName)
 		, bFlat(InFlat)
 		, FolderPathTextFilter(decltype(FolderPathTextFilter)::FItemToStringArray::CreateStatic(
 			  [](FStringView Input, TArray<FString>& Out) { Out.Emplace(FString(Input)); }))
@@ -896,6 +897,49 @@ void FPathViewData::RemoveFolderItem(const TSharedRef<FTreeItem>& TreeItem)
 	Version++;
 }
 
+
+SLATE_IMPLEMENT_WIDGET(SPathView)
+void SPathView::PrivateRegisterAttributes(FSlateAttributeInitializer& AttributeInitializer)
+{
+	struct FInvalidation
+	{
+		static void HandleShowRedirectorsChanged(SWidget& OwningWidget)
+		{
+			SPathView& PathView = static_cast<SPathView&>(OwningWidget);
+			UE_LOG(LogPathView, Verbose, TEXT("PathView bShowRedirectors changed to %d"), PathView.bShowRedirectorsAttribute.Get());
+			PathView.HandleSettingChanged("ShowRedirectors");
+		}
+		static void RequestTreeRefresh(SWidget& OwningWidget)
+		{
+			SPathView& PathView = static_cast<SPathView&>(OwningWidget);
+			if (ensure(PathView.TreeViewPtr))
+			{
+				PathView.TreeViewPtr->RequestTreeRefresh();
+			}
+		}
+	};
+	SLATE_ADD_MEMBER_ATTRIBUTE_DEFINITION(AttributeInitializer, bShowRedirectorsAttribute, EInvalidateWidgetReason::Paint)
+		.OnValueChanged(FSlateAttributeDescriptor::FAttributeValueChangedDelegate::CreateStatic(FInvalidation::HandleShowRedirectorsChanged));
+	SLATE_ADD_MEMBER_ATTRIBUTE_DEFINITION(AttributeInitializer, TreeDataVersionAttribute, EInvalidateWidgetReason::Paint)
+		.OnValueChanged(FSlateAttributeDescriptor::FAttributeValueChangedDelegate::CreateStatic(FInvalidation::RequestTreeRefresh));
+
+	SLATE_ADD_MEMBER_ATTRIBUTE_DEFINITION(AttributeInitializer, bPopulationRequestedAttribute, [](const SWidget& OwningWidget)
+	{
+		// Invalidate Paint only if value changed to true as this is where we process the Populate
+		return static_cast<const SPathView&>(OwningWidget).bPopulationRequestedAttribute.Get() ? EInvalidateWidgetReason::Paint : EInvalidateWidgetReason::None;
+	});
+	SLATE_ADD_MEMBER_ATTRIBUTE_DEFINITION(AttributeInitializer, bInitialPopulationRequestedAttribute, [](const SWidget& OwningWidget)
+	{
+		// Invalidate Paint only if value changed to true as this is where we process the Populate
+		return static_cast<const SPathView&>(OwningWidget).bInitialPopulationRequestedAttribute.Get() ? EInvalidateWidgetReason::Paint : EInvalidateWidgetReason::None;
+	});
+	SLATE_ADD_MEMBER_ATTRIBUTE_DEFINITION(AttributeInitializer, bPopulationForSettingChangeRequestedAttribute, [](const SWidget& OwningWidget)
+	{
+		// Invalidate Paint only if value changed to true as this is where we process the Populate
+		return static_cast<const SPathView&>(OwningWidget).bPopulationForSettingChangeRequestedAttribute.Get() ? EInvalidateWidgetReason::Paint : EInvalidateWidgetReason::None;
+	});
+}
+
 SPathView::FScopedSelectionChangedEvent::FScopedSelectionChangedEvent(const TSharedRef<SPathView>& InPathView, const bool InShouldEmitEvent)
 	: PathView(InPathView)
 	, bShouldEmitEvent(InShouldEmitEvent)
@@ -936,6 +980,15 @@ TSet<FName> SPathView::FScopedSelectionChangedEvent::GetSelectionSet() const
 	return SelectionSet;
 }
 
+SPathView::SPathView()
+	: bShowRedirectorsAttribute(*this, false)
+	, TreeDataVersionAttribute(*this, 0)
+	, bPopulationRequestedAttribute(*this, false)
+	, bInitialPopulationRequestedAttribute(*this, false)
+	, bPopulationForSettingChangeRequestedAttribute(*this, false)
+{
+}
+
 SPathView::~SPathView()
 {
 	UE::PathView::AllPathViews.RemoveAllSwap([this](const TWeakPtr<SPathView> Weak) { return Weak.Pin().Get() == this; });
@@ -955,6 +1008,8 @@ SPathView::~SPathView()
 
 void SPathView::Construct( const FArguments& InArgs )
 {
+	SetCanTick(false);
+
 	UE::PathView::AllPathViews.Add(SharedThis(this));
 
 	OwningContentBrowserName = InArgs._OwningContentBrowserName;
@@ -964,11 +1019,9 @@ void SPathView::Construct( const FArguments& InArgs )
 	InitialCategoryFilter = InArgs._InitialCategoryFilter;
 	bAllowClassesFolder = InArgs._AllowClassesFolder;
 	bAllowReadOnlyFolders = InArgs._AllowReadOnlyFolders;
-	bShowRedirectors = InArgs._ShowRedirectors;
 	bCanShowDevelopersFolder = InArgs._CanShowDevelopersFolder;
 	bForceShowEngineContent = InArgs._ForceShowEngineContent;
 	bForceShowPluginContent = InArgs._ForceShowPluginContent;
-	bLastShowRedirectors = bShowRedirectors.Get(false);
 	PreventTreeItemChangedDelegateCount = 0;
 	TreeTitle = LOCTEXT("AssetTreeTitle", "Asset Tree");
 	if ( InArgs._FocusSearchBoxWhenOpened )
@@ -1162,8 +1215,11 @@ void SPathView::Construct( const FArguments& InArgs )
 
 	CustomFolderPermissionList = InArgs._CustomFolderPermissionList;
 
+	bShowRedirectorsAttribute.Assign(*this, InArgs._ShowRedirectors);
+	TreeDataVersionAttribute.Assign(*this, TAttribute<uint64>::CreateLambda([this](){return TreeData->GetVersion();}));
+
 	// Add all paths currently gathered from the asset registry
-	bInitialPopulationRequested = true;
+	bInitialPopulationRequestedAttribute.Assign(*this, true);
 
 	DefaultPath = InArgs._DefaultPath;
 	bCreateDefaultPath = InArgs._CreateDefaultPath;
@@ -1179,33 +1235,18 @@ void SPathView::ConfigureTreeView(STreeView<TSharedPtr<FTreeItem>>::FArguments&
 		.HighlightParentNodesForSelection(true);
 }
 
-void SPathView::Tick(const FGeometry& AllottedGeometry, const double InCurrentTime, const float InDeltaTime)
+int32 SPathView::OnPaint(const FPaintArgs& Args, const FGeometry& AllottedGeometry, const FSlateRect& MyCullingRect, FSlateWindowElementList& OutDrawElements, int32 LayerId, const FWidgetStyle& InWidgetStyle, bool bParentEnabled) const
 {
-	Super::Tick(AllottedGeometry, InCurrentTime, InDeltaTime);
+	SPathView* MutableThis = const_cast<SPathView*>(this);
 
 	if (IsPopulationRequested())
 	{
-		Populate();
-	}
-
-	if (LastTreeDataVersion != TreeData->GetVersion())
-	{
-		LastTreeDataVersion = TreeData->GetVersion();
-		TreeViewPtr->RequestTreeRefresh();
+		MutableThis->Populate();
 	}
 
-	const bool bNewShowRedirectors = bShowRedirectors.Get(false);
-	if (bNewShowRedirectors != bLastShowRedirectors)
-	{
-		UE_LOG(LogPathView, Verbose, TEXT("PathView bShowRedirectors changed to %d"), bNewShowRedirectors);
-		bLastShowRedirectors = bNewShowRedirectors;
-		HandleSettingChanged("ShowRedirectors");
-	}
+	MutableThis->UpdateLastExpandedPathsIfDirty();
 
-	if (bLastExpandedPathsDirty)
-	{
-		UpdateLastExpandedPathsIfDirty();
-	}
+	return SCompoundWidget::OnPaint(Args, AllottedGeometry, MyCullingRect, OutDrawElements, LayerId, InWidgetStyle, bParentEnabled);
 }
 
 FReply SPathView::OnKeyDown(const FGeometry& MyGeometry, const FKeyEvent& InKeyEvent)
@@ -1260,7 +1301,7 @@ void SPathView::PluginPathFilterClicked(TSharedRef<FContentBrowserPluginFilter>
 {
 	SetPluginPathFilterActive(Filter, !IsPluginPathFilterInUse(Filter));
 	
-	bPopulationRequested = true;
+	bPopulationRequestedAttribute.Assign(*this, true);
 
 	// Temporary turn off the delayed population to investigate navigation issue when the view is not populated
 	Populate();
@@ -1298,7 +1339,7 @@ void SPathView::ResetPluginPathFilters()
 		SetPluginPathFilterActive(Filter, false);
 	}
 
-	bPopulationRequested = true;
+	bPopulationRequestedAttribute.Assign(*this, true);
 	// Temporary turn off the delayed population to investigate navigation issue when the view is not populated
 	Populate();
 }
@@ -1618,7 +1659,7 @@ EContentBrowserItemCategoryFilter SPathView::GetContentBrowserItemCategoryFilter
 		ItemCategoryFilter &= ~EContentBrowserItemCategoryFilter::IncludeClasses;
 	}
 	ItemCategoryFilter &= ~EContentBrowserItemCategoryFilter::IncludeCollections;
-	if (bShowRedirectors.Get(false))
+	if (bShowRedirectorsAttribute.Get())
 	{
 		ItemCategoryFilter |= EContentBrowserItemCategoryFilter::IncludeRedirectors;
 	}
@@ -2237,7 +2278,7 @@ FText SPathView::GetHighlightText() const
 
 void SPathView::RequestPopulation(FSimpleMulticastDelegate::FDelegate&& AfterPopulation)
 {
-	bPopulationRequested = true;
+	bPopulationRequestedAttribute.Assign(*this, true);
 
 	if (AfterPopulation.IsBound())
 	{
@@ -2259,17 +2300,17 @@ void SPathView::Populate(const bool bIsRefreshingFilter)
 	// Batch the selection changed event
 	// Only emit events when the user isn't filtering or doing the first population.
 	FScopedSelectionChangedEvent ScopedSelectionChangedEvent(SharedThis(this),
-		!bFilteringByText && !bIsRefreshingFilter && !bInitialPopulationRequested);
+		!bFilteringByText && !bIsRefreshingFilter && !bInitialPopulationRequestedAttribute.Get());
 	TreeData->PopulateFullFolderTree(CreateCompiledFolderFilter());
 	TreeData->FilterFullFolderTree();
 	TreeData->SortRootItems();
 
-	if (bPopulationForSettingChangeRequested)
+	if (bPopulationForSettingChangeRequestedAttribute.Get())
 	{
 		PostPopulateHandleSettingChanged(!ScopedSelectionChangedEvent.IsInitialSelectionEmpty());
 	}
 
-	if (bInitialPopulationRequested)
+	if (bInitialPopulationRequestedAttribute.Get())
 	{
 		PostInitialPopulation();
 	}
@@ -2291,9 +2332,9 @@ void SPathView::SharedPostPopulate()
 	OnNextPopulate.Broadcast();
 	OnNextPopulate.Clear();
 
-	bPopulationRequested = false;
-	bInitialPopulationRequested = false;
-	bPopulationForSettingChangeRequested = false;
+	bPopulationRequestedAttribute.Assign(*this, false);
+	bInitialPopulationRequestedAttribute.Assign(*this, false);
+	bPopulationForSettingChangeRequestedAttribute.Assign(*this, false);
 }
 
 
@@ -2479,7 +2520,7 @@ void SPathView::HandleItemDataUpdated(TArrayView<const FContentBrowserItemDataUp
 
 void SPathView::HandleItemDataRefreshed()
 {
-	bPopulationRequested = true;
+	bPopulationRequestedAttribute.Assign(*this, true);
 
 	// Temporary turn off the delayed population to investigate navigation issue when the view is not populated
 	Populate();
@@ -2512,7 +2553,7 @@ void SPathView::HandleSettingChanged(FName PropertyName)
 		const bool bHadSelectedPath = TreeViewPtr->GetNumItemsSelected() > 0;
 
 		// Update our path view so that it can include/exclude the permitted folder
-		bPopulationForSettingChangeRequested = true;
+		bPopulationForSettingChangeRequestedAttribute.Assign(*this, true);
 
 		// Temporary turn off the delayed population to investigate navigation issue when the view is not populated
 		Populate();
@@ -2609,7 +2650,7 @@ void SPathView::PostInitialPopulation()
 
 bool SPathView::IsPopulationRequested() const
 {
-	return bPopulationRequested || bPopulationForSettingChangeRequested || bInitialPopulationRequested;
+	return bPopulationRequestedAttribute.Get() || bPopulationForSettingChangeRequestedAttribute.Get() || bInitialPopulationRequestedAttribute.Get();
 }
 
 TArray<FName> SPathView::GetDefaultPathsToSelect() const
@@ -2654,6 +2695,7 @@ TArray<FName> SPathView::GetDefaultPathsToExpand() const
 void SPathView::DirtyLastExpandedPaths()
 {
 	bLastExpandedPathsDirty = true;
+	Invalidate(EInvalidateWidgetReason::Paint);
 }
 
 void SPathView::UpdateLastExpandedPathsIfDirty()
diff --git a/Engine/Source/Editor/ContentBrowser/Private/SPathView.h b/Engine/Source/Editor/ContentBrowser/Private/SPathView.h
--- a/Engine/Source/Editor/ContentBrowser/Private/SPathView.h
+++ b/Engine/Source/Editor/ContentBrowser/Private/SPathView.h
@@ -58,6 +58,12 @@ struct FPathViewData;
 struct FPointerEvent;
 struct FContentBrowserInstanceConfig;
 
+template<>
+struct TWidgetTypeTraits<class SPathView>
+{
+	static constexpr bool SupportsInvalidation() { return true; }
+};
+
 /**
  * The tree view of folders which contain content.
  */
@@ -65,7 +71,7 @@ class SPathView
 	: public SCompoundWidget
 	, public IScrollableWidget
 {
-	using Super = SCompoundWidget;
+	SLATE_DECLARE_WIDGET(SPathView, SCompoundWidget)
 
 public:
 	/** Delegate for when plugin filters have changed */
@@ -169,14 +175,16 @@ public:
 
 	SLATE_END_ARGS()
 
+	SPathView();
+
 	/** Destructor */
 	~SPathView();
 
 	/** Constructs this widget with InArgs */
 	virtual void Construct( const FArguments& InArgs );
 
-	/** Tick to poll attributes */
-	virtual void Tick(const FGeometry& AllottedGeometry, const double InCurrentTime, const float InDeltaTime) override;
+	/** Paint to poll attributes. This avoids the need for the widget to tick each frame */
+	virtual int32 OnPaint(const FPaintArgs& Args, const FGeometry& AllottedGeometry, const FSlateRect& MyCullingRect, FSlateWindowElementList& OutDrawElements, int32 LayerId, const FWidgetStyle& InWidgetStyle, bool bParentEnabled) const override;
 
 	/** Process the Commands of the PathView */
 	virtual FReply OnKeyDown( const FGeometry& MyGeometry, const FKeyEvent& InKeyEvent ) override;
@@ -258,7 +266,7 @@ public:
 	/** Sets an alternate tree title*/
 	void SetTreeTitle(FText InTitle)
 	{
-		TreeTitle = InTitle;
+		TreeTitle = MoveTemp(InTitle);
 	};
 
 	FText GetTreeTitle() const
@@ -488,9 +496,6 @@ protected:
 	/** Items in the tree and associated data. Shared ptr for easy binding to delegates/attributes */
 	TSharedPtr<FPathViewData> TreeData;
 
-	// Last version number retrieved from TreeData so we can decide if the tree view need rebuilding
-	uint64 LastTreeDataVersion = 0;
-
 	/** Should this path tree be flat like the favorites tree */
 	bool bFlat = false;
 
@@ -517,9 +522,6 @@ protected:
 	/** Custom Folder permissions */
 	TSharedPtr<FPathPermissionList> CustomFolderPermissionList;
 
-	TAttribute<bool> bShowRedirectors;
-	bool bLastShowRedirectors = false;
-
 	/** The config instance to use. */
 	FName OwningContentBrowserName;
 
@@ -527,6 +529,10 @@ protected:
 	FSimpleMulticastDelegate OnNextPopulate;
 
 private:
+	TSlateAttribute<bool> bShowRedirectorsAttribute;
+	/** Last version number retrieved from TreeData so we can decide if the tree view need rebuilding */
+	TSlateAttribute<uint64> TreeDataVersionAttribute;
+
 	/** Used to track if the list of last expanded path should be updated */
 	bool bLastExpandedPathsDirty = false;
 
@@ -564,13 +570,13 @@ private:
 	bool bForceShowPluginContent : 1;
 	
 	/** Request a population of the path view */
-	bool bPopulationRequested : 1;
+	TSlateAttribute<bool> bPopulationRequestedAttribute;
 
 	/** Request the initial population of the path view */
-	bool bInitialPopulationRequested : 1;
+	TSlateAttribute<bool> bInitialPopulationRequestedAttribute;
 
 	/** Request a population of the path view but because of a setting change */
-	bool bPopulationForSettingChangeRequested : 1;
+	TSlateAttribute<bool> bPopulationForSettingChangeRequestedAttribute;
 
 	bool bCreateDefaultPath : 1;
 	FString DefaultPath;
@@ -592,6 +598,11 @@ private:
 };
 
 
+template<>
+struct TWidgetTypeTraits<class SFavoritePathView>
+{
+	static constexpr bool SupportsInvalidation() { return true; }
+};
 
 /**
 * The tree view of folders which contain favorited folders.
diff --git a/Engine/Source/Editor/ContentBrowser/Private/SourcesViewWidgets.cpp b/Engine/Source/Editor/ContentBrowser/Private/SourcesViewWidgets.cpp
--- a/Engine/Source/Editor/ContentBrowser/Private/SourcesViewWidgets.cpp
+++ b/Engine/Source/Editor/ContentBrowser/Private/SourcesViewWidgets.cpp
@@ -176,7 +176,10 @@ void SAssetTreeItem::Construct( const FArguments& InArgs )
 
 	if( InlineRenameWidget.IsValid() )
 	{
-		EnterEditingModeDelegateHandle = TreeItem.Pin()->OnRenameRequested().AddSP( InlineRenameWidget.Get(), &SInlineEditableTextBlock::EnterEditingMode );
+		if ( const TSharedPtr<FTreeItem> TreeItemPtr = TreeItem.Pin() )
+		{
+			EnterEditingModeDelegateHandle = TreeItemPtr->OnRenameRequested().AddSP( InlineRenameWidget.Get(), &SInlineEditableTextBlock::EnterEditingMode );
+		}
 	}
 }
 
@@ -224,7 +227,10 @@ SAssetTreeItem::~SAssetTreeItem()
 {
 	if( InlineRenameWidget.IsValid() )
 	{
-		TreeItem.Pin()->OnRenameRequested().Remove( EnterEditingModeDelegateHandle );
+		if ( const TSharedPtr<FTreeItem> TreeItemPtr = TreeItem.Pin() )
+		{
+			TreeItemPtr->OnRenameRequested().Remove( EnterEditingModeDelegateHandle );
+		}
 	}
 }
 
@@ -271,9 +277,8 @@ FReply SAssetTreeItem::OnDrop( const FGeometry& MyGeometry, const FDragDropEvent
 
 bool SAssetTreeItem::VerifyNameChanged(const FText& InName, FText& OutError) const
 {
-	if ( TreeItem.IsValid() )
+	if ( TSharedPtr<FTreeItem> TreeItemPtr = TreeItem.Pin() )
 	{
-		TSharedPtr<FTreeItem> TreeItemPtr = TreeItem.Pin();
 		if(OnVerifyNameChanged.IsBound())
 		{
 			return OnVerifyNameChanged.Execute(TreeItemPtr, InName.ToString(), OutError);
@@ -285,10 +290,8 @@ bool SAssetTreeItem::VerifyNameChanged(const FText& InName, FText& OutError) con
 
 void SAssetTreeItem::HandleNameCommitted( const FText& NewText, ETextCommit::Type CommitInfo )
 {
-	if ( TreeItem.IsValid() )
+	if ( const TSharedPtr<FTreeItem> TreeItemPtr = TreeItem.Pin() )
 	{
-		TSharedPtr<FTreeItem> TreeItemPtr = TreeItem.Pin();
-
 		if ( TreeItemPtr->IsNamingFolder() )
 		{
 			TreeItemPtr->SetNamingFolder(false);
@@ -305,9 +308,9 @@ void SAssetTreeItem::HandleNameCommitted( const FText& NewText, ETextCommit::Typ
 
 bool SAssetTreeItem::IsReadOnly() const
 {
-	if ( TreeItem.IsValid() )
+	if ( const TSharedPtr<FTreeItem> TreeItemPtr = TreeItem.Pin() )
 	{
-		return !TreeItem.Pin()->IsNamingFolder();
+		return !TreeItemPtr->IsNamingFolder();
 	}
 	else
 	{
@@ -317,10 +320,11 @@ bool SAssetTreeItem::IsReadOnly() const
 
 const FSlateBrush* SAssetTreeItem::GetFolderIcon() const
 {
+	const bool bIsItemExpanded = IsItemExpanded.Get(false);
 	// First see if this item has a custom implementation
 	if (TSharedPtr<FTreeItem> TreeItemPin = TreeItem.Pin())
 	{
-		FContentBrowserItemDataAttributeValue FolderAttributeValue = TreeItemPin->GetItem().GetItemAttribute(IsItemExpanded.Get() ? ContentBrowserItemAttributes::ItemAssetTreeFolderOpenBrush : ContentBrowserItemAttributes::ItemAssetTreeFolderClosedBrush);
+		FContentBrowserItemDataAttributeValue FolderAttributeValue = TreeItemPin->GetItem().GetItemAttribute(bIsItemExpanded ? ContentBrowserItemAttributes::ItemAssetTreeFolderOpenBrush : ContentBrowserItemAttributes::ItemAssetTreeFolderClosedBrush);
 		if (FolderAttributeValue.IsValid())
 		{
 			FName BrushName = FolderAttributeValue.GetValue<FName>();
@@ -335,18 +339,18 @@ const FSlateBrush* SAssetTreeItem::GetFolderIcon() const
 	switch( FolderType )
 	{
 	case EFolderType::Code:
-		return IsItemExpanded.Get() ? Brushes.FolderOpenCodeBrush : Brushes.FolderClosedCodeBrush;
+		return bIsItemExpanded ? Brushes.FolderOpenCodeBrush : Brushes.FolderClosedCodeBrush;
 
 	case EFolderType::Developer:
-		return IsItemExpanded.Get() ? Brushes.FolderOpenDeveloperBrush : Brushes.FolderClosedDeveloperBrush;
+		return bIsItemExpanded ? Brushes.FolderOpenDeveloperBrush : Brushes.FolderClosedDeveloperBrush;
 
 	case EFolderType::CustomVirtual:
-		return IsItemExpanded.Get() ? Brushes.FolderOpenVirtualBrush : Brushes.FolderClosedVirtualBrush;
+		return bIsItemExpanded ? Brushes.FolderOpenVirtualBrush : Brushes.FolderClosedVirtualBrush;
 	case EFolderType::PluginRoot:
-		return (IsItemExpanded.Get()) ? Brushes.FolderOpenPluginRootBrush : Brushes.FolderClosedPluginRootBrush;
+		return bIsItemExpanded ? Brushes.FolderOpenPluginRootBrush : Brushes.FolderClosedPluginRootBrush;
 
 	default:
-		return IsItemExpanded.Get() ? Brushes.FolderOpenBrush : Brushes.FolderClosedBrush;
+		return bIsItemExpanded ? Brushes.FolderOpenBrush : Brushes.FolderClosedBrush;
 	}
 }
 
diff --git a/Engine/Source/Editor/ContentBrowser/Private/SourcesViewWidgets.h b/Engine/Source/Editor/ContentBrowser/Private/SourcesViewWidgets.h
--- a/Engine/Source/Editor/ContentBrowser/Private/SourcesViewWidgets.h
+++ b/Engine/Source/Editor/ContentBrowser/Private/SourcesViewWidgets.h
@@ -32,6 +32,12 @@ struct FCollectionItem;
 struct FSlateBrush;
 struct FSlateFontInfo;
 
+template <>
+struct TWidgetTypeTraits<class SAssetTreeItem>
+{
+	static constexpr bool SupportsInvalidation() { return true; }
+};
+
 /** A single item in the asset tree. Represents a folder. */
 class SAssetTreeItem : public SCompoundWidget
 {
@@ -135,6 +141,13 @@ private:
 	FDelegateHandle EnterEditingModeDelegateHandle;
 };
 
+
+template <>
+struct TWidgetTypeTraits<class SCollectionTreeItem>
+{
+	static constexpr bool SupportsInvalidation() { return true; }
+};
+
 /** A single item in the collection tree. */
 class SCollectionTreeItem : public SCompoundWidget
 {
diff --git a/Engine/Source/Editor/ContentBrowser/Private/Widgets/SContentBrowserSourceTree.cpp b/Engine/Source/Editor/ContentBrowser/Private/Widgets/SContentBrowserSourceTree.cpp
--- a/Engine/Source/Editor/ContentBrowser/Private/Widgets/SContentBrowserSourceTree.cpp
+++ b/Engine/Source/Editor/ContentBrowser/Private/Widgets/SContentBrowserSourceTree.cpp
@@ -14,6 +14,7 @@
 #include "Widgets/Layout/SScrollBorder.h"
 #include "Widgets/Layout/SWidgetSwitcher.h"
 #include "Widgets/SBoxPanel.h"
+#include "Widgets/SInvalidationPanel.h"
 #include "Widgets/SWidget.h"
 #include "Widgets/Text/SRichTextBlock.h"
 #include "Widgets/Text/STextBlock.h"
@@ -417,12 +418,16 @@ namespace UE::Editor::ContentBrowser::Private
 					+ SVerticalBox::Slot()
 					.Padding(FMargin(0, 1))
 					[
-						// Surround scrollable with a scrollbox (adds drop shadows)
-						SNew(SScrollBorder, InBody)
+						SNew(SInvalidationPanel)
+						.UseDynamicInvalidation(true)
 						[
-							BodyScrollableWidget->GetScrollWidget()
+							// Surround scrollable with a scrollbox (adds drop shadows)
+							SNew(SScrollBorder, InBody)
+							[
+								BodyScrollableWidget->GetScrollWidget()
+							]
+							.Style(InArgs._bEnableShadowBoxStyle ? InArgs._ShadowBoxStyle : &FCoreStyle::Get().GetWidgetStyle<FScrollBorderStyle>("ScrollBorder"))
 						]
-						.Style(InArgs._bEnableShadowBoxStyle ? InArgs._ShadowBoxStyle : &FCoreStyle::Get().GetWidgetStyle<FScrollBorderStyle>("ScrollBorder"))
 					]
 				]
 
diff --git a/Engine/Source/Editor/ContentBrowser/Private/Widgets/SContentBrowserSourceTree.h b/Engine/Source/Editor/ContentBrowser/Private/Widgets/SContentBrowserSourceTree.h
--- a/Engine/Source/Editor/ContentBrowser/Private/Widgets/SContentBrowserSourceTree.h
+++ b/Engine/Source/Editor/ContentBrowser/Private/Widgets/SContentBrowserSourceTree.h
@@ -28,7 +28,23 @@ namespace UE::Editor::ContentBrowser::Private
 {
 	class SContentBrowserSourceTree;
 	class SContentBrowserSourceTreeArea;
+}
+
+template<>
+struct TWidgetTypeTraits<UE::Editor::ContentBrowser::Private::SContentBrowserSourceTree>
+{
+	static constexpr bool SupportsInvalidation() { return true; }
+};
+
+template<>
+struct TWidgetTypeTraits<UE::Editor::ContentBrowser::Private::SContentBrowserSourceTreeArea>
+{
+	static constexpr bool SupportsInvalidation() { return true; }
+};
+
 
+namespace UE::Editor::ContentBrowser::Private
+{
 	/** The Content Browser Source Tree, containing "Favorites", etc. */
 	class SContentBrowserSourceTree : public SCompoundWidget
 	{
diff --git a/Engine/Source/Runtime/Slate/Private/Widgets/Views/STableViewBase.cpp b/Engine/Source/Runtime/Slate/Private/Widgets/Views/STableViewBase.cpp
--- a/Engine/Source/Runtime/Slate/Private/Widgets/Views/STableViewBase.cpp
+++ b/Engine/Source/Runtime/Slate/Private/Widgets/Views/STableViewBase.cpp
@@ -367,14 +367,6 @@ EActiveTimerReturnType STableViewBase::UpdateInertialScroll(double InCurrentTime
 	return bKeepTicking ? EActiveTimerReturnType::Continue : EActiveTimerReturnType::Stop;
 }
 
-EActiveTimerReturnType STableViewBase::EnsureTickToRefresh(double InCurrentTime, float InDeltaTime)
-{
-	// Actual refresh isn't implemented here as it can be needed in response to changes in the panel geometry.
-	// Since that isn't known until Tick (called after this when registered), refreshing here could result in two refreshes in one frame.
-
-	return EActiveTimerReturnType::Stop;
-}
-
 void STableViewBase::Tick( const FGeometry& AllottedGeometry, const double InCurrentTime, const float InDeltaTime )
 {
 	TRACE_CPUPROFILER_EVENT_SCOPE(STableViewBase::Tick);
@@ -488,7 +480,6 @@ void STableViewBase::Tick( const FGeometry& AllottedGeometry, const double InCur
 			if (ScrollIntoViewResult == EScrollIntoViewResult::Deferred || CurrentScrollOffset != TargetScrollOffset)
 			{
 				// Either we haven't made the item yet or we still have scrolling to do, so we'll need another refresh next frame
-				// We call this rather than just leave bItemsNeedRefresh as true to ensure that EnsureTickToRefresh is registered
 				RequestLayoutRefresh();
 			}
 			else if (CurrentScrollOffset == TargetScrollOffset
@@ -1005,7 +996,7 @@ void STableViewBase::UpdateSmoothDeltaTime(float InDeltaTime)
 float STableViewBase::ScrollBy(const FGeometry& MyGeometry, float ScrollByAmountInSlateUnits, EAllowOverscroll InAllowOverscroll)
 {
 	const int32 NumItemsBeingObserved = GetNumItemsBeingObserved();
-	const float FractionalScrollOffsetInItems = (DesiredScrollOffset + GetScrollRateInItems() * ScrollByAmountInSlateUnits) / NumItemsBeingObserved;
+	const float FractionalScrollOffsetInItems = NumItemsBeingObserved == 0 ? 0.f : (DesiredScrollOffset + GetScrollRateInItems() * ScrollByAmountInSlateUnits) / NumItemsBeingObserved;
 	const double ClampedScrollOffsetInItems = FMath::Clamp<double>( FractionalScrollOffsetInItems*NumItemsBeingObserved, -10.0f, NumItemsBeingObserved+10.0f ) * NumItemsBeingObserved;
 	if (InAllowOverscroll == EAllowOverscroll::Yes)
 	{
@@ -1323,11 +1314,7 @@ float STableViewBase::GetScrollRateInItems() const
 
 void STableViewBase::RequestLayoutRefresh()
 {
-	if (!bItemsNeedRefresh)
-	{
-		bItemsNeedRefresh = true;
-		RegisterActiveTimer(0.f, FWidgetActiveTimerDelegate::CreateSP(this, &STableViewBase::EnsureTickToRefresh));
-	}
+	bItemsNeedRefresh = true;
 
 	if (ItemsPanel.IsValid())
 	{
diff --git a/Engine/Source/Runtime/Slate/Public/Widgets/Views/STableViewBase.h b/Engine/Source/Runtime/Slate/Public/Widgets/Views/STableViewBase.h
--- a/Engine/Source/Runtime/Slate/Public/Widgets/Views/STableViewBase.h
+++ b/Engine/Source/Runtime/Slate/Public/Widgets/Views/STableViewBase.h
@@ -625,9 +625,6 @@ protected:
 	/** Active timer to update the inertial scroll */
 	SLATE_API EActiveTimerReturnType UpdateInertialScroll(double InCurrentTime, float InDeltaTime);
 
-	/** One-off active timer to refresh the contents of the table as needed */
-	SLATE_API EActiveTimerReturnType EnsureTickToRefresh(double InCurrentTime, float InDeltaTime);
-
 	/** Whether the active timer to update the inertial scrolling is currently registered */
 	bool bIsScrollingActiveTimerRegistered;
 
