Fixed typo in AreBackendFiltersDifferent.

#rnx

diff --git a/Engine/Source/Editor/ContentBrowser/Private/SAssetView.cpp b/Engine/Source/Editor/ContentBrowser/Private/SAssetView.cpp
--- a/Engine/Source/Editor/ContentBrowser/Private/SAssetView.cpp
+++ b/Engine/Source/Editor/ContentBrowser/Private/SAssetView.cpp
@@ -136,7 +136,7 @@ namespace UE::AssetView
 	bool AreBackendFiltersDifferent(const FARFilter& A, const FARFilter& B)
 	{
 		if (A.PackageNames.Num() != B.PackageNames.Num()
-		|| A.PackagePaths.Num() != B.PackageNames.Num() 
+		|| A.PackagePaths.Num() != B.PackagePaths.Num()
 		|| A.SoftObjectPaths.Num() != B.SoftObjectPaths.Num()
 		|| A.ClassPaths.Num() != B.ClassPaths.Num()
 		|| A.TagsAndValues.Num() != B.TagsAndValues.Num()
