Fixed memory leaks in SAssetView and SCollectionView.
Updated a bunch of raw callbacks to use shared pointers.
Unregistered various collection event listeners to clean up the multicast delegate list and make it easier to spot unwanted callbacks.

#rnx

diff --git a/Engine/Source/Editor/ContentBrowser/Private/CollectionContextMenu.cpp b/Engine/Source/Editor/ContentBrowser/Private/CollectionContextMenu.cpp
--- a/Engine/Source/Editor/ContentBrowser/Private/CollectionContextMenu.cpp
+++ b/Engine/Source/Editor/ContentBrowser/Private/CollectionContextMenu.cpp
@@ -131,7 +131,7 @@ TSharedPtr<SWidget> FCollectionContextMenu::MakeCollectionTreeContextMenu(TShare
 			MenuBuilder.AddSubMenu(
 				LOCTEXT("NewChildCollection", "New..."),
 				LOCTEXT("NewChildCollectionTooltip", "Create a child collection."),
-				FNewMenuDelegate::CreateRaw( this, &FCollectionContextMenu::MakeNewCollectionSubMenu, ECollectionStorageMode::Static, SCollectionView::FCreateCollectionPayload( ParentCollection ) ),
+				FNewMenuDelegate::CreateSP( this, &FCollectionContextMenu::MakeNewCollectionSubMenu, ECollectionStorageMode::Static, SCollectionView::FCreateCollectionPayload( ParentCollection ) ),
 				FUIAction(
 					FExecuteAction(),
 					FCanExecuteAction::CreateLambda( [=]{ return bHasSingleSelectedCollection && !bIsFirstSelectedCollectionReadOnly && bIsFirstSelectedCollectionStatic; } )
@@ -148,7 +148,7 @@ TSharedPtr<SWidget> FCollectionContextMenu::MakeCollectionTreeContextMenu(TShare
 		MenuBuilder.AddSubMenu(
 			LOCTEXT("SetCollectionShareType", "Set Share Type"),
 			LOCTEXT("SetCollectionShareTypeTooltip", "Change the share type of this collection."),
-			FNewMenuDelegate::CreateRaw( this, &FCollectionContextMenu::MakeCollectionShareTypeSubMenu ),
+			FNewMenuDelegate::CreateSP( this, &FCollectionContextMenu::MakeCollectionShareTypeSubMenu ),
 			FUIAction(
 				FExecuteAction(),
 				FCanExecuteAction::CreateLambda( [=]{ return bHasSingleSelectedCollection && !bIsFirstSelectedCollectionReadOnly; } )
@@ -164,7 +164,7 @@ TSharedPtr<SWidget> FCollectionContextMenu::MakeCollectionTreeContextMenu(TShare
 			MenuBuilder.AddSubMenu(
 				LOCTEXT("SetColor", "Set Color"),
 				LOCTEXT("SetCollectionColorTooltip", "Sets the color this collection should appear as."),
-				FNewMenuDelegate::CreateRaw( this, &FCollectionContextMenu::MakeSetColorSubMenu )
+				FNewMenuDelegate::CreateSP( this, &FCollectionContextMenu::MakeSetColorSubMenu )
 				);
 		}
 		else
diff --git a/Engine/Source/Editor/ContentBrowser/Private/Menus/AssetViewMenus.cpp b/Engine/Source/Editor/ContentBrowser/Private/Menus/AssetViewMenus.cpp
--- a/Engine/Source/Editor/ContentBrowser/Private/Menus/AssetViewMenus.cpp
+++ b/Engine/Source/Editor/ContentBrowser/Private/Menus/AssetViewMenus.cpp
@@ -246,7 +246,7 @@ FDelayedAutoRegisterHelper SAssetView::AssetViewOptionsMenuRegistration(
 							LOCTEXT("ThumbnailsHeading", "Thumbnails"),
 							ThumbnailSectionInsert);
 
-						auto CreateThumbnailSizeSubMenu = [AssetView](UToolMenu* SubMenu)
+						auto CreateThumbnailSizeSubMenu = [AssetView = &AssetView.Get()](UToolMenu* SubMenu)
 						{
 							FToolMenuSection& SizeSection = SubMenu->FindOrAddSection("ThumbnailSizes");
 
@@ -279,7 +279,7 @@ FDelayedAutoRegisterHelper SAssetView::AssetViewOptionsMenuRegistration(
 							"ThumbnailSize",
 							FToolUIActionChoice(),
 							Private::GetThumbnailSizeWidget(),
-							FNewToolMenuDelegate::CreateLambda(CreateThumbnailSizeSubMenu)
+							FNewToolMenuDelegate::CreateSPLambda(AssetView, CreateThumbnailSizeSubMenu)
 						));
 
 						ThumbnailSection.AddMenuEntry(
@@ -590,17 +590,13 @@ FDelayedAutoRegisterHelper SAssetView::ToolBarMenuExtensionRegistration(
 							const TAttribute<FSlateIcon> GetIcon = TAttribute<FSlateIcon>::Create(
 								TAttribute<FSlateIcon>::FGetter::CreateSPLambda(
 									AssetView,
-									[WeakMenuContext = MakeWeakObjectPtr(Context)]()
+									[AssetView = &AssetView.Get()]()
 									{
 										EColumnSortMode::Type ActiveSortMode = EColumnSortMode::Ascending;
-										if (WeakMenuContext.IsValid() && WeakMenuContext->AssetView.IsValid())
+										if (const TSharedPtr<FAssetViewSortManager> StrongSortManager = AssetView->GetSortManager().Pin())
 										{
-											const TSharedPtr<SAssetView> AssetViewForIcon = WeakMenuContext->AssetView.Pin();
-											if (const TSharedPtr<FAssetViewSortManager> StrongSortManager = AssetViewForIcon->GetSortManager().Pin())
-											{
-												// @note: we don't use Secondary priorities, so this should change if/when we do
-												ActiveSortMode = StrongSortManager->GetSortMode(EColumnSortPriority::Primary);
-											}
+											// @note: we don't use Secondary priorities, so this should change if/when we do
+											ActiveSortMode = StrongSortManager->GetSortMode(EColumnSortPriority::Primary);
 										}
 
 										return ActiveSortMode == EColumnSortMode::Ascending
@@ -611,7 +607,7 @@ FDelayedAutoRegisterHelper SAssetView::ToolBarMenuExtensionRegistration(
 							FToolUIAction SortAction;
 							SortAction.IsActionVisibleDelegate = FToolMenuIsActionButtonVisible::CreateSPLambda(
 								AssetView,
-								[AssetView](const FToolMenuContext&)
+								[AssetView = &AssetView.Get()](const FToolMenuContext&)
 								{
 									// Sort button only visible in Tile view mode - others have their own sorting UI representation.
 									return AssetView->IsCurrentViewType(EAssetViewType::Tile);
diff --git a/Engine/Source/Editor/ContentBrowser/Private/PathContextMenu.cpp b/Engine/Source/Editor/ContentBrowser/Private/PathContextMenu.cpp
--- a/Engine/Source/Editor/ContentBrowser/Private/PathContextMenu.cpp
+++ b/Engine/Source/Editor/ContentBrowser/Private/PathContextMenu.cpp
@@ -141,7 +141,7 @@ void FPathContextMenu::MakePathViewContextMenu(UToolMenu* Menu)
 						"NewAsset",
 						LOCTEXT("AddImportLabel", "Add/Import Content"),
 						NewAssetToolTip,
-						FNewToolMenuDelegate::CreateRaw(this, &FPathContextMenu::MakeNewAssetSubMenu),
+						FNewToolMenuDelegate::CreateSP(this, &FPathContextMenu::MakeNewAssetSubMenu),
 						FUIAction(),
 						EUserInterfaceActionType::Button,
 						false,
@@ -191,7 +191,7 @@ void FPathContextMenu::MakePathViewContextMenu(UToolMenu* Menu)
 					"SetColor",
 					LOCTEXT("SetColor", "Set Color"),
 					LOCTEXT("SetColorTooltip", "Sets the color this folder should appear as."),
-					FNewToolMenuDelegate::CreateRaw( this, &FPathContextMenu::MakeSetColorSubMenu ),
+					FNewToolMenuDelegate::CreateSP( this, &FPathContextMenu::MakeSetColorSubMenu ),
 					false,
 					FSlateIcon(FAppStyle::GetAppStyleSetName(), "Icons.Color")
 					);
@@ -355,7 +355,10 @@ void FPathContextMenu::MakeSetColorSubMenu(UToolMenu* Menu)
 
 void FPathContextMenu::ExecuteExplore()
 {
-	ContentBrowserUtils::ExploreFolders(SelectedFolders, ParentContent.Pin().ToSharedRef());
+	if (TSharedPtr<SWidget> PinnedParentContent = ParentContent.Pin())
+	{
+		ContentBrowserUtils::ExploreFolders(SelectedFolders, PinnedParentContent.ToSharedRef());
+	}
 }
 
 bool FPathContextMenu::CanExecuteRename() const
diff --git a/Engine/Source/Editor/ContentBrowser/Private/SAssetView.cpp b/Engine/Source/Editor/ContentBrowser/Private/SAssetView.cpp
--- a/Engine/Source/Editor/ContentBrowser/Private/SAssetView.cpp
+++ b/Engine/Source/Editor/ContentBrowser/Private/SAssetView.cpp
@@ -1213,11 +1213,30 @@ SAssetView::~SAssetView()
 	// Remove the listener for when view settings are changed
 	UContentBrowserSettings::OnSettingChanged().RemoveAll(this);
 
+	TSet<TSharedPtr<ICollectionContainer>> UniqueCollectionContainers;
+	for (const FCollectionRef& Collection : ContentSources.GetCollections())
+	{
+		bool bIsAlreadyInSet = false;
+		UniqueCollectionContainers.Add(Collection.Container, &bIsAlreadyInSet);
+		if (!bIsAlreadyInSet)
+		{
+			Collection.Container->OnAssetsAddedToCollection().RemoveAll(this);
+			Collection.Container->OnAssetsRemovedFromCollection().RemoveAll(this);
+			Collection.Container->OnCollectionRenamed().RemoveAll(this);
+			Collection.Container->OnCollectionUpdated().RemoveAll(this);
+		}
+	}
+
 	if ( FrontendFilters.IsValid() )
 	{
 		// Clear the frontend filter changed delegate
 		FrontendFilters->OnChanged().RemoveAll( this );
 	}
+
+	if (TextFilter.IsValid())
+	{
+		TextFilter->OnChanged().RemoveAll(this);
+	}
 }
 
 
diff --git a/Engine/Source/Editor/ContentBrowser/Private/SCollectionView.cpp b/Engine/Source/Editor/ContentBrowser/Private/SCollectionView.cpp
--- a/Engine/Source/Editor/ContentBrowser/Private/SCollectionView.cpp
+++ b/Engine/Source/Editor/ContentBrowser/Private/SCollectionView.cpp
@@ -120,6 +120,33 @@ bool TestComplexExpression(const FCollectionItem& InCollection, const FName& InK
 
 } // namespace CollectionViewFilter
 
+SCollectionView::~SCollectionView()
+{
+	if (CollectionContainer)
+	{
+		CollectionContainer->OnCollectionCreated().RemoveAll(this);
+		CollectionContainer->OnCollectionRenamed().RemoveAll(this);
+		CollectionContainer->OnCollectionReparented().RemoveAll(this);
+		CollectionContainer->OnCollectionDestroyed().RemoveAll(this);
+		CollectionContainer->OnCollectionUpdated().RemoveAll(this);
+		CollectionContainer->OnAssetsAddedToCollection().RemoveAll(this);
+		CollectionContainer->OnAssetsRemovedFromCollection().RemoveAll(this);
+	}
+
+	if (ISourceControlModule* SourceControlModule = ISourceControlModule::GetPtr())
+	{
+		SourceControlModule->UnregisterProviderChanged(SourceControlProviderChangedHandle);
+		SourceControlModule->GetProvider().UnregisterSourceControlStateChanged_Handle(SourceControlStateChangedDelegateHandle);
+	}
+
+	if (ExternalSearchPtr)
+	{
+		ExternalSearchPtr->OnSearchChanged().RemoveAll(this);
+	}
+
+	UContentBrowserSettings::OnSettingChanged().RemoveAll(this);
+}
+
 void SCollectionView::Construct( const FArguments& InArgs )
 {
 	OnCollectionSelected = InArgs._OnCollectionSelected;
@@ -146,7 +173,7 @@ void SCollectionView::Construct( const FArguments& InArgs )
 		CollectionContainer->OnAssetsRemovedFromCollection().AddSP(this, &SCollectionView::HandleAssetsRemovedFromCollection);
 	}
 
-	ISourceControlModule::Get().RegisterProviderChanged(FSourceControlProviderChanged::FDelegate::CreateSP(this, &SCollectionView::HandleSourceControlProviderChanged));
+	SourceControlProviderChangedHandle = ISourceControlModule::Get().RegisterProviderChanged(FSourceControlProviderChanged::FDelegate::CreateSP(this, &SCollectionView::HandleSourceControlProviderChanged));
 	SourceControlStateChangedDelegateHandle = ISourceControlModule::Get().GetProvider().RegisterSourceControlStateChanged_Handle(FSourceControlStateChanged::FDelegate::CreateSP(this, &SCollectionView::HandleSourceControlStateChanged));
 
 	Commands = TSharedPtr< FUICommandList >(new FUICommandList);
@@ -170,13 +197,12 @@ void SCollectionView::Construct( const FArguments& InArgs )
 		CollectionListContextMenuOpening = FOnContextMenuOpening::CreateSP( this, &SCollectionView::MakeCollectionTreeContextMenu );
 	}
 
-	SearchPtr = InArgs._ExternalSearch;
-	if (SearchPtr)
+	ExternalSearchPtr = InArgs._ExternalSearch;
+	if (ExternalSearchPtr)
 	{
-		SearchPtr->OnSearchChanged().AddSP(this, &SCollectionView::SetCollectionsSearchFilterText);
+		ExternalSearchPtr->OnSearchChanged().AddSP(this, &SCollectionView::SetCollectionsSearchFilterText);
 	}
 
-	ExternalSearchPtr = InArgs._ExternalSearch;
 	TitleContent = SNew(SHorizontalBox);
 
 	UContentBrowserCollectionProjectSettings* CollectionProjectSettings = GetMutableDefault<UContentBrowserCollectionProjectSettings>();
@@ -266,6 +292,8 @@ void SCollectionView::Construct( const FArguments& InArgs )
 		]
 	];
 
+	UContentBrowserSettings::OnSettingChanged().AddSP(this, &SCollectionView::HandleSettingChanged);
+
 	UpdateCollectionItems();
 }
 
@@ -570,8 +598,6 @@ void SCollectionView::UpdateFilteredCollectionItems()
 			ExpandParentItems(CollectionItem);
 		}
 	}
-	
-	UContentBrowserSettings::OnSettingChanged().AddSP(this, &SCollectionView::HandleSettingChanged);
 
 	VisibleRootCollectionItems.Sort(FCollectionItem::FCompareFCollectionItemByName());
 	CollectionTreePtr->RequestTreeRefresh();
@@ -616,9 +642,9 @@ void SCollectionView::SetSelectedCollections(const TArray<FCollectionNameType>&
 			CollectionTreePtr->SetItemSelection(CollectionItemToSelect, true);
 
 			// If the selected collection doesn't pass our current filter, we need to clear it
-			if (bEnsureVisible && !CollectionItemTextFilter->PassesFilter(*CollectionItemToSelect))
+			if (bEnsureVisible && ExternalSearchPtr.IsValid() && !CollectionItemTextFilter->PassesFilter(*CollectionItemToSelect))
 			{
-				SearchPtr->ClearSearch();
+				ExternalSearchPtr->ClearSearch();
 			}
 		}
 	}
@@ -949,7 +975,10 @@ void SCollectionView::CreateCollectionItem( ECollectionShareType::Type Collectio
 		NewItem->StorageMode = StorageMode;
 
 		// Adding a new collection now, so clear any filter we may have applied
-		SearchPtr->ClearSearch();
+		if (ExternalSearchPtr)
+		{
+			ExternalSearchPtr->ClearSearch();
+		}
 
 		if ( InCreationPayload.ParentCollection.IsSet() )
 		{
@@ -1685,9 +1714,9 @@ bool SCollectionView::CollectionNameChangeCommit( const TSharedPtr< FCollectionI
 		}
 
 		// If the new name doesn't pass our current filter, we need to clear it
-		if ( !CollectionItemTextFilter->PassesFilter( FCollectionItem(CollectionItem->CollectionContainer.ToSharedRef(), NewNameFinal, CollectionItem->CollectionType) ) )
+		if (ExternalSearchPtr.IsValid() && !CollectionItemTextFilter->PassesFilter(FCollectionItem(CollectionItem->CollectionContainer.ToSharedRef(), NewNameFinal, CollectionItem->CollectionType)))
 		{
-			SearchPtr->ClearSearch();
+			ExternalSearchPtr->ClearSearch();
 		}
 
 		// Otherwise perform the rename
diff --git a/Engine/Source/Editor/ContentBrowser/Private/SCollectionView.h b/Engine/Source/Editor/ContentBrowser/Private/SCollectionView.h
--- a/Engine/Source/Editor/ContentBrowser/Private/SCollectionView.h
+++ b/Engine/Source/Editor/ContentBrowser/Private/SCollectionView.h
@@ -101,6 +101,8 @@ public:
 
 	SLATE_END_ARGS()
 
+	~SCollectionView();
+
 	/** Constructs this widget with InArgs */
 	void Construct( const FArguments& InArgs );
 
@@ -334,9 +336,6 @@ private:
 	/** The collection container to view */
 	TSharedPtr<ICollectionContainer> CollectionContainer;
 
-	/** The collection search interface */
-	TSharedPtr<FSourcesSearch> SearchPtr;
-
 	/** The collection external search interface (if any) */
 	TSharedPtr<FSourcesSearch> ExternalSearchPtr;
 
@@ -399,6 +398,9 @@ private:
 	 */
 	TWeakPtr<class FCollectionDragDropOp> CurrentCollectionDragDropOp;
 
+	/** Delegate handle for the HandleSourceControlProviderChanged function callback */
+	FDelegateHandle SourceControlProviderChangedHandle;
+
 	/** Delegate handle for the HandleSourceControlStateChanged function callback */
 	FDelegateHandle SourceControlStateChangedDelegateHandle;
 
diff --git a/Engine/Source/Editor/ContentBrowser/Private/SContentBrowser.cpp b/Engine/Source/Editor/ContentBrowser/Private/SContentBrowser.cpp
--- a/Engine/Source/Editor/ContentBrowser/Private/SContentBrowser.cpp
+++ b/Engine/Source/Editor/ContentBrowser/Private/SContentBrowser.cpp
@@ -197,17 +197,19 @@ class SContentBrowser::FCollectionSource
 {
 public:
 	FCollectionSource(SContentBrowser& InContentBrowser, const TSharedRef<ICollectionContainer>& InCollectionContainer)
-		: ContentBrowser(InContentBrowser)
 	{
 		CollectionSearch = MakeShared<FSourcesSearch>();
 		CollectionSearch->Initialize();
 		CollectionSearch->SetHintText(LOCTEXT("CollectionsViewSearchBoxHint", "Search Collections"));
 
 		CollectionViewPtr = SNew(SCollectionView)
-			.OnCollectionSelected_Lambda([this](const FCollectionNameType& SelectedCollection)
+			.OnCollectionSelected_Lambda([ContentBrowser = StaticCastWeakPtr<SContentBrowser>(InContentBrowser.AsWeak()), CollectionContainer = InCollectionContainer](const FCollectionNameType& SelectedCollection)
+			{
+				if (TSharedPtr<SContentBrowser> PinnedContentBrowser = ContentBrowser.Pin())
 				{
-					ContentBrowser.CollectionSelected(GetCollectionContainer(), SelectedCollection);
-				})
+					PinnedContentBrowser->CollectionSelected(CollectionContainer, SelectedCollection);
+				}
+			})
 			.AddMetaData<FTagMetaData>(FTagMetaData(TEXT("ContentBrowserCollections")))
 			.AllowCollectionDrag(true)
 			.AllowQuickAssetManagement(true)
@@ -248,22 +250,6 @@ public:
 		CollectionArea->SaveSettings(EditorPerProjectIni, SettingsIniSection, SettingsString);
 	}
 
-	SSplitter::ESizeRule GetCollectionsAreaSizeRule() const
-	{
-		// Make sure the area is expanded
-		return CollectionArea->IsExpanded() ? SSplitter::ESizeRule::FractionOfParent : SSplitter::ESizeRule::SizeToContent;
-	}
-
-	/** Handler for clicking the add collection button */
-	FReply OnAddCollectionClicked()
-	{
-		CollectionArea->SetExpanded(true);
-
-		CollectionViewPtr->MakeAddCollectionMenu(ContentBrowser.AsShared());
-
-		return FReply::Handled();
-	}
-
 private:
 	FString GetSettingsString(const FName& InInstanceName) const
 	{
@@ -279,8 +265,6 @@ private:
 	}
 
 public:
-	SContentBrowser& ContentBrowser;
-
 	/** The sources search for collections*/
 	TSharedPtr<FSourcesSearch> CollectionSearch;
 
@@ -305,6 +289,17 @@ SContentBrowser::~SContentBrowser()
 	{
 		FCollectionManagerModule& CollectionManagerModule = FCollectionManagerModule::GetModule();
 
+		TArray<TSharedPtr<ICollectionContainer>> CollectionContainers;
+		CollectionManagerModule.Get().GetCollectionContainers(CollectionContainers);
+
+		for (const TSharedPtr<ICollectionContainer>& CollectionContainer : CollectionContainers)
+		{
+			CollectionContainer->OnIsHiddenChanged().RemoveAll(this);
+			CollectionContainer->OnCollectionRenamed().RemoveAll(this);
+			CollectionContainer->OnCollectionDestroyed().RemoveAll(this);
+			CollectionContainer->OnCollectionUpdated().RemoveAll(this);
+		}
+
 		CollectionManagerModule.Get().OnCollectionContainerCreated().RemoveAll(this);
 		CollectionManagerModule.Get().OnCollectionContainerDestroyed().RemoveAll(this);
 	}
@@ -955,7 +950,7 @@ TSharedRef<SWidget> SContentBrowser::CreateAssetView(const FContentBrowserConfig
 					.CanChangeOrientation(true)
 					.OnFilterBarLayoutChanging(this, &SContentBrowser::OnFilterBarLayoutChanging)
 					.UseSharedSettings(true)
-					.CreateTextFilter(SFilterList::FCreateTextFilter::CreateLambda([this]
+					.CreateTextFilter(SFilterList::FCreateTextFilter::CreateSPLambda(this, [this]
 					{
 						TSharedPtr<FFrontendFilter_CustomText> NewFilter = MakeShared<FFrontendFilter_CustomText>();
 
@@ -1055,8 +1050,14 @@ SContentBrowser::FCollectionSource& SContentBrowser::AddSlotForCollectionContain
 	{
 		const float SourceTreeHeaderHeightMin = 26.0f + 3.0f;
 
+		const TSharedRef<UE::Editor::ContentBrowser::Private::SContentBrowserSourceTreeArea> CollectionArea = CreateCollectionsView(*CollectionSource);
+
 		LegacyContentSourceWidgets->PathFavoriteSplitterPtr->AddSlot(SlotIndex)
-			.SizeRule_Raw(CollectionSource.Get(), &FCollectionSource::GetCollectionsAreaSizeRule)
+			.SizeRule_Lambda([CollectionArea]()
+			{
+				// Make sure the area is expanded
+				return CollectionArea->IsExpanded() ? SSplitter::ESizeRule::FractionOfParent : SSplitter::ESizeRule::SizeToContent;
+			})
 			.MinSize(SourceTreeHeaderHeightMin)
 			.Value(0.4f)
 			[
@@ -1064,7 +1065,7 @@ SContentBrowser::FCollectionSource& SContentBrowser::AddSlotForCollectionContain
 					.BorderImage(FAppStyle::Get().GetBrush("Brushes.Header"))
 					.Padding(0.0f, 2.0f, 0.0f, 0.0f)
 					[
-						CreateCollectionsView(*CollectionSource)
+						CollectionArea
 					]
 			];
 
@@ -1169,7 +1170,7 @@ TSharedRef<SWidget> SContentBrowser::CreatePathView(const FContentBrowserConfig*
 		.FocusSearchBoxWhenOpened(false)
 		.ShowTreeTitle(false)
 		.ShowSeparator(false)
-		.ShowRedirectors_Lambda([this]() { return ContentBrowserUtils::ShouldShowRedirectors(LegacyContentSourceWidgets->FilterListPtr); })
+		.ShowRedirectors(MakeAttributeSPLambda(this, [this]() { return ContentBrowserUtils::ShouldShowRedirectors(LegacyContentSourceWidgets->FilterListPtr); }))
 		.AddMetaData<FTagMetaData>(FTagMetaData(TEXT("ContentBrowserSources")))
 		.ExternalSearch(SourcesSearch)
 		.PluginPathFilters(PluginPathFilters)
@@ -1217,7 +1218,20 @@ TSharedRef<UE::Editor::ContentBrowser::Private::SContentBrowserSourceTreeArea> S
 					SNew(SButton)
 						.ButtonStyle(FAppStyle::Get(), "SimpleButton")
 						.ToolTipText(LOCTEXT("AddCollectionButtonTooltip", "Add a collection."))
-						.OnClicked_Raw(&CollectionSource, &FCollectionSource::OnAddCollectionClicked)
+						.OnClicked(FOnClicked::CreateSPLambda(this, [this, CollectionContainer = CollectionSource.GetCollectionContainer()]()
+						{
+							for (const TUniquePtr<FCollectionSource>& CollectionSource : CollectionSources)
+							{
+								if (CollectionSource->GetCollectionContainer() == CollectionContainer)
+								{
+									CollectionSource->CollectionArea->SetExpanded(true);
+									CollectionSource->CollectionViewPtr->MakeAddCollectionMenu(SharedThis(this));
+									break;
+								}
+							}
+
+							return FReply::Handled();
+						}))
 						.ContentPadding(FMargin(1, 0))
 						.Visibility_Lambda([CollectionContainer = CollectionSource.GetCollectionContainer()]() { return !CollectionContainer->IsReadOnly(ECollectionShareType::CST_All) ? EVisibility::Visible : EVisibility::Collapsed; })
 						[
@@ -1227,9 +1241,9 @@ TSharedRef<UE::Editor::ContentBrowser::Private::SContentBrowserSourceTreeArea> S
 						]
 				]
 		]
-		.IsEmpty_Lambda([PathView = CollectionSource.CollectionViewPtr]
+		.IsEmpty_Lambda([CollectionView = CollectionSource.CollectionViewPtr]
 		{
-			return PathView->IsEmpty();
+			return CollectionView->IsEmpty();
 		})
 		.EmptyBodyLabel_Lambda([CollectionContainer = CollectionSource.GetCollectionContainer()]()
 		{
@@ -2432,10 +2446,7 @@ void SContentBrowser::OnSaveSearchButtonClicked(const FText& InSearchText)
 				MenuBuilder.AddSubMenu(
 					CollectionContainer->GetCollectionSource()->GetTitle(),
 					TAttribute<FText>(),
-					FNewMenuDelegate::CreateLambda([CollectionView = CollectionSource->CollectionViewPtr, FinalQueryText](FMenuBuilder& InSubMenuBuilder)
-						{
-							CollectionView->MakeSaveDynamicCollectionMenu(InSubMenuBuilder, FinalQueryText);
-						}));
+					FNewMenuDelegate::CreateSP(CollectionSource->CollectionViewPtr.ToSharedRef(), &SCollectionView::MakeSaveDynamicCollectionMenu, FinalQueryText));
 			}
 
 			MenuBuilder.EndSection();
@@ -3260,7 +3271,6 @@ void SContentBrowser::OnLegacyContentSourceDisabled()
 	TArray<TSharedPtr<ICollectionContainer>> CollectionContainers;
 	CollectionManagerModule.Get().GetCollectionContainers(CollectionContainers);
 
-	int32 InsertIndex = 0;
 	for (const TSharedPtr<ICollectionContainer>& CollectionContainer : CollectionContainers)
 	{
 		CollectionContainer->OnIsHiddenChanged().RemoveAll(this);
@@ -3280,7 +3290,7 @@ void SContentBrowser::OnLegacyContentSourceDisabled()
 
 TSharedRef<SWidget> SContentBrowser::CreateLegacyAssetViewWidgets()
 {
-	LegacyContentSourceWidgets = MakeShared<FLegacyContentSourceWidgets>();
+	LegacyContentSourceWidgets = MakeUnique<FLegacyContentSourceWidgets>();
 	
 	using namespace UE::Editor::ContentBrowser::Private;
 
@@ -4073,7 +4083,7 @@ void SContentBrowser::HandleDeleteFavorite(TSharedPtr<SWidget> ParentWidget)
 			LOCTEXT("FavoriteRemoveConfirm_Yes", "Remove"),
 			LOCTEXT("FavoriteRemoveConfirm_No", "Cancel"),
 			ParentWidget.ToSharedRef(),
-			FOnClicked::CreateLambda([this, SelectedFolders]() -> FReply
+			FOnClicked::CreateSPLambda(this, [this, SelectedFolders]() -> FReply
 			{
 				for (const FContentBrowserItem& Folder : SelectedFolders)
 				{
diff --git a/Engine/Source/Editor/ContentBrowser/Private/SContentBrowser.h b/Engine/Source/Editor/ContentBrowser/Private/SContentBrowser.h
--- a/Engine/Source/Editor/ContentBrowser/Private/SContentBrowser.h
+++ b/Engine/Source/Editor/ContentBrowser/Private/SContentBrowser.h
@@ -731,7 +731,7 @@ private:
 		TSharedPtr<SBox> VerticalFilterViewContainer;
 	};
 
-	TSharedPtr<FLegacyContentSourceWidgets> LegacyContentSourceWidgets;
+	TUniquePtr<FLegacyContentSourceWidgets> LegacyContentSourceWidgets;
 
 	/** The tab that contains this browser */
 	TWeakPtr<SDockTab> ContainingTab;
diff --git a/Engine/Source/Editor/ContentBrowser/Private/SourcesViewWidgets.cpp b/Engine/Source/Editor/ContentBrowser/Private/SourcesViewWidgets.cpp
--- a/Engine/Source/Editor/ContentBrowser/Private/SourcesViewWidgets.cpp
+++ b/Engine/Source/Editor/ContentBrowser/Private/SourcesViewWidgets.cpp
@@ -565,10 +565,13 @@ void SCollectionTreeItem::HandleNameCommitted( const FText& NewText, ETextCommit
 			{
 				FText WarningMessage;
 				bool bIsCommitted = (CommitInfo != ETextCommit::OnCleared);
-				if ( !OnNameChangeCommit.Execute(CollectionItemPtr, NewText.ToString(), bIsCommitted, WarningMessage) && ParentWidget.IsValid() && bIsCommitted )
+				if ( !OnNameChangeCommit.Execute(CollectionItemPtr, NewText.ToString(), bIsCommitted, WarningMessage) && bIsCommitted )
 				{
-					// Failed to rename/create a collection, display a warning.
-					ContentBrowserUtils::DisplayMessage(WarningMessage, CachedGeometry.GetLayoutBoundingRect(), ParentWidget.ToSharedRef(), ContentBrowserUtils::EDisplayMessageType::Error);
+					if (TSharedPtr<SWidget> PinnedParentWidget = ParentWidget.Pin())
+					{
+						// Failed to rename/create a collection, display a warning.
+						ContentBrowserUtils::DisplayMessage(WarningMessage, CachedGeometry.GetLayoutBoundingRect(), PinnedParentWidget.ToSharedRef(), ContentBrowserUtils::EDisplayMessageType::Error);
+					}
 				}
 			}				
 		}
diff --git a/Engine/Source/Editor/ContentBrowser/Private/SourcesViewWidgets.h b/Engine/Source/Editor/ContentBrowser/Private/SourcesViewWidgets.h
--- a/Engine/Source/Editor/ContentBrowser/Private/SourcesViewWidgets.h
+++ b/Engine/Source/Editor/ContentBrowser/Private/SourcesViewWidgets.h
@@ -251,8 +251,8 @@ private:
 	void BuildToolTipInfo(const FOnBuildAssetTagItemToolTipInfoEntry& InCallback) const;
 
 private:
-	/** A shared pointer to the parent widget. */
-	TSharedPtr<SWidget> ParentWidget;
+	/** A weak pointer to the parent widget. */
+	TWeakPtr<SWidget> ParentWidget;
 
 	/** The data for this item */
 	TWeakPtr<FCollectionItem> CollectionItem;
diff --git a/Engine/Source/Editor/ContentBrowser/Private/Widgets/SContentBrowserSourceTree.cpp b/Engine/Source/Editor/ContentBrowser/Private/Widgets/SContentBrowserSourceTree.cpp
--- a/Engine/Source/Editor/ContentBrowser/Private/Widgets/SContentBrowserSourceTree.cpp
+++ b/Engine/Source/Editor/ContentBrowser/Private/Widgets/SContentBrowserSourceTree.cpp
@@ -329,20 +329,19 @@ namespace UE::Editor::ContentBrowser::Private
 			SearchButtonWidget =
 				SAssignNew(SearchToggleButton, SSearchToggleButton, Search->GetWidget())
 				.Visibility(this, &SContentBrowserSourceTreeArea::GetHeaderSearchActionVisibility)
-				.OnSearchBoxShown_Lambda([this]() { SetExpanded(true); });
+				.OnSearchBoxShown(FOnSearchBoxShown::CreateSPLambda(this, [this]() { SetExpanded(true); }));
 		}
 
 		TAttribute<int32> BodyContentIndexAttribute;
 		if (UE::Editor::ContentBrowser::IsNewStyleEnabled()
 			&& (InArgs._IsEmpty.IsSet() || InArgs._IsEmpty.IsBound()))
 		{
-			BodyContentIndexAttribute = TAttribute<int32>::Create(
-			TAttribute<int32>::FGetter::CreateSPLambda(
+			BodyContentIndexAttribute = MakeAttributeSPLambda(
 				this,
-				[IsEmptyAttribute = InArgs._IsEmpty]()
-			{
-				return IsEmptyAttribute.Get() ? 1 : 0;
-			}));
+				[this]()
+				{
+					return IsEmpty() ? 1 : 0;
+				});
 		}
 		else
 		{
