Fixed Content Browser UI getting out of sync with the open Collection Containers.

#rnx

diff --git a/Engine/Source/Editor/ContentBrowser/Private/Widgets/SContentBrowserSourceTree.cpp b/Engine/Source/Editor/ContentBrowser/Private/Widgets/SContentBrowserSourceTree.cpp
--- a/Engine/Source/Editor/ContentBrowser/Private/Widgets/SContentBrowserSourceTree.cpp
+++ b/Engine/Source/Editor/ContentBrowser/Private/Widgets/SContentBrowserSourceTree.cpp
@@ -123,7 +123,7 @@ namespace UE::Editor::ContentBrowser::Private
 			[this, HostSplitter = Splitter](const FSlot* InSlotAdded, int32 InSlotIdx)
 			{
 				// When the slot is added, also add it to the splitter
-				HostSplitter->AddSlot()
+				HostSplitter->AddSlot(InSlotIdx)
 				.Value(InSlotAdded->GetInitialSlotSize())
 				.SizeRule(this, &SContentBrowserSourceTree::GetExpandedSizeRule, InSlotAdded->GetEntryWidget())
 				.MinSize(this, &SContentBrowserSourceTree::GetMinHeight, InSlotAdded->GetEntryWidget())
