EditorUsability : ContentBrowser
- Added Filter by Selection capability in Content Browser and AssetPicker with keybind ALT+F, it uses the active selection in the context to set filters
- Extending capability to AssetPicker as well when filtering is enabled

diff --git a/Engine/Source/Editor/ContentBrowser/Private/ContentBrowserCommands.cpp b/Engine/Source/Editor/ContentBrowser/Private/ContentBrowserCommands.cpp
--- a/Engine/Source/Editor/ContentBrowser/Private/ContentBrowserCommands.cpp
+++ b/Engine/Source/Editor/ContentBrowser/Private/ContentBrowserCommands.cpp
@@ -23,6 +23,7 @@ void FContentBrowserCommands::RegisterCommands()
 	UI_COMMAND(SaveAllCurrentFolder, "Save All", "Save All in current folder", EUserInterfaceActionType::Button, FInputChord());
 	UI_COMMAND(ResaveAllCurrentFolder, "Resave All", "Resave all assets contained in the current folder", EUserInterfaceActionType::Button, FInputChord());
 	UI_COMMAND(EditPath, "Edit Path", "Edit the current content browser path", EUserInterfaceActionType::Button, FInputChord(EKeys::F4), FInputChord(EModifierKey::Control, EKeys::L));
+	UI_COMMAND(FilterBySelection, "Filter By Selection", "Apply a filter based on the asset types of the current selection of the content browser", EUserInterfaceActionType::Button, FInputChord(EModifierKey::Alt, EKeys::F));
 
 	// AssetView commands
 	UI_COMMAND(AssetViewCopyObjectPath, "Copy Selected Object Path", "Copy the selected object path", EUserInterfaceActionType::Button, FInputChord(EModifierKey::Control | EModifierKey::Shift, EKeys::C));
diff --git a/Engine/Source/Editor/ContentBrowser/Private/ContentBrowserCommands.h b/Engine/Source/Editor/ContentBrowser/Private/ContentBrowserCommands.h
--- a/Engine/Source/Editor/ContentBrowser/Private/ContentBrowserCommands.h
+++ b/Engine/Source/Editor/ContentBrowser/Private/ContentBrowserCommands.h
@@ -16,7 +16,6 @@ class FContentBrowserCommands
 	: public TCommands<FContentBrowserCommands>
 {
 public:
-
 	static const FName GenericCommandsName;
 
 	/** Default constructor. */
@@ -24,13 +23,9 @@ public:
 		: TCommands<FContentBrowserCommands>(TEXT("ContentBrowser"), NSLOCTEXT( "ContentBrowser", "ContentBrowser", "Content Browser" ), (FInputBindingManager::Get().GetContextByName(GenericCommandsName).IsValid() ? GenericCommandsName : NAME_None), FAppStyle::GetAppStyleSetName())
 	{ }
 
-public:
-
-	//~ TCommands interface
-
+	//~ Begin TCommands
 	virtual void RegisterCommands() override;
-
-public:
+	//~ End TCommands
 
 	TSharedPtr<FUICommandInfo> CreateNewFolder;
 	TSharedPtr<FUICommandInfo> OpenAssetsOrFolders;
@@ -45,4 +40,5 @@ public:
 	TSharedPtr<FUICommandInfo> GridViewShortcut;
 	TSharedPtr<FUICommandInfo> ListViewShortcut;
 	TSharedPtr<FUICommandInfo> ColumnViewShortcut;
+	TSharedPtr<FUICommandInfo> FilterBySelection;
 };
diff --git a/Engine/Source/Editor/ContentBrowser/Private/SAssetPicker.cpp b/Engine/Source/Editor/ContentBrowser/Private/SAssetPicker.cpp
--- a/Engine/Source/Editor/ContentBrowser/Private/SAssetPicker.cpp
+++ b/Engine/Source/Editor/ContentBrowser/Private/SAssetPicker.cpp
@@ -8,7 +8,9 @@
 #include "AssetThumbnail.h"
 #include "CollectionManagerModule.h"
 #include "CollectionManagerTypes.h"
+#include "ContentBrowserCommands.h"
 #include "ContentBrowserDataFilter.h"
+#include "ContentBrowserDataSource.h"
 #include "ContentBrowserItem.h"
 #include "ContentBrowserStyle.h"
 #include "ContentBrowserUtils.h"
@@ -224,6 +226,7 @@ void SAssetPicker::Construct( const FArguments& InArgs )
 		}		
 		
 		SAssignNew(FilterListPtr, SFilterList)
+			.CommandList(Commands)
 			.OnFilterChanged(this, &SAssetPicker::OnFilterChanged)
 			.FrontendFilters(FrontendFilters)
 			.InitialClassFilters(FilterClassList)
@@ -816,6 +819,40 @@ bool SAssetPicker::CanExecuteRenameRequested()
 	return ContentBrowserUtils::CanRenameFromAssetView(AssetViewPtr);
 }
 
+bool SAssetPicker::CanFilterBySelection()
+{
+	if (!AssetViewPtr.IsValid() || !FilterListPtr.IsValid())
+	{
+		return false;
+	}
+
+	const TArray<FContentBrowserItem> SelectedItems = AssetViewPtr->GetSelectedItems();
+	return !SelectedItems.IsEmpty();
+}
+
+void SAssetPicker::OnFilterBySelection()
+{
+	if (!AssetViewPtr.IsValid() || !FilterListPtr.IsValid())
+	{
+		return;
+	}
+
+	const TArray<FContentBrowserItem> SelectedItems = AssetViewPtr->GetSelectedItems();
+
+	TArray<FTopLevelAssetPath> AssetPaths;
+	AssetPaths.Reserve(SelectedItems.Num());
+	for (const FContentBrowserItem& SelectedItem : SelectedItems)
+	{
+		FContentBrowserItemDataAttributeValue AttributeValue = SelectedItem.GetItemAttribute(ContentBrowserItemAttributes::ItemTypeName);
+		if (AttributeValue.IsValid())
+		{
+			AssetPaths.Add(FTopLevelAssetPath(AttributeValue.GetValueString()));
+		}
+	}
+
+	FilterListPtr->SetActiveFilters(AssetPaths);
+}
+
 void SAssetPicker::ExecuteRenameCommand()
 {
 	if (Commands.IsValid())
@@ -832,6 +869,14 @@ void SAssetPicker::BindCommands()
 		FExecuteAction::CreateSP( this, &SAssetPicker::OnRenameRequested ),
 		FCanExecuteAction::CreateSP( this, &SAssetPicker::CanExecuteRenameRequested )
 		));
+
+	Commands->MapAction(
+		FContentBrowserCommands::Get().FilterBySelection,
+		FUIAction(
+			FExecuteAction::CreateSP(this, &SAssetPicker::OnFilterBySelection),
+			FCanExecuteAction::CreateSP(this, &SAssetPicker::CanFilterBySelection)
+		)
+	);
 }
 
 void SAssetPicker::LoadSettings()
diff --git a/Engine/Source/Editor/ContentBrowser/Private/SAssetPicker.h b/Engine/Source/Editor/ContentBrowser/Private/SAssetPicker.h
--- a/Engine/Source/Editor/ContentBrowser/Private/SAssetPicker.h
+++ b/Engine/Source/Editor/ContentBrowser/Private/SAssetPicker.h
@@ -135,6 +135,12 @@ private:
 	/** Returns true if the user is able to execute a rename request */
 	bool CanExecuteRenameRequested();
 
+	/** Returns true if the user is able to execute filter by selection request */
+	bool CanFilterBySelection();
+
+	/** Executes command and applies filters based on active selection */
+	void OnFilterBySelection();
+
 	/** Bind our UI commands */
 	void BindCommands();
 
diff --git a/Engine/Source/Editor/ContentBrowser/Private/SContentBrowser.cpp b/Engine/Source/Editor/ContentBrowser/Private/SContentBrowser.cpp
--- a/Engine/Source/Editor/ContentBrowser/Private/SContentBrowser.cpp
+++ b/Engine/Source/Editor/ContentBrowser/Private/SContentBrowser.cpp
@@ -427,59 +427,67 @@ void SContentBrowser::BindCommands()
 {
 	Commands = TSharedPtr< FUICommandList >(new FUICommandList);
 
-	Commands->MapAction(FGenericCommands::Get().Rename, FUIAction(
+	FContentBrowserCommands& ContentBrowserCommands = FContentBrowserCommands::Get();
+	FGenericCommands& GenericCommands = FGenericCommands::Get();
+
+	Commands->MapAction(GenericCommands.Rename, FUIAction(
 		FExecuteAction::CreateSP(this, &SContentBrowser::HandleRenameCommand),
 		FCanExecuteAction::CreateSP(this, &SContentBrowser::HandleRenameCommandCanExecute)
 	));
 
-	Commands->MapAction(FGenericCommands::Get().Delete, FUIAction(
+	Commands->MapAction(GenericCommands.Delete, FUIAction(
 		FExecuteAction::CreateSP(this, &SContentBrowser::HandleDeleteCommandExecute),
 		FCanExecuteAction::CreateSP(this, &SContentBrowser::HandleDeleteCommandCanExecute)
 	));
 
-	Commands->MapAction(FContentBrowserCommands::Get().OpenAssetsOrFolders, FUIAction(
+	Commands->MapAction(ContentBrowserCommands.OpenAssetsOrFolders, FUIAction(
 		FExecuteAction::CreateSP(this, &SContentBrowser::HandleOpenAssetsOrFoldersCommandExecute)
 	));
 
-	Commands->MapAction(FContentBrowserCommands::Get().PreviewAssets, FUIAction(
+	Commands->MapAction(ContentBrowserCommands.PreviewAssets, FUIAction(
 		FExecuteAction::CreateSP(this, &SContentBrowser::HandlePreviewAssetsCommandExecute)
 	));
 
-	Commands->MapAction(FContentBrowserCommands::Get().CreateNewFolder, FUIAction(
+	Commands->MapAction(ContentBrowserCommands.CreateNewFolder, FUIAction(
 		FExecuteAction::CreateSP(this, &SContentBrowser::HandleCreateNewFolderCommandExecute)
 	));
 
-	Commands->MapAction(FContentBrowserCommands::Get().GoUpToParentFolder, FUIAction(
+	Commands->MapAction(ContentBrowserCommands.GoUpToParentFolder, FUIAction(
 		FExecuteAction::CreateSP(this, &SContentBrowser::HandleGoUpToParentFolder),
 		FCanExecuteAction::CreateSP(this, &SContentBrowser::HandleCanGoUpToParentFolder)
 	));
 
-	Commands->MapAction(FContentBrowserCommands::Get().SaveSelectedAsset, FUIAction(
+	Commands->MapAction(ContentBrowserCommands.SaveSelectedAsset, FUIAction(
 		FExecuteAction::CreateSP(this, &SContentBrowser::HandleSaveAssetCommand),
 		FCanExecuteAction::CreateSP(this, &SContentBrowser::HandleSaveAssetCommandCanExecute)
 	));
 
-	Commands->MapAction(FContentBrowserCommands::Get().SaveAllCurrentFolder, FUIAction(
+	Commands->MapAction(ContentBrowserCommands.SaveAllCurrentFolder, FUIAction(
 		FExecuteAction::CreateSP(this, &SContentBrowser::HandleSaveAllCurrentFolderCommand)
 	));
 
-	Commands->MapAction(FContentBrowserCommands::Get().ResaveAllCurrentFolder, FUIAction(
+	Commands->MapAction(ContentBrowserCommands.ResaveAllCurrentFolder, FUIAction(
 		FExecuteAction::CreateSP(this, &SContentBrowser::HandleResaveAllCurrentFolderCommand)
 	));
 
-	Commands->MapAction(FContentBrowserCommands::Get().EditPath, FUIAction(
+	Commands->MapAction(ContentBrowserCommands.EditPath, FUIAction(
 		FExecuteAction::CreateSP(this, &SContentBrowser::EditPathCommand)
 	));
 
+	Commands->MapAction(ContentBrowserCommands.FilterBySelection, FUIAction(
+		FExecuteAction::CreateSP(this, &SContentBrowser::HandleFilterBySelectionCommandExecute),
+		FCanExecuteAction::CreateSP(this, &SContentBrowser::HandleFilterBySelectionCommandCanExecute)
+	));
+
 	// Allow extenders to add commands
 	FContentBrowserModule& ContentBrowserModule = FModuleManager::GetModuleChecked<FContentBrowserModule>(TEXT("ContentBrowser"));
-	TArray<FContentBrowserCommandExtender> CommmandExtenderDelegates = ContentBrowserModule.GetAllContentBrowserCommandExtenders();
+	TArray<FContentBrowserCommandExtender> CommandExtenderDelegates = ContentBrowserModule.GetAllContentBrowserCommandExtenders();
 
-	for (int32 i = 0; i < CommmandExtenderDelegates.Num(); ++i)
+	for (int32 i = 0; i < CommandExtenderDelegates.Num(); ++i)
 	{
-		if (CommmandExtenderDelegates[i].IsBound())
+		if (CommandExtenderDelegates[i].IsBound())
 		{
-			CommmandExtenderDelegates[i].Execute(Commands.ToSharedRef(), FOnContentBrowserGetSelection::CreateSP(this, &SContentBrowser::GetSelectionState));
+			CommandExtenderDelegates[i].Execute(Commands.ToSharedRef(), FOnContentBrowserGetSelection::CreateSP(this, &SContentBrowser::GetSelectionState));
 		}
 	}
 
@@ -942,6 +950,7 @@ TSharedRef<SWidget> SContentBrowser::CreateAssetView(const FContentBrowserConfig
 {
 	// Create the Filter Bar Widget
 	LegacyContentSourceWidgets->FilterListPtr = SNew(SFilterList)
+					.CommandList(Commands)
 					.OnFilterChanged(this, &SContentBrowser::OnFilterChanged)
 					.Visibility((Config != nullptr ? Config->bCanShowFilters : true) ? EVisibility::Visible : EVisibility::Collapsed)
 					.FrontendFilters(FrontendFilters)
@@ -3853,6 +3862,40 @@ FReply SContentBrowser::ForwardClicked()
 	return FReply::Handled();
 }
 
+bool SContentBrowser::HandleFilterBySelectionCommandCanExecute() const
+{
+	if (!LegacyContentSourceWidgets.IsValid() || !LegacyContentSourceWidgets->FilterListPtr.IsValid())
+	{
+		return false;
+	}
+
+	const TArray<FAssetData> SelectedAssets = LegacyContentSourceWidgets->AssetViewPtr->GetSelectedAssets();
+	return !SelectedAssets.IsEmpty();
+}
+
+void SContentBrowser::HandleFilterBySelectionCommandExecute()
+{
+	if (!LegacyContentSourceWidgets.IsValid() || !LegacyContentSourceWidgets->FilterListPtr.IsValid())
+	{
+		return;
+	}
+
+	const TArray<FAssetData> SelectedAssets = LegacyContentSourceWidgets->AssetViewPtr->GetSelectedAssets();
+	if (SelectedAssets.IsEmpty())
+	{
+		return;
+	}
+
+	TArray<FTopLevelAssetPath> SelectedAssetPaths;
+	SelectedAssetPaths.Reserve(SelectedAssets.Num());
+	for (const FAssetData& SelectedAsset : SelectedAssets)
+	{
+		SelectedAssetPaths.Add(SelectedAsset.AssetClassPath);
+	}
+
+	LegacyContentSourceWidgets->FilterListPtr->SetActiveFilters(SelectedAssetPaths);
+}
+
 bool SContentBrowser::HandleRenameCommandCanExecute() const
 {
 	// The order of these conditions are carefully crafted to match the logic of the context menu summoning, as this callback 
diff --git a/Engine/Source/Editor/ContentBrowser/Private/SContentBrowser.h b/Engine/Source/Editor/ContentBrowser/Private/SContentBrowser.h
--- a/Engine/Source/Editor/ContentBrowser/Private/SContentBrowser.h
+++ b/Engine/Source/Editor/ContentBrowser/Private/SContentBrowser.h
@@ -443,6 +443,12 @@ private:
 	/** Handler for clicking the history forward button */
 	FReply ForwardClicked();
 
+	/** Handler to check if filter by selection command is allowed */
+	bool HandleFilterBySelectionCommandCanExecute() const;
+
+	/** Handler for filter by selection */
+	void HandleFilterBySelectionCommandExecute();
+
 	/** Handler to check to see if a rename command is allowed */
 	bool HandleRenameCommandCanExecute() const;
 
@@ -680,7 +686,10 @@ private:
 	
 private:
 
-	/** Struct that contains the widgets used by content browser's default layout (i.e Legacy Content Source) */
+	/**
+	 * Struct that contains the widgets used by content browser's default layout (i.e Legacy Content Source)
+	 * ContentBrowser 2.0 will use TEDS, this is moved here but is still in use and correct while Content Browser refactor is on hold
+	 */
 	struct FLegacyContentSourceWidgets
 	{
 		/** The context menu manager for the path view */
diff --git a/Engine/Source/Editor/ContentBrowser/Private/SFilterList.cpp b/Engine/Source/Editor/ContentBrowser/Private/SFilterList.cpp
--- a/Engine/Source/Editor/ContentBrowser/Private/SFilterList.cpp
+++ b/Engine/Source/Editor/ContentBrowser/Private/SFilterList.cpp
@@ -5,6 +5,7 @@
 
 #include "Algo/AnyOf.h"
 #include "AssetRegistry/ARFilter.h"
+#include "ContentBrowserCommands.h"
 #include "ContentBrowserDataFilter.h"
 #include "ContentBrowserDataSource.h"
 #include "ContentBrowserDataSubsystem.h"
@@ -12,6 +13,7 @@
 #include "ContentBrowserItem.h"
 #include "ContentBrowserItemData.h"
 #include "ContentBrowserMenuContexts.h"
+#include "ContentBrowserModule.h"
 #include "ContentBrowserStyle.h"
 #include "ContentBrowserUtils.h"
 #include "Filters.h"
@@ -69,7 +71,7 @@ void SFilterList::Construct( const FArguments& InArgs )
 	OnFilterBarLayoutChanging = InArgs._OnFilterBarLayoutChanging;
 	this->OnFilterChanged = InArgs._OnFilterChanged;
 	this->ActiveFilters = InArgs._FrontendFilters;
-	InitialClassFilters = InArgs._InitialClassFilters; 
+	InitialClassFilters = InArgs._InitialClassFilters;
 
 	TSharedPtr<FFrontendFilterCategory> DefaultCategory = MakeShareable( new FFrontendFilterCategory(LOCTEXT("FrontendFiltersCategory", "Other Filters"), LOCTEXT("FrontendFiltersCategoryTooltip", "Filter assets by all filters in this category.")) );
 	
@@ -151,6 +153,7 @@ void SFilterList::Construct( const FArguments& InArgs )
 	Args._FilterMenuName = FName("ContentBrowser.FilterMenu");
 	Args._DefaultMenuExpansionCategory = InArgs._DefaultMenuExpansionCategory;
 	Args._bUseSectionsForCustomCategories = InArgs._bUseSectionsForCustomCategories;
+	Args._CommandList = InArgs._CommandList;
 
 	SAssetFilterBar<FAssetFilterType>::Construct(Args);
 
@@ -464,6 +467,16 @@ void SFilterList::PopulateAddFilterMenu(UToolMenu* Menu, TSharedPtr<FFilterCateg
 					FIsActionChecked::CreateSP(this, &SFilterList::IsFilteringPathsRecursively)),
 				EUserInterfaceActionType::ToggleButton
 			);
+
+			// Add filter by selection if command is mapped in the command list
+			if (const TSharedPtr<FUICommandList> CommandList = CommandListWeak.Pin())
+			{
+				const FContentBrowserCommands& ContentBrowserCommands = FContentBrowserCommands::Get();
+				if (CommandList->IsActionMapped(ContentBrowserCommands.FilterBySelection))
+				{
+					ManageSection->AddMenuEntry(ContentBrowserCommands.FilterBySelection);
+				}
+			}
 		}
 
 		FToolMenuSection& DisplaySection = Menu->AddSection(
@@ -657,6 +670,15 @@ void SFilterList::ToggleFilterPathsRecursively()
 	SetFilterPathsRecursively(!bFilterPathsRecursively);
 }
 
+void SFilterList::SetActiveFilters(TConstArrayView<FTopLevelAssetPath> InFilters)
+{
+	RemoveAllFilters();
+	for (const FTopLevelAssetPath& FilterAssetPath : InFilters)
+	{
+		SetAssetTypeFilterCheckState(FilterAssetPath, ECheckBoxState::Checked, /** Enabled */true);
+	}
+}
+
 TSharedPtr<SWidget> SFilterList::GetActiveFilterContainer() const
 {
 	if (HorizontalFilterBox.IsValid() && HorizontalFilterBox->GetChildren()->Num() > 0)
@@ -835,7 +857,7 @@ void SFilterList::SaveSettingsInternal(FFilterBarSettings* InSettingsToSave)
 	InstanceSettings->Empty();
 	SharedSettings->Empty();
 
-	// Save all the programatically added filters normally
+	// Save all the programmatically added filters normally
 	SaveFilters(InstanceSettings);
 
 	/** For each custom text filter: Save the filterdata into the common settings, so that all instances that use it
@@ -962,7 +984,7 @@ void SFilterList::LoadSettings(const FName& InInstanceName)
 	const FFilterBarSettings* InstanceSettings = UFilterBarConfig::Get()->FilterBars.Find(InInstanceName);
 	const FFilterBarSettings* SharedSettings = UFilterBarConfig::Get()->FilterBars.Find(SharedIdentifier);
 
-	// Load the filters specified programatically normally
+	// Load the filters specified programmatically normally
 	LoadFilters(InstanceSettings);
 
 	// Load the custom text filters from the shared settings
diff --git a/Engine/Source/Editor/ContentBrowser/Private/SFilterList.h b/Engine/Source/Editor/ContentBrowser/Private/SFilterList.h
--- a/Engine/Source/Editor/ContentBrowser/Private/SFilterList.h
+++ b/Engine/Source/Editor/ContentBrowser/Private/SFilterList.h
@@ -103,6 +103,9 @@ public:
 		/** If true, adds custom categories as sections (expanded) vs. as sub-menus */
 		SLATE_ARGUMENT(bool, bUseSectionsForCustomCategories)
 
+		/** The command list provided to the filter context menu */
+		SLATE_ARGUMENT(TSharedPtr<FUICommandList>, CommandList)
+
 	SLATE_END_ARGS()
 
 	/** Constructs this widget with InArgs */
@@ -177,6 +180,9 @@ public:
 	/** Toggle whether the filters are being applied to paths recursively */
 	void ToggleFilterPathsRecursively();
 
+	/** Set the active asset filters for this filter list */
+	void SetActiveFilters(TConstArrayView<FTopLevelAssetPath> InFilters);
+
 	/** Returns a pointer to the active filter container widget (Horizontal or Vertical), nullptr if neither is active */
 	TSharedPtr<SWidget> GetActiveFilterContainer() const;
 
@@ -215,7 +221,6 @@ private:
 	bool RestoreCustomTextFilterState(const FCustomTextFilterState& InFilterState);
 	
 private:
-
 	/** List of classes that our filters must match */
 	TArray<UClass*> InitialClassFilters;
 
diff --git a/Engine/Source/Editor/EditorWidgets/Public/Filters/SAssetFilterBar.h b/Engine/Source/Editor/EditorWidgets/Public/Filters/SAssetFilterBar.h
--- a/Engine/Source/Editor/EditorWidgets/Public/Filters/SAssetFilterBar.h
+++ b/Engine/Source/Editor/EditorWidgets/Public/Filters/SAssetFilterBar.h
@@ -233,6 +233,9 @@ public:
 
 		/** If true, adds custom categories as sections (expanded) vs. as sub-menus */
 		SLATE_ARGUMENT(bool, bUseSectionsForCustomCategories)
+
+		/** The command list provided to the filter context menu */
+		SLATE_ARGUMENT(TSharedPtr<FUICommandList>, CommandList)
 	
 	SLATE_END_ARGS()
 
@@ -247,6 +250,7 @@ public:
 		FilterMenuName = InArgs._FilterMenuName;
 		FilterBarIdentifier = InArgs._FilterBarIdentifier;
 		DefaultMenuExpansionCategory = InArgs._DefaultMenuExpansionCategory;
+		CommandListWeak = InArgs._CommandList;
 		
 		typename SBasicFilterBar<FilterType>::FArguments Args;
 		Args._OnFilterChanged = InArgs._OnFilterChanged;
@@ -374,11 +378,11 @@ protected:
 		SaveConfig();
 	}
 
-	/** Save all the custom filters (created programatically) into the specified config */
+	/** Save all the custom filters (created programmatically) into the specified config */
 	void SaveFilters(FFilterBarSettings* FilterBarConfig)
 	{
 		/* For the remaining (custom and type) filters, go through the currently active (i.e visible in the filter bar)
-		 * ones to save their state, since they will be added to the filter bar programatically every time
+		 * ones to save their state, since they will be added to the filter bar programmatically every time
 		 */
 		for ( const TSharedPtr<SFilter> Filter : this->Filters )
 		{
@@ -436,7 +440,7 @@ protected:
 		}
 	}
 	
-	/** Load all the custom filters (created programatically) from the specified config */
+	/** Load all the custom filters (created programmatically) from the specified config */
 	void LoadFilters(const FFilterBarSettings* FilterBarConfig)
 	{
 		// Load all the custom filters (i.e FrontendFilters)
@@ -1059,7 +1063,7 @@ protected:
 			}));
 		}
 
-		FToolMenuContext ToolMenuContext(CreateAssetFilterBarContext());
+		const FToolMenuContext ToolMenuContext(CommandListWeak.Pin(), TSharedPtr<FExtender>(), CreateAssetFilterBarContext());
 		return UToolMenus::Get()->GenerateWidget(FilterMenuName, ToolMenuContext);
 	} 
 
@@ -1245,6 +1249,9 @@ protected:
 	}
 
 protected:
+	/** The command list we bind to */
+	TWeakPtr<FUICommandList> CommandListWeak;
+
 	/** A copy of all AssetFilters in this->Filters for convenient access */
 	TArray< TSharedRef<SAssetFilter> > AssetFilters;
 
