Added a new SyncBrowserToCollections function to the UCollectionManagerScriptingSubsystem that allows the user to select collections in the content browser.

diff --git a/Engine/Source/Editor/ContentBrowser/Private/ContentBrowserSingleton.cpp b/Engine/Source/Editor/ContentBrowser/Private/ContentBrowserSingleton.cpp
--- a/Engine/Source/Editor/ContentBrowser/Private/ContentBrowserSingleton.cpp
+++ b/Engine/Source/Editor/ContentBrowser/Private/ContentBrowserSingleton.cpp
@@ -551,6 +551,21 @@ void FContentBrowserSingleton::SyncBrowserToFolders(const TArray<FString>& Folde
 	}
 }
 
+void FContentBrowserSingleton::SyncBrowserToCollections(const TArray<FCollectionRef>& CollectionsToSync, bool bAllowLockedBrowsers, bool bFocusContentBrowser, const FName& InstanceName, bool bNewSpawnBrowser)
+{
+	TSharedPtr<SContentBrowser> ContentBrowserToSync = FindContentBrowserToSync(bAllowLockedBrowsers, InstanceName, bNewSpawnBrowser);
+
+	if (ContentBrowserToSync.IsValid())
+	{
+		// Finally, focus and sync the browser that was found
+		if (bFocusContentBrowser)
+		{
+			FocusContentBrowser(ContentBrowserToSync);
+		}
+		ContentBrowserToSync->SyncToCollections(CollectionsToSync);
+	}
+}
+
 void FContentBrowserSingleton::SyncBrowserToItems(const TArray<FContentBrowserItem>& ItemsToSync, bool bAllowLockedBrowsers, bool bFocusContentBrowser, const FName& InstanceName, bool bNewSpawnBrowser)
 {
 	TSharedPtr<SContentBrowser> ContentBrowserToSync = FindContentBrowserToSync(bAllowLockedBrowsers, InstanceName, bNewSpawnBrowser);
diff --git a/Engine/Source/Editor/ContentBrowser/Private/ContentBrowserSingleton.h b/Engine/Source/Editor/ContentBrowser/Private/ContentBrowserSingleton.h
--- a/Engine/Source/Editor/ContentBrowser/Private/ContentBrowserSingleton.h
+++ b/Engine/Source/Editor/ContentBrowser/Private/ContentBrowserSingleton.h
@@ -92,6 +92,7 @@ public:
 	virtual void SyncBrowserToAssets(const TArray<struct FAssetData>& AssetDataList, bool bAllowLockedBrowsers = false, bool bFocusContentBrowser = true, const FName& InstanceName = FName(), bool bNewSpawnBrowser = false) override;
 	virtual void SyncBrowserToAssets(const TArray<UObject*>& AssetList, bool bAllowLockedBrowsers = false, bool bFocusContentBrowser = true, const FName& InstanceName = FName(), bool bNewSpawnBrowser = false) override;
 	virtual void SyncBrowserToFolders(const TArray<FString>& FolderList, bool bAllowLockedBrowsers = false, bool bFocusContentBrowser = true, const FName& InstanceName = FName(), bool bNewSpawnBrowser = false) override;
+	virtual void SyncBrowserToCollections(const TArray<FCollectionRef>& CollectionsToSync, bool bAllowLockedBrowsers = false, bool bFocusContentBrowser = true, const FName& InstanceName = FName(), bool bNewSpawnBrowser = false) override;
 	virtual void SyncBrowserToItems(const TArray<FContentBrowserItem>& ItemsToSync, bool bAllowLockedBrowsers = false, bool bFocusContentBrowser = true, const FName& InstanceName = FName(), bool bNewSpawnBrowser = false) override;
 	virtual void SyncBrowserTo(const FContentBrowserSelection& ItemSelection, bool bAllowLockedBrowsers = false, bool bFocusContentBrowser = true, const FName& InstanceName = FName(), bool bNewSpawnBrowser = false) override;
 	virtual void GetSelectedAssets(TArray<FAssetData>& SelectedAssets) override;
diff --git a/Engine/Source/Editor/ContentBrowser/Private/SContentBrowser.cpp b/Engine/Source/Editor/ContentBrowser/Private/SContentBrowser.cpp
--- a/Engine/Source/Editor/ContentBrowser/Private/SContentBrowser.cpp
+++ b/Engine/Source/Editor/ContentBrowser/Private/SContentBrowser.cpp
@@ -1596,6 +1596,47 @@ void SContentBrowser::SyncToFolders(TArrayView<const FString> FolderList, const
 	SyncToLegacy(TArrayView<const FAssetData>(), FolderList, bAllowImplicitSync, /*bDisableFiltersThatHideAssets*/false);
 }
 
+void SContentBrowser::SyncToCollections(TArrayView<const FCollectionRef> CollectionsToSync)
+{
+	// For now we just forcefully enable the legacy content source when a sync is requested so it succeeds
+	ContentSourcesContainer->ActivateLegacyContentSource();
+
+	// You may not select both collections and paths
+	LegacyContentSourceWidgets->PathViewPtr->ClearSelection();
+	LegacyContentSourceWidgets->FavoritePathViewPtr->ClearSelection();
+
+	TArray<FCollectionRef> SelectedCollections;
+	SelectedCollections.Reserve(CollectionsToSync.Num());
+	{
+		TArray<FCollectionNameType> CollectionsToSelect;
+		for (const TUniquePtr<FCollectionSource>& CollectionSource : CollectionSources)
+		{
+			CollectionsToSelect.Reset(CollectionsToSync.Num());
+			for (const FCollectionRef& CollectionToSync : CollectionsToSync)
+			{
+				if (CollectionSource->GetCollectionContainer() == CollectionToSync.Container)
+				{
+					CollectionsToSelect.Emplace(CollectionToSync.Name, CollectionToSync.Type);
+				}
+			}
+			CollectionSource->CollectionViewPtr->SetSelectedCollections(CollectionsToSelect),
+
+			Algo::Transform(
+				CollectionSource->CollectionViewPtr->GetSelectedCollections(),
+				SelectedCollections,
+				[&CollectionSource](const FCollectionNameType& Collection) { return FCollectionRef(CollectionSource->GetCollectionContainer(), Collection); });
+		}
+	}
+
+	if (SelectedCollections.Num() == 0 && !CollectionSources.IsEmpty())
+	{
+		// Select a dummy "None" collection to avoid the sources view switching to the paths view
+		SelectedCollections.Emplace(CollectionSources[0]->GetCollectionContainer(), FCollectionNameType(NAME_None, ECollectionShareType::CST_System));
+	}
+
+	SourcesChanged({}, SelectedCollections);
+}
+
 void SContentBrowser::SyncToItems(TArrayView<const FContentBrowserItem> ItemsToSync, const bool bAllowImplicitSync, const bool bDisableFiltersThatHideAssets)
 {
 	// For now we just forcefully enable the legacy content source when a sync is requested so it succeeds
diff --git a/Engine/Source/Editor/ContentBrowser/Private/SContentBrowser.h b/Engine/Source/Editor/ContentBrowser/Private/SContentBrowser.h
--- a/Engine/Source/Editor/ContentBrowser/Private/SContentBrowser.h
+++ b/Engine/Source/Editor/ContentBrowser/Private/SContentBrowser.h
@@ -146,6 +146,13 @@ public:
 	 */
 	void SyncToFolders( TArrayView<const FString> FolderList, const bool bAllowImplicitSync = false );
 
+	/**
+	 * Changes sources to show the specified collections and selects them in the collection view
+	 *
+	 *	@param CollectionsToSync	- A list of collections to sync the view to
+	 */
+	void SyncToCollections(TArrayView<const FCollectionRef> CollectionsToSync);
+
 	/**
 	 * Changes sources to show the specified items and selects them in the asset view
 	 *
diff --git a/Engine/Source/Editor/ContentBrowser/Public/IContentBrowserSingleton.h b/Engine/Source/Editor/ContentBrowser/Public/IContentBrowserSingleton.h
--- a/Engine/Source/Editor/ContentBrowser/Public/IContentBrowserSingleton.h
+++ b/Engine/Source/Editor/ContentBrowser/Public/IContentBrowserSingleton.h
@@ -695,8 +695,8 @@ public:
 	/** 
 	 * Selects the supplied assets in the primary content browser. 
      *
-	 * @param AssetList                     An array of AssetDataList structs to sync
-	 * @param bAllowLockedBrowsers 	        When true, even locked browsers may handle the sync. Only set to true if the sync doesn't seem external to the content browser
+	 * @param AssetList                     An array of AssetDataList structs to sync.
+	 * @param bAllowLockedBrowsers 	        When true, even locked browsers may handle the sync. Only set to true if the sync doesn't seem external to the content browser.
 	 * @param bFocusContentBrowser          When true, brings the ContentBrowser into the foreground.
 	 * @param InstanceName					When supplied, will only sync the Content Browser with the matching InstanceName.  bAllowLockedBrowsers is ignored.
 	 */ 
@@ -705,8 +705,8 @@ public:
 	/** 
 	 * Selects the supplied assets in the primary content browser. 
      *
-	 * @param AssetList                     An array of UObject pointers to sync
-	 * @param bAllowLockedBrowsers 	        When true, even locked browsers may handle the sync. Only set to true if the sync doesn't seem external to the content browser
+	 * @param AssetList                     An array of UObject pointers to sync.
+	 * @param bAllowLockedBrowsers 	        When true, even locked browsers may handle the sync. Only set to true if the sync doesn't seem external to the content browser.
 	 * @param bFocusContentBrowser          When true, brings the ContentBrowser into the foreground.
 	 * @param InstanceName					When supplied, will only sync the Content Browser with the matching InstanceName.  bAllowLockedBrowsers is ignored.
 	 */ 
@@ -715,19 +715,30 @@ public:
 	/** 
 	 * Selects the supplied assets in the primary content browser. 
      *
-	 * @param AssetList                     An array of strings representing folder asset paths to sync
-	 * @param bAllowLockedBrowsers 	        When true, even locked browsers may handle the sync. Only set to true if the sync doesn't seem external to the content browser
+	 * @param AssetList                     An array of strings representing folder asset paths to sync.
+	 * @param bAllowLockedBrowsers 	        When true, even locked browsers may handle the sync. Only set to true if the sync doesn't seem external to the content browser.
 	 * @param bFocusContentBrowser          When true, brings the ContentBrowser into the foreground.
 	 * @param InstanceName					When supplied, will only sync the Content Browser with the matching InstanceName.  bAllowLockedBrowsers is ignored.
 	 * @param bNewSpawnBrowser				When supplied, will spawn a new Content Browser instead of selecting the assets in an existing one.
 	 */ 
 	virtual void SyncBrowserToFolders(const TArray<FString>& FolderList, bool bAllowLockedBrowsers = false, bool bFocusContentBrowser = true, const FName& InstanceName = FName(), bool bNewSpawnBrowser = false) = 0;
 
+	/**
+	 * Selects the supplied collections in the primary content browser.
+	 *
+	 * @param CollectionsToSync				An array of collections to sync.
+	 * @param bAllowLockedBrowsers 	        When true, even locked browsers may handle the sync. Only set to true if the sync doesn't seem external to the content browser.
+	 * @param bFocusContentBrowser          When true, brings the ContentBrowser into the foreground.
+	 * @param InstanceName					When supplied, will only sync the Content Browser with the matching InstanceName.  bAllowLockedBrowsers is ignored.
+	 * @param bNewSpawnBrowser				When supplied, will spawn a new Content Browser instead of selecting the assets in an existing one.
+	 */
+	virtual void SyncBrowserToCollections(const TArray<FCollectionRef>& CollectionsToSync, bool bAllowLockedBrowsers = false, bool bFocusContentBrowser = true, const FName& InstanceName = FName(), bool bNewSpawnBrowser = false) = 0;
+
 	/**
 	 * Selects the supplied items in the primary content browser.
 	 *
-	 * @param ItemsToSync                   An array of items to sync
-	 * @param bAllowLockedBrowsers 	        When true, even locked browsers may handle the sync. Only set to true if the sync doesn't seem external to the content browser
+	 * @param ItemsToSync                   An array of items to sync.
+	 * @param bAllowLockedBrowsers 	        When true, even locked browsers may handle the sync. Only set to true if the sync doesn't seem external to the content browser.
 	 * @param bFocusContentBrowser          When true, brings the ContentBrowser into the foreground.
 	 * @param InstanceName					When supplied, will only sync the Content Browser with the matching InstanceName.  bAllowLockedBrowsers is ignored.
 	 * @param bNewSpawnBrowser				When supplied, will spawn a new Content Browser instead of selecting the assets in an existing one.
@@ -737,8 +748,8 @@ public:
 	/** 
 	 * Selects the supplied assets in the primary content browser. 
      *
-	 * @param ItemSelection 				A struct containing AssetData and Folders to sync
-	 * @param bAllowLockedBrowsers 	        When true, even locked browsers may handle the sync. Only set to true if the sync doesn't seem external to the content browser
+	 * @param ItemSelection 				A struct containing AssetData and Folders to sync.
+	 * @param bAllowLockedBrowsers 	        When true, even locked browsers may handle the sync. Only set to true if the sync doesn't seem external to the content browser.
 	 * @param bFocusContentBrowser          When true, brings the ContentBrowser into the foreground.
 	 * @param InstanceName					When supplied, will only sync the Content Browser with the matching InstanceName.  bAllowLockedBrowsers is ignored.
 	 * @param bNewSpawnBrowser				When supplied, will spawn a new Content Browser instead of selecting the assets in an existing one.
diff --git a/Engine/Source/Editor/UnrealEd/Private/Subsystems/CollectionManagerScriptingSubsystem.cpp b/Engine/Source/Editor/UnrealEd/Private/Subsystems/CollectionManagerScriptingSubsystem.cpp
--- a/Engine/Source/Editor/UnrealEd/Private/Subsystems/CollectionManagerScriptingSubsystem.cpp
+++ b/Engine/Source/Editor/UnrealEd/Private/Subsystems/CollectionManagerScriptingSubsystem.cpp
@@ -5,9 +5,11 @@
 #include "AssetRegistry/AssetData.h"
 #include "AssetRegistry/AssetRegistryModule.h"
 #include "CollectionManagerModule.h"
+#include "ContentBrowserModule.h"
 #include "ICollectionContainer.h"
 #include "ICollectionManager.h"
 #include "ICollectionSource.h"
+#include "IContentBrowserSingleton.h"
 
 #include UE_INLINE_GENERATED_CPP_BY_NAME(CollectionManagerScriptingSubsystem)
 
@@ -593,4 +595,37 @@ FCollectionScriptingContainerSource UCollectionManagerScriptingSubsystem::GetBas
 	return Container;
 }
 
+void UCollectionManagerScriptingSubsystem::SyncBrowserToCollections(const TArray<FCollectionScriptingRef>& Collections)
+{
+	TGuardValue<bool> UnattendedScriptGuard(GIsRunningUnattendedScript, true);
+
+	TArray<FCollectionRef> CollectionsToSync;
+	CollectionsToSync.Reserve(Collections.Num());
+	for (const FCollectionScriptingRef& Collection : Collections)
+	{
+		const TSharedPtr<ICollectionContainer> CollectionContainer = FCollectionManagerModule::GetModule().Get().FindCollectionContainer(Collection.Container);
+		if (!CollectionContainer.IsValid())
+		{
+			CollectionManagerScriptingSubsystemUtil::LogLastCollectionManagerWarning(FText::Format(CollectionManagerScriptingSubsystemUtil::FConstants::Get().ContainerNotFound, { FText::FromName(Collection.Container) }));
+			continue;
+		}
+
+		const ECollectionShareType::Type ShareType = CollectionManagerScriptingSubsystemUtil::ECollectionScriptingShareType_To_ECollectionShareType(Collection.ShareType);
+		if (!CollectionContainer->CollectionExists(Collection.Name, ShareType))
+		{
+			CollectionManagerScriptingSubsystemUtil::LogLastCollectionManagerWarning(FText::Format(CollectionManagerScriptingSubsystemUtil::FConstants::Get().CollectionNotFound,
+				{ FText::FromName(Collection.Name), FText::FromString(CollectionManagerScriptingSubsystemUtil::ECollectionScriptingShareType_To_String(Collection.ShareType)), FText::FromName(Collection.Container) }));
+			continue;
+		}
+
+		CollectionsToSync.Emplace(CollectionContainer, Collection.Name, ShareType);
+	}
+
+	if (!CollectionsToSync.IsEmpty())
+	{
+		FContentBrowserModule& ContentBrowserModule = FModuleManager::Get().LoadModuleChecked<FContentBrowserModule>("ContentBrowser");
+		ContentBrowserModule.Get().SyncBrowserToCollections(CollectionsToSync);
+	}
+}
+
 #undef LOCTEXT_NAMESPACE
diff --git a/Engine/Source/Editor/UnrealEd/Public/Subsystems/CollectionManagerScriptingSubsystem.h b/Engine/Source/Editor/UnrealEd/Public/Subsystems/CollectionManagerScriptingSubsystem.h
--- a/Engine/Source/Editor/UnrealEd/Public/Subsystems/CollectionManagerScriptingSubsystem.h
+++ b/Engine/Source/Editor/UnrealEd/Public/Subsystems/CollectionManagerScriptingSubsystem.h
@@ -379,6 +379,15 @@ public:
 	UFUNCTION(BlueprintCallable, Category="CollectionManager")
 	UE_API FCollectionScriptingContainerSource GetBaseGameCollectionContainer() const;
 
+	/**
+	 * Selects the collections in the most recently used Content Browser (summoning one if necessary)
+	 * This is an asynchronous operation that can take a couple of frames to resolve the request
+	 * 
+	 * @param Collections The list of collections to sync to in the Content Browser
+	 */
+	UFUNCTION(BlueprintCallable, Category = "CollectionManager | Content Browser")
+	UE_API void SyncBrowserToCollections(const TArray<FCollectionScriptingRef>& Collections);
+
 private:
 	bool SetCollectionColor(const FCollectionScriptingRef& Collection, const TOptional<FLinearColor>& Color);
 };
