Content Browser - Folder - Favorites test added

#rnx

diff --git a/Engine/Source/Editor/ContentBrowser/Private/SContentBrowser.cpp b/Engine/Source/Editor/ContentBrowser/Private/SContentBrowser.cpp
--- a/Engine/Source/Editor/ContentBrowser/Private/SContentBrowser.cpp
+++ b/Engine/Source/Editor/ContentBrowser/Private/SContentBrowser.cpp
@@ -3402,7 +3402,8 @@ TSharedRef<SWidget> SContentBrowser::CreateLegacyAssetViewWidgets()
 				SNew(SActionButton)
 				.ActionButtonType(EActionButtonType::Simple)
 				.OnGetMenuContent(LegacyContentSourceWidgets->AssetViewPtr.ToSharedRef(), &SAssetView::GetViewButtonContent)
-				.Icon(FAppStyle::Get().GetBrush("Icons.Settings"));
+				.Icon(FAppStyle::Get().GetBrush("Icons.Settings"))
+				.AddMetaData<FTagMetaData>(FTagMetaData(TEXT("AssetViewMenuButton")));
 		}
 		else
 		{
@@ -3411,6 +3412,7 @@ TSharedRef<SWidget> SContentBrowser::CreateLegacyAssetViewWidgets()
 				.ComboButtonStyle(&FAppStyle::Get().GetWidgetStyle<FComboButtonStyle>("SimpleComboButton"))
 				.OnGetMenuContent(LegacyContentSourceWidgets->AssetViewPtr.ToSharedRef(), &SAssetView::GetViewButtonContent)
 				.HasDownArrow(false)
+				.AddMetaData<FTagMetaData>(FTagMetaData(TEXT("AssetViewMenuButton")))
 				.ButtonContent()
 				[
 					SNew(SHorizontalBox)
diff --git a/Engine/Source/Editor/ContentBrowser/Public/SAssetView.h b/Engine/Source/Editor/ContentBrowser/Public/SAssetView.h
--- a/Engine/Source/Editor/ContentBrowser/Public/SAssetView.h
+++ b/Engine/Source/Editor/ContentBrowser/Public/SAssetView.h
@@ -533,6 +533,9 @@ public:
 		const TOptional<const FName>& InColumnId,
 		const TOptional<const EColumnSortMode::Type>& InNewSortMode);
 
+	/** @return true when we are showing favorites */
+	UE_INTERNAL UE_API bool IsShowingFavorites() const;
+
 private:
 	/** Sets the pending selection to the current selection (used when changing views or refreshing the view). */
 	UE_API void SyncToSelection( const bool bFocusOnSync = true );
@@ -706,9 +709,6 @@ private:
 	/** Whether or not it's possible to toggle favorites */
 	UE_API bool IsToggleShowFavoritesAllowed() const;
 
-	/** @return true when we are showing favorites */
-	UE_API bool IsShowingFavorites() const;
-
 	/** Whether or not it's possible to show C++ content */
 	UE_API bool IsToggleShowCppContentAllowed() const;
 
