Added missing implementation for ~IAssetTextFilterHandler() destructor causing linker errors when a derived class needs to be destructed.

diff --git a/Engine/Source/Editor/ContentBrowser/Public/AssetTextFilter.h b/Engine/Source/Editor/ContentBrowser/Public/AssetTextFilter.h
--- a/Engine/Source/Editor/ContentBrowser/Public/AssetTextFilter.h
+++ b/Engine/Source/Editor/ContentBrowser/Public/AssetTextFilter.h
@@ -241,7 +241,7 @@ private:
 class IAssetTextFilterHandler
 {
 public:
-	CONTENTBROWSER_API virtual ~IAssetTextFilterHandler();
+	CONTENTBROWSER_API virtual ~IAssetTextFilterHandler() = default;
 
 	/**
 	 * Implement this function to handle basic text matching.
