EditorUsability : ContentBrowser
- Added an option to sort the content browser filters in alphabetical order or insertion order
- Refactored loading filter code to restore filter in the correct order instead of the random type registration order
- Reordered all context menu in content browser to be alphabetically sorted (Other filters, Custom filters, Common filters)

diff --git a/Engine/Source/Developer/ToolWidgets/Public/Filters/SBasicFilterBar.h b/Engine/Source/Developer/ToolWidgets/Public/Filters/SBasicFilterBar.h
--- a/Engine/Source/Developer/ToolWidgets/Public/Filters/SBasicFilterBar.h
+++ b/Engine/Source/Developer/ToolWidgets/Public/Filters/SBasicFilterBar.h
@@ -84,6 +84,9 @@ template<typename FilterType>
 class SBasicFilterBar : public SCompoundWidget
 {
 public:
+	/** Delegate for when layout changes */
+	DECLARE_DELEGATE_OneParam(FOnFilterBarLayoutChanging, EFilterBarLayout /* NewLayout */)
+
 	/** Delegate for when filters have changed */
 	DECLARE_DELEGATE( FOnFilterChanged );
 	
@@ -91,11 +94,12 @@ public:
     DECLARE_DELEGATE_RetVal( TSharedPtr<ICustomTextFilter<FilterType>>, FCreateTextFilter);
 	
  	SLATE_BEGIN_ARGS( SBasicFilterBar<FilterType> )
- 		: _FilterBarLayout(EFilterBarLayout::Horizontal)
+		: _FilterBarLayout(EFilterBarLayout::Horizontal)
 		, _CanChangeOrientation(false)
 		, _FilterPillStyle(EFilterPillStyle::Default)
 		, _UseSectionsForCategories(false)
 		, _bPinAllFrontendFilters(false)
+		, _bSortFilters(true)
 	{}
 
  		/** Delegate for when filters have changed */
@@ -119,6 +123,9 @@ public:
 		 */
         SLATE_ARGUMENT(TSharedPtr<SFilterSearchBox>, FilterSearchBox)
 
+		/** Called when the filter bar layout is changing, before the filters are added to the layout */
+		SLATE_EVENT(FOnFilterBarLayoutChanging, OnFilterBarLayoutChanging)
+
 		/** The layout that determines how the filters are laid out */
 		SLATE_ARGUMENT(EFilterBarLayout, FilterBarLayout)
 	
@@ -133,6 +140,9 @@ public:
 
 		/** Whether we want to automatically pin all frontend filters to make them visible by default. */
 		SLATE_ARGUMENT(bool, bPinAllFrontendFilters)
+
+		/** Whether filters are sorted in the bar or appear in the order they are added */
+		SLATE_ARGUMENT(bool, bSortFilters)
 	
  	SLATE_END_ARGS()
 
@@ -141,11 +151,13 @@ public:
 		OnFilterChanged = InArgs._OnFilterChanged;
  		OnExtendAddFilterMenu = InArgs._OnExtendAddFilterMenu;
  		CreateTextFilter = InArgs._CreateTextFilter;
+ 		OnFilterBarLayoutChanging = InArgs._OnFilterBarLayoutChanging;
  		FilterBarLayout = InArgs._FilterBarLayout;
  		bCanChangeOrientation = InArgs._CanChangeOrientation;
  		FilterPillStyle = InArgs._FilterPillStyle;
-		bUseSectionsForCategories = InArgs._UseSectionsForCategories;
- 		bPinAllFrontendFilters = InArgs._bPinAllFrontendFilters;
+        bUseSectionsForCategories = InArgs._UseSectionsForCategories;
+        bPinAllFrontendFilters = InArgs._bPinAllFrontendFilters;
+        bSortFiltersAlphabetically = InArgs._bSortFilters;
 
  		// We use a widgetswitcher to allow dynamically swapping between the layouts
  		ChildSlot
@@ -228,27 +240,23 @@ public:
  		}
  		
  		FilterBarLayout = InFilterBarLayout;
+
+ 		HorizontalFilterBox->ClearChildren();
+ 		VerticalFilterBox->ClearChildren();
  		
  		if(FilterBarLayout == EFilterBarLayout::Horizontal)
  		{
-			VerticalFilterBox->ClearChildren();
- 			
  			FilterBox->SetActiveWidget(HorizontalFilterBox.ToSharedRef());
  		}
         else
         {
-        	HorizontalFilterBox->ClearChildren();
-        	
         	FilterBox->SetActiveWidget(VerticalFilterBox.ToSharedRef());
         }
 
- 		// Add all the filters to the new layout
- 		for(TSharedRef<SFilter> Filter: Filters)
- 		{
- 			AddWidgetToLayout(Filter);
- 		}
+ 		OnFilterBarLayoutChanging.ExecuteIfBound(FilterBarLayout);
 
- 		this->Invalidate(EInvalidateWidgetReason::Layout);
+ 		// Filters needs to be re-added in the new active widget container
+ 		RequestFiltersUpdate();
  	}
 
 	/** Makes the button that summons the Add Filter dropdown on click */
@@ -884,6 +892,22 @@ public:
  		return false;
  	}
 
+	/** Are we sorting filters alphabetically */
+    bool GetSortFiltersAlphabetically() const
+    {
+ 		return bSortFiltersAlphabetically;
+    }
+
+	/** Set filtering method for filters */
+	void SetSortFiltersAlphabetically(bool bInSort)
+    {
+        if (bSortFiltersAlphabetically != bInSort)
+        {
+        	bSortFiltersAlphabetically = bInSort;
+        	RequestFiltersUpdate();
+        }
+    }
+
 protected:
 
 	/** Overlay for filter icon, badge if there are any active filters */
@@ -1002,8 +1026,57 @@ protected:
 	void AddFilterToBar(const TSharedRef<SFilter>& FilterToAdd)
 	{
 		Filters.Add(FilterToAdd);
-	
-		AddWidgetToLayout(FilterToAdd);
+
+		if (bSortFiltersAlphabetically)
+		{
+			// Update next tick in case this is called in a loop
+			RequestFiltersUpdate();
+		}
+		else
+		{
+			// Add new filter to the end directly
+			AddWidgetToLayout(FilterToAdd);
+		}
+	}
+
+	/** Updates filters in the layout if requested */
+	virtual void Tick(const FGeometry& InGeometry, const double InCurrentTime, const float InDeltaTime) override
+	{
+		if (bRequestFiltersUpdate)
+		{
+			bRequestFiltersUpdate = false;
+
+			HorizontalFilterBox->ClearChildren();
+			VerticalFilterBox->ClearChildren();
+
+			auto AddFiltersToLayout = [this](const TArray<TSharedRef<SFilter>>& InFilters)
+				{
+					for (const TSharedRef<SFilter>& Filter : InFilters)
+					{
+						AddWidgetToLayout(Filter);
+					}
+				};
+
+			if (bSortFiltersAlphabetically)
+			{
+				TArray<TSharedRef<SFilter>> FiltersSorted = Filters;
+
+				FiltersSorted.Sort(
+					[](const TSharedRef<SFilter>& A, const TSharedRef<SFilter>& B)
+					{
+						return A->GetFilterDisplayName().CompareTo(B->GetFilterDisplayName()) < 0;
+					}
+				);
+
+				AddFiltersToLayout(FiltersSorted);
+			}
+			else
+			{
+				AddFiltersToLayout(Filters);
+			}
+		}
+
+		SCompoundWidget::Tick(InGeometry, InCurrentTime, InDeltaTime);
 	}
 
 	/** Handler for when the enable only this button was clicked on a single filter */
@@ -1598,6 +1671,12 @@ protected:
 		}
 	}
 
+	/** Request a filter list update next tick */
+	void RequestFiltersUpdate()
+	{
+		bRequestFiltersUpdate = true;
+	}
+
 private:
 
 	/** Handler for when the add filter button was clicked */
@@ -1635,6 +1714,9 @@ private:
 		
 		PopulateCustomFilters(Menu);
 	}
+
+	/** Whether a refresh is requested next tick to re-render filters in the bar */
+	bool bRequestFiltersUpdate = false;
 	
 protected:
 
@@ -1649,6 +1731,9 @@ protected:
 
 	/** The current Filter Layout being used */
 	EFilterBarLayout FilterBarLayout;
+
+	/** Delegate for when the layout is changed */
+	FOnFilterBarLayoutChanging OnFilterBarLayoutChanging;
 	
  	/** All SFilters in the list */
  	TArray< TSharedRef<SFilter> > Filters;
@@ -1688,6 +1773,9 @@ protected:
 
 	/** Determines how each individual pill looks like */
 	EFilterPillStyle FilterPillStyle;
+
+	/** Determines whether filters are sorted in the bar or appear in the order they are added */
+	bool bSortFiltersAlphabetically;
 	
 	friend struct FFrontendFilterExternalActivationHelper<FilterType>;
 };
diff --git a/Engine/Source/Editor/ContentBrowser/Private/SFilterList.cpp b/Engine/Source/Editor/ContentBrowser/Private/SFilterList.cpp
--- a/Engine/Source/Editor/ContentBrowser/Private/SFilterList.cpp
+++ b/Engine/Source/Editor/ContentBrowser/Private/SFilterList.cpp
@@ -68,7 +68,6 @@ SFilterList::FCustomTextFilterEvent SFilterList::CustomTextFilterEvent;
 void SFilterList::Construct( const FArguments& InArgs )
 {
 	bUseSharedSettings = InArgs._UseSharedSettings;
-	OnFilterBarLayoutChanging = InArgs._OnFilterBarLayoutChanging;
 	this->OnFilterChanged = InArgs._OnFilterChanged;
 	this->ActiveFilters = InArgs._FrontendFilters;
 	InitialClassFilters = InArgs._InitialClassFilters;
@@ -147,6 +146,7 @@ void SFilterList::Construct( const FArguments& InArgs )
 	Args._OnFilterChanged = this->OnFilterChanged;
 	Args._CreateTextFilter = InArgs._CreateTextFilter;
 	Args._FilterBarIdentifier = InArgs._FilterBarIdentifier;
+	Args._OnFilterBarLayoutChanging = InArgs._OnFilterBarLayoutChanging;
 	Args._FilterBarLayout = InArgs._FilterBarLayout;
 	Args._CanChangeOrientation = InArgs._CanChangeOrientation;
 	Args._OnExtendAddFilterMenu = InArgs._OnExtendAddFilterMenu;
@@ -477,6 +477,18 @@ void SFilterList::PopulateAddFilterMenu(UToolMenu* Menu, TSharedPtr<FFilterCateg
 					ManageSection->AddMenuEntry(ContentBrowserCommands.FilterBySelection);
 				}
 			}
+
+			ManageSection->AddMenuEntry(
+				"FilterAlphabetically",
+				LOCTEXT("FilterAlphabetically", "Sort Filters Alphabetically"),
+				LOCTEXT("FilterAlphabeticallyToolTip", "Active filters will be sorted and displayed in alphabetical order in the layout instead of their insertion order"),
+				FSlateIcon(),
+				FUIAction(
+					FExecuteAction::CreateSP(this, &SFilterList::ToggleSortFiltersAlphabetically),
+					FCanExecuteAction(),
+					FIsActionChecked::CreateSP(this, &SFilterList::GetSortFiltersAlphabetically)),
+				EUserInterfaceActionType::ToggleButton
+			);
 		}
 
 		FToolMenuSection& DisplaySection = Menu->AddSection(
@@ -528,6 +540,47 @@ void SFilterList::PopulateAddFilterMenu(UToolMenu* Menu, TSharedPtr<FFilterCateg
 		}
 	}
 
+	// Sort by category name so that we add the submenus in alphabetical order
+	CategoryToMenuMap.KeySort(
+		[](const TSharedPtr<FFilterCategory>& A, const TSharedPtr<FFilterCategory>& B)
+		{
+			return A.IsValid()
+				&& B.IsValid()
+				&& A->Title.CompareTo(B->Title) < 0;
+		}
+	);
+
+	// Sort each submenu internals 
+	for (TPair<TSharedPtr<FFilterCategory>, FFilterCategoryMenu>& CategoryToMenu : CategoryToMenuMap)
+	{
+		CategoryToMenu.Value.Classes.Sort(
+			[](const TSharedPtr<FCustomClassFilterData>& A, const TSharedPtr<FCustomClassFilterData>& B)
+			{
+				return A.IsValid()
+					&& B.IsValid()
+					&& A->GetName().CompareTo(B->GetName()) < 0;
+			}
+		);
+	}
+
+	// Sort frontend filters
+	AllFrontendFilters.Sort(
+		[](const TSharedRef<FFilterBase<FAssetFilterType>>& A, const TSharedRef<FFilterBase<FAssetFilterType>>& B)
+		{
+			return A->GetDisplayName().CompareTo(B->GetDisplayName()) < 0;
+		}
+	);
+
+	// Sort user filters
+	CustomTextFilters.Sort(
+		[](const TSharedRef<ICustomTextFilter<FAssetFilterType>>& A, const TSharedRef<ICustomTextFilter<FAssetFilterType>>& B)
+		{
+			return A->GetFilter().IsValid()
+				&& B->GetFilter().IsValid()
+				&& A->GetFilter()->GetDisplayName().CompareTo(B->GetFilter()->GetDisplayName()) < 0;
+		}
+	);
+
 	// If we want to expand a category
 	if (MenuExpansion)
 	{
@@ -559,20 +612,6 @@ void SFilterList::PopulateAddFilterMenu(UToolMenu* Menu, TSharedPtr<FFilterCateg
 		}
 	}
 
-	// Sort by category name so that we add the submenus in alphabetical order
-	CategoryToMenuMap.KeySort([](const TSharedPtr<FFilterCategory>& A, const TSharedPtr<FFilterCategory>& B) {
-		return A->Title.CompareTo(B->Title) < 0;
-	});
-
-	// Sort each submenu internals 
-	for (TPair<TSharedPtr<FFilterCategory>, FFilterCategoryMenu>& CategoryToMenu : CategoryToMenuMap)
-	{
-		CategoryToMenu.Value.Classes.Sort([](const TSharedPtr<FCustomClassFilterData>& A, const TSharedPtr<FCustomClassFilterData>& B)
-		{
-			return A->GetName().CompareTo(B->GetName()) < 0;
-		});
-	}
-
 	FToolMenuSection& AdvancedSection = Menu->FindOrAddSection("FilterBarAdvanced");
 
 	FToolMenuEntry& AllFiltersSubMenu = AdvancedSection.AddSubMenu(
@@ -670,13 +709,10 @@ void SFilterList::ToggleFilterPathsRecursively()
 	SetFilterPathsRecursively(!bFilterPathsRecursively);
 }
 
-void SFilterList::SetActiveFilters(TConstArrayView<FTopLevelAssetPath> InFilters)
+void SFilterList::ToggleSortFiltersAlphabetically()
 {
-	RemoveAllFilters();
-	for (const FTopLevelAssetPath& FilterAssetPath : InFilters)
-	{
-		SetAssetTypeFilterCheckState(FilterAssetPath, ECheckBoxState::Checked, /** Enabled */true);
-	}
+	bSortFiltersAlphabetically = !bSortFiltersAlphabetically;
+	RequestFiltersUpdate();
 }
 
 TSharedPtr<SWidget> SFilterList::GetActiveFilterContainer() const
@@ -838,6 +874,7 @@ void SFilterList::SaveSettingsInternal(FFilterBarSettings* InSettingsToSave)
 				InSettingsToSave->bFilterRecursively = this->bFilterPathsRecursively;
 			}
 		}
+		InSettingsToSave->bSortFiltersAlphabetically = this->bSortFiltersAlphabetically;
 
 		SAssetFilterBar<FAssetFilterType>::SaveSettingsInternal(InSettingsToSave);
 		return;
@@ -907,6 +944,7 @@ void SFilterList::SaveSettingsInternal(FFilterBarSettings* InSettingsToSave)
 	{
 		InstanceSettings->FilterBarLayout = this->FilterBarLayout;
 	}
+	InSettingsToSave->bSortFiltersAlphabetically = this->bSortFiltersAlphabetically;
 
 	SaveConfig();
 }
@@ -967,6 +1005,7 @@ void SFilterList::LoadSettings(const FName& InInstanceName)
 			{
 				this->SetFilterPathsRecursively(FilterBarConfig->bFilterRecursively);
 			}
+			this->SetSortFiltersAlphabetically(FilterBarConfig->bSortFiltersAlphabetically);
 		}
 
 		SAssetFilterBar<FAssetFilterType>::LoadSettings();
@@ -1002,6 +1041,7 @@ void SFilterList::LoadSettings(const FName& InInstanceName)
 	if(InstanceSettings->bIsLayoutSaved)
 	{
 		FilterBarLayout = InstanceSettings->FilterBarLayout;
+		SetSortFiltersAlphabetically(InstanceSettings->bSortFiltersAlphabetically);
 	}
 
 	// We want to call this even if the Layout isn't saved, to make sure OnFilterBarLayoutChanging is fired
@@ -1054,39 +1094,6 @@ void SFilterList::AddWidgetToCurrentLayout(TSharedRef<SWidget> InWidget)
 	}
 }
 
-void SFilterList::SetFilterLayout(EFilterBarLayout InFilterBarLayout)
-{
-	FilterBarLayout = InFilterBarLayout;
-
-	/* Clear both layouts, because for SFilterList it is valid to call SetFilterLayout with InFilterBarLayout being the
-	 * same as the current layout just to fire OnFilterBarLayoutChanging.
-	 * Unlike the parent class SBasicFilterBar which guards against that. If we don't clear both child widgets you can
-	 * end up with duplicate widgets.
-	 */
-	HorizontalFilterBox->ClearChildren();
-	VerticalFilterBox->ClearChildren();
- 		
-	if(FilterBarLayout == EFilterBarLayout::Horizontal)
-	{
-		FilterBox->SetActiveWidget(HorizontalFilterBox.ToSharedRef());
-	}
-	else
-	{
-		FilterBox->SetActiveWidget(VerticalFilterBox.ToSharedRef());
-	}
-
-	OnFilterBarLayoutChanging.ExecuteIfBound(FilterBarLayout);
-
-	for(TSharedRef<SFilter> Filter: Filters)
-	{
-		AddWidgetToLayout(Filter);
-	}
- 		
-	this->Invalidate(EInvalidateWidgetReason::Layout);
-
- 		
-}
-
 /////////////////////////////////////////
 // FFilterListCustomTextFilter
 /////////////////////////////////////////
diff --git a/Engine/Source/Editor/ContentBrowser/Private/SFilterList.h b/Engine/Source/Editor/ContentBrowser/Private/SFilterList.h
--- a/Engine/Source/Editor/ContentBrowser/Private/SFilterList.h
+++ b/Engine/Source/Editor/ContentBrowser/Private/SFilterList.h
@@ -42,7 +42,6 @@ class SFilterList : public SAssetFilterBar<FAssetFilterType>
 {
 	using Super = SAssetFilterBar<FAssetFilterType>; 
 public:
-	DECLARE_DELEGATE_OneParam(FOnFilterBarLayoutChanging, EFilterBarLayout /* NewLayout */)
 	/**
 	 * An event delegate that is executed when a custom text filter has been created/modified/deleted in any FilterList
 	 * that is using shared settings.
@@ -159,9 +158,6 @@ public:
 	void SaveSettings(const FString& IniFilename, const FString& IniSection, const FString& SettingsString);
 	void LoadSettings(const FName& InInstanceName, const FString& IniFilename, const FString& IniSection, const FString& SettingsString);
 
-
-	virtual void SetFilterLayout(EFilterBarLayout InFilterBarLayout) override;
-
 	/** Handler for when a custom text filter is created */
 	virtual void OnCreateCustomTextFilter(const FCustomTextFilterData& InFilterData, bool bApplyFilter) override;
 
@@ -180,8 +176,8 @@ public:
 	/** Toggle whether the filters are being applied to paths recursively */
 	void ToggleFilterPathsRecursively();
 
-	/** Set the active asset filters for this filter list */
-	void SetActiveFilters(TConstArrayView<FTopLevelAssetPath> InFilters);
+	/** Toggle whether the filters are sorted and displayed alphabetically */
+	void ToggleSortFiltersAlphabetically();
 
 	/** Returns a pointer to the active filter container widget (Horizontal or Vertical), nullptr if neither is active */
 	TSharedPtr<SWidget> GetActiveFilterContainer() const;
@@ -227,9 +223,6 @@ private:
 	/** Delegate for when filters have changed */
 	FOnFilterChanged OnFilterChanged;
 
-	/** Delegate for when the layout is changed */
-	FOnFilterBarLayoutChanging OnFilterBarLayoutChanging;
-
 	/** A reference to AllFrontEndFilters so we can access the filters as FFrontEndFilter instead of FFilterBase<FAssetFilterType> */
 	TArray< TSharedRef<FFrontendFilter> > AllFrontendFilters_Internal;
 
diff --git a/Engine/Source/Editor/EditorWidgets/Public/Filters/FilterBarConfig.h b/Engine/Source/Editor/EditorWidgets/Public/Filters/FilterBarConfig.h
--- a/Engine/Source/Editor/EditorWidgets/Public/Filters/FilterBarConfig.h
+++ b/Engine/Source/Editor/EditorWidgets/Public/Filters/FilterBarConfig.h
@@ -69,6 +69,10 @@ public:
 	UPROPERTY()
 	bool bFilterRecursively = true;
 
+	/** Whether filters are sorted and displayed alphabetically or in the order they were added */
+	UPROPERTY()
+	bool bSortFiltersAlphabetically = true;
+
 	void Empty()
 	{
 		CustomFilters.Empty();
diff --git a/Engine/Source/Editor/EditorWidgets/Public/Filters/SAssetFilterBar.h b/Engine/Source/Editor/EditorWidgets/Public/Filters/SAssetFilterBar.h
--- a/Engine/Source/Editor/EditorWidgets/Public/Filters/SAssetFilterBar.h
+++ b/Engine/Source/Editor/EditorWidgets/Public/Filters/SAssetFilterBar.h
@@ -218,6 +218,9 @@ public:
 	
 		/** Whether the filter bar should provide the default asset filters */
 		SLATE_ARGUMENT(bool, UseDefaultAssetFilters)
+
+		/** Called when the filter bar layout is changing, before the filters are added to the layout */
+		SLATE_EVENT(SBasicFilterBar<FilterType>::FOnFilterBarLayoutChanging, OnFilterBarLayoutChanging)
 		
 		/** The layout that determines how the filters are laid out */
 		SLATE_ARGUMENT(EFilterBarLayout, FilterBarLayout)
@@ -258,6 +261,7 @@ public:
 		Args._OnExtendAddFilterMenu = InArgs._OnExtendAddFilterMenu;
 		Args._CreateTextFilter = InArgs._CreateTextFilter;
 		Args._FilterSearchBox = InArgs._FilterSearchBox;
+		Args._OnFilterBarLayoutChanging = InArgs._OnFilterBarLayoutChanging;
 		Args._FilterBarLayout = InArgs._FilterBarLayout;
 		Args._CanChangeOrientation = InArgs._CanChangeOrientation;
 		Args._FilterPillStyle = InArgs._FilterPillStyle;
@@ -444,28 +448,38 @@ protected:
 	void LoadFilters(const FFilterBarSettings* FilterBarConfig)
 	{
 		// Load all the custom filters (i.e FrontendFilters)
-		for ( auto FrontendFilterIt = this->AllFrontendFilters.CreateIterator(); FrontendFilterIt; ++FrontendFilterIt )
+		TMap<FString, TSharedRef<FFilterBase<FilterType>>> CustomFilters;
+		CustomFilters.Reserve(FilterBarConfig->CustomFilters.Num());
+		for (const TSharedRef<FFilterBase<FilterType>>& FrontendFilter : this->AllFrontendFilters)
 		{
-			TSharedRef<FFilterBase<FilterType>>& FrontendFilter = *FrontendFilterIt;
-			const FString& FilterName = FrontendFilter->GetName();
-			
 			if (!this->IsFrontendFilterInUse(FrontendFilter))
 			{
 				// Try to find this type filter in our list of saved filters
-				if ( const bool* bIsActive = FilterBarConfig->CustomFilters.Find(FilterName) )
+				const FString& FilterName = FrontendFilter->GetName();
+				if (FilterBarConfig->CustomFilters.Contains(FilterName))
 				{
-					TSharedRef<SFilter> NewFilter = this->AddFilterToBar(FrontendFilter);
-					
-					NewFilter->SetEnabled(*bIsActive, false);
-					this->SetFrontendFilterActive(FrontendFilter, NewFilter->IsEnabled());
+					CustomFilters.Add(FilterName, FrontendFilter);
 				}
 			}
 		}
 
+		// We add them here to preserve the saving config order
+		for (const TPair<FString, bool>& CustomFilter : FilterBarConfig->CustomFilters)
+		{
+			if (const TSharedRef<FFilterBase<FilterType>>* Filter = CustomFilters.Find(CustomFilter.Key))
+			{
+				TSharedRef<SFilter> NewFilter = this->AddFilterToBar(*Filter);
+				NewFilter->SetEnabled(CustomFilter.Value, false);
+				this->SetFrontendFilterActive(*Filter, NewFilter->IsEnabled());
+			}
+		}
+
 		// Save a copy of all the type filters, and then remove anything we find
 		UnknownTypeFilters = FilterBarConfig->TypeFilters;
 
 		// Load all the type filters
+		TMap<FString, TSharedRef<FCustomClassFilterData>> AssetClassFilters;
+		AssetClassFilters.Reserve(FilterBarConfig->TypeFilters.Num());
 		for(const TSharedRef<FCustomClassFilterData> &CustomClassFilter : this->CustomClassFilters)
 		{
 			const FString FilterName = CustomClassFilter->GetFilterName();
@@ -473,10 +487,9 @@ protected:
 			if(!this->IsClassTypeInUse(CustomClassFilter))
 			{
 				// Try to find this type filter in our list of saved filters
-				if ( const bool* bIsActive = FilterBarConfig->TypeFilters.Find(FilterName) )
+				if (FilterBarConfig->TypeFilters.Contains(FilterName))
 				{
-					TSharedRef<SFilter> NewFilter = this->AddAssetFilterToBar(CustomClassFilter);
-					NewFilter->SetEnabled(*bIsActive, false);
+					AssetClassFilters.Add(FilterName, CustomClassFilter);
 
 					// Remove it from the Unknown list, we have encountered this type filter in the current instance
 					UnknownTypeFilters.Remove(FilterName);
@@ -488,6 +501,16 @@ protected:
 				UnknownTypeFilters.Remove(FilterName);
 			}
 		}
+
+		// We add them here to preserve the saving config order
+		for (const TPair<FString, bool>& TypeFilter : FilterBarConfig->TypeFilters)
+		{
+			if (const TSharedRef<FCustomClassFilterData>* Filter = AssetClassFilters.Find(TypeFilter.Key))
+			{
+				TSharedRef<SFilter> NewFilter = this->AddAssetFilterToBar(*Filter);
+				NewFilter->SetEnabled(TypeFilter.Value, false);
+			}
+		}
 	}
 	
 	typedef typename SBasicFilterBar<FilterType>::SFilter SFilter;
@@ -820,6 +843,16 @@ public:
 		}
 	}
 
+	/** Set the active asset filters for this filter list */
+	void SetActiveFilters(TConstArrayView<FTopLevelAssetPath> InFilters)
+	{
+		RemoveAllFilters();
+		for (const FTopLevelAssetPath& FilterAssetPath : InFilters)
+		{
+			SetAssetTypeFilterCheckState(FilterAssetPath, ECheckBoxState::Checked, /** Enabled */true);
+		}
+	}
+
 protected:
 
 	/** AssetFilter specific override to SBasicFilterBar::RemoveAllButThis */
