Fixed ContentBrowserModule's "AssetClassesRequiringDynamicTags" feature. It now stores classes as FTopLevelAssetPath in order to properly detect them.
Also, now, child classes of an asset class requiring dynamic tags are also considered requiring dynamic tags

diff --git a/Engine/Plugins/FX/Niagara/Source/NiagaraEditor/Private/NiagaraEditorModule.cpp b/Engine/Plugins/FX/Niagara/Source/NiagaraEditor/Private/NiagaraEditorModule.cpp
--- a/Engine/Plugins/FX/Niagara/Source/NiagaraEditor/Private/NiagaraEditorModule.cpp
+++ b/Engine/Plugins/FX/Niagara/Source/NiagaraEditor/Private/NiagaraEditorModule.cpp
@@ -78,6 +78,7 @@
 #include "NiagaraDigestDatabase.h"
 #include "NiagaraEmitter.h"
 #include "NiagaraScriptSource.h"
+#include "NiagaraSystem.h"
 #include "NiagaraTypes.h"
 #include "NiagaraSystemFactoryNew.h"
 #include "NiagaraSystemEditorData.h"
@@ -1203,8 +1204,8 @@ void FNiagaraEditorModule::StartupModule()
 	TempPackage->AddToRoot();
 
 	FContentBrowserModule& ContentBrowserModule = FModuleManager::LoadModuleChecked<FContentBrowserModule>(TEXT("ContentBrowser"));
-	ContentBrowserModule.AddDynamicTagAssetClass(TEXT("NiagaraSystem"));
-	ContentBrowserModule.AddDynamicTagAssetClass(TEXT("NiagaraEmitter"));
+	ContentBrowserModule.RegisterDynamicTagAssetClass(UNiagaraSystem::StaticClass()->GetClassPathName());
+	ContentBrowserModule.RegisterDynamicTagAssetClass(UNiagaraEmitter::StaticClass()->GetClassPathName());
 
 	// Preload all parameter definition & collection assets so that they will be postloaded before postload calls to scripts/emitters/systems that rely on them.
 	{
@@ -1861,6 +1862,13 @@ void FNiagaraEditorModule::ShutdownModule()
 	FAssetCompilingManager::Get().UnregisterManager(&FNiagaraSystemCompilingManager::Get());
 
 	TraversalCache.Reset();
+	
+	if (FContentBrowserModule* ContentBrowserModule = FModuleManager::GetModulePtr<FContentBrowserModule>("ContentBrowser");
+		(ContentBrowserModule != nullptr) && UObjectInitialized())
+	{
+		ContentBrowserModule->UnregisterDynamicTagAssetClass(UNiagaraSystem::StaticClass()->GetClassPathName());
+		ContentBrowserModule->UnregisterDynamicTagAssetClass(UNiagaraEmitter::StaticClass()->GetClassPathName());
+	}
 }
 
 void FNiagaraEditorModule::OnPostEngineInit()
diff --git a/Engine/Source/Editor/ContentBrowser/Private/ContentBrowserModule.cpp b/Engine/Source/Editor/ContentBrowser/Private/ContentBrowserModule.cpp
--- a/Engine/Source/Editor/ContentBrowser/Private/ContentBrowserModule.cpp
+++ b/Engine/Source/Editor/ContentBrowser/Private/ContentBrowserModule.cpp
@@ -86,6 +86,38 @@ void FContentBrowserModule::ContentBrowserSettingChanged(FName InName)
 	OnContentBrowserSettingChanged.Broadcast(InName);
 }
 
+void FContentBrowserModule::RegisterDynamicTagAssetClass(const FTopLevelAssetPath& InClassPathName)
+{
+	AssetClassesRequiringDynamicTags.AddUnique(InClassPathName);
+}
+
+void FContentBrowserModule::UnregisterDynamicTagAssetClass(const FTopLevelAssetPath& InClassPathName)
+{
+	AssetClassesRequiringDynamicTags.Remove(InClassPathName);
+}
+
+bool FContentBrowserModule::IsDynamicTagAssetClass(const FTopLevelAssetPath& InClassPathName) const 
+{
+	// Fast path if the class name matches perfectly with one of the asset classes requiring dynamic tags : 
+	if (AssetClassesRequiringDynamicTags.Contains(InClassPathName))
+	{
+		return true;
+	}
+
+	const UClass* RequestedClass = FindObject<UClass>(InClassPathName);
+	if (RequestedClass == nullptr)
+	{
+		return false;
+	}
+
+	// Otherwise, do the slow search and check if the requested class is a child of one of the asset classes requiring dynamic tags : 
+	return AssetClassesRequiringDynamicTags.ContainsByPredicate([RequestedClass](const FTopLevelAssetPath& InDynamicTagsClassPath)
+		{
+			const UClass* DynamicTagsClass = FindObject<UClass>(InDynamicTagsClassPath);
+			return ((DynamicTagsClass != nullptr) && RequestedClass->IsChildOf(DynamicTagsClass));
+		});
+}
+
 void FContentBrowserModule::RegisterWizard(const FWizard& InWizard)
 {
 	Wizards.Add(InWizard);
diff --git a/Engine/Source/Editor/ContentBrowser/Private/SAssetView.cpp b/Engine/Source/Editor/ContentBrowser/Private/SAssetView.cpp
--- a/Engine/Source/Editor/ContentBrowser/Private/SAssetView.cpp
+++ b/Engine/Source/Editor/ContentBrowser/Private/SAssetView.cpp
@@ -3035,6 +3035,7 @@ TSharedRef<SAssetListView> SAssetView::CreateListView()
 			.HAlignCell(HAlign_Center)
 			.VAlignCell(VAlign_Center)
 			.DefaultLabel( LOCTEXT("Column_RC", "Revision Control") )
+			.DefaultTooltip(LOCTEXT("Column_RC_Tooltip", "Revision Control Status"))
 			[
 				RevisionControlColumnIcon
 			]
@@ -4070,7 +4071,7 @@ void SAssetView::SetMajorityAssetType(FName NewMajorityAssetType)
 
 	FContentBrowserModule& ContentBrowserModule = FModuleManager::GetModuleChecked<FContentBrowserModule>(TEXT("ContentBrowser"));
 
-	bool bHasDynamicColumns = ContentBrowserModule.IsDynamicTagAssetClass(NewMajorityAssetType);
+	bool bHasDynamicColumns = ContentBrowserModule.IsDynamicTagAssetClass(FTopLevelAssetPath(NewMajorityAssetType.ToString()));
 
 	if ( NewMajorityAssetType != MajorityAssetType || bHasDynamicColumns)
 	{
diff --git a/Engine/Source/Editor/ContentBrowser/Public/ContentBrowserModule.h b/Engine/Source/Editor/ContentBrowser/Public/ContentBrowserModule.h
--- a/Engine/Source/Editor/ContentBrowser/Public/ContentBrowserModule.h
+++ b/Engine/Source/Editor/ContentBrowser/Public/ContentBrowserModule.h
@@ -207,22 +207,17 @@ public:
 
 	static const FName NumberOfRecentAssetsName;
 
-	void AddDynamicTagAssetClass(const FName& InName) 
-	{
-		AssetClassesRequiringDynamicTags.AddUnique(InName);
-	}
-
-
-	void RemoveDynamicTagAssetClass(const FName& InName)
-	{
-		AssetClassesRequiringDynamicTags.Remove(InName);
-	}
-		
-	bool IsDynamicTagAssetClass(const FName& InName)
-	{
-		return AssetClassesRequiringDynamicTags.Contains(InName);
-	}
+	CONTENTBROWSER_API void RegisterDynamicTagAssetClass(const FTopLevelAssetPath& InClassPathName);
+	CONTENTBROWSER_API void UnregisterDynamicTagAssetClass(const FTopLevelAssetPath& InClassPathName);
+	CONTENTBROWSER_API bool IsDynamicTagAssetClass(const FTopLevelAssetPath& InClassPathName) const;
 	
+	UE_DEPRECATED(5.8, "Use the new version of this function, which uses a FTopLevelAssetPath")
+	void AddDynamicTagAssetClass(const FName& InName) {}
+	UE_DEPRECATED(5.8, "Use the new version of this function, which uses a FTopLevelAssetPath")
+	void RemoveDynamicTagAssetClass(const FName& InName) {}
+	UE_DEPRECATED(5.8, "Use the new version of this function, which uses a FTopLevelAssetPath")
+	bool IsDynamicTagAssetClass(const FName& InName) { return false; }
+
 	FOnSetFolderColor& GetOnSetFolderColor()
 	{
 		return OnSetFolderColor;
@@ -249,7 +244,7 @@ private:
 	void ContentBrowserSettingChanged(FName InName);
 	
 	/** List of asset classes whose tags are dynamic and therefore we should union all asset's tags rather than grabbing the first available. */
-	TArray<FName> AssetClassesRequiringDynamicTags;
+	TArray<FTopLevelAssetPath> AssetClassesRequiringDynamicTags;
 
 	/** Called when a custom folder color is set on a folder */
 	FOnSetFolderColor OnSetFolderColor;
