Re-submit of CL 45001515

- Original changes done by [REDACTED_PERSON]
- Additionally set selected paths on PathViewPtr right away even if requesting an async population in SContentBrowser::SetSelectedPaths

#Synced-CL 48093036

diff --git a/Engine/Source/Editor/ContentBrowser/Private/SContentBrowser.cpp b/Engine/Source/Editor/ContentBrowser/Private/SContentBrowser.cpp
--- a/Engine/Source/Editor/ContentBrowser/Private/SContentBrowser.cpp
+++ b/Engine/Source/Editor/ContentBrowser/Private/SContentBrowser.cpp
@@ -1511,14 +1511,6 @@ void SContentBrowser::PrepareToSyncItems(TArrayView<const FContentBrowserItem> I
 		bSomePluginPathFiltersChanged = true;
 		bRepopulate = true;
 	}
-
-	// Check to see if any item paths don't exist (this can happen if we haven't ticked since the path was created)
-	if (!bRepopulate)
-	{
-		bRepopulate = Algo::AnyOf(ItemsToSync, [this](const FContentBrowserItem& Item) {
-			return Item.IsFolder() && !LegacyContentSourceWidgets->PathViewPtr->DoesItemExist(Item.GetVirtualPath());
-		});
-	}
 	
 	if (bDisableFiltersThatHideAssets)
 	{
@@ -2253,17 +2245,12 @@ void SContentBrowser::SetSelectedPaths(const TArray<FString>& FolderPaths, bool
 
 		if (bNeedsRefresh)
 		{
-			LegacyContentSourceWidgets->PathViewPtr->RequestPopulation(FSimpleMulticastDelegate::FDelegate::CreateSPLambda(this, [this, NormalizedFolderPaths]()
-				{
-					LegacyContentSourceWidgets->PathViewPtr->SetSelectedPaths(NormalizedFolderPaths);
-				}));
+			LegacyContentSourceWidgets->PathViewPtr->RequestPopulation();	
 			LegacyContentSourceWidgets->FavoritePathViewPtr->Populate();
 		}
-		else
-		{
-			LegacyContentSourceWidgets->PathViewPtr->SetSelectedPaths(NormalizedFolderPaths);
-		}
 
+		// Need to set selected paths on PathViewPtr here even if we just called RequestPopulation above because PathSelected below expects it to be set 
+		LegacyContentSourceWidgets->PathViewPtr->SetSelectedPaths(NormalizedFolderPaths);
 		LegacyContentSourceWidgets->FavoritePathViewPtr->SetSelectedPaths(NormalizedFolderPaths);
 		PathSelected(NormalizedFolderPaths[0]);
 	}
diff --git a/Engine/Source/Editor/ContentBrowser/Private/SPathView.cpp b/Engine/Source/Editor/ContentBrowser/Private/SPathView.cpp
--- a/Engine/Source/Editor/ContentBrowser/Private/SPathView.cpp
+++ b/Engine/Source/Editor/ContentBrowser/Private/SPathView.cpp
@@ -347,6 +347,9 @@ public:
 	// Searches all items, not just visible ones according to the current text filter
 	TSharedPtr<FTreeItem> FindBestItemForPath(FStringView InVirtualPath);
 
+	void ProcessVirtualPaths(TConstArrayView<FName> InVirtualPaths, TFunction<FContentBrowserDataCompiledFilter()> CompiledFilterGetter);
+	void ProcessVirtualPaths(TConstArrayView<FName> InVirtualPaths, TFunction<FContentBrowserDataCompiledFilter()> CompiledFilterGetter, TSet <TSharedRef<FTreeItem>>& OutItems);
+
 	// Apply new/modified/removed data notifications to the tree
 	// bFlat parameter adds all items at the root of the tree as in the favorites tree
 	void ProcessDataUpdates(TConstArrayView<FContentBrowserItemDataUpdate> InUpdatedItems,
@@ -727,6 +730,40 @@ TSharedPtr<FTreeItem> FPathViewData::FindBestItemForPath(FStringView InVirtualPa
 	return Found;
 }
 
+void FPathViewData::ProcessVirtualPaths(TConstArrayView<FName> InVirtualPaths, TFunction<FContentBrowserDataCompiledFilter()> CompiledFilterGetter)
+{
+	UContentBrowserDataSubsystem* ContentBrowserData = IContentBrowserDataModule::Get().GetSubsystem();
+
+	TArray<FContentBrowserItemDataUpdate> FolderAtPath;
+	for (const FName& Path : InVirtualPaths)
+	{
+		ContentBrowserData->EnumerateItemsAtPath(Path, EContentBrowserItemTypeFilter::IncludeFolders, [&FolderAtPath](FContentBrowserItemData&& InItem) -> bool
+		{
+			FolderAtPath.Add(FContentBrowserItemDataUpdate::MakeItemAddedUpdate(MoveTemp(InItem)));
+			return true;
+		});
+	}
+
+	if (!FolderAtPath.IsEmpty())
+	{
+		ProcessDataUpdates(FolderAtPath, CompiledFilterGetter());
+	}
+}
+
+void FPathViewData::ProcessVirtualPaths(TConstArrayView<FName> InVirtualPaths, TFunction<FContentBrowserDataCompiledFilter()> CompiledFilterGetter, TSet<TSharedRef<FTreeItem>>& OutItems)
+{
+	ProcessVirtualPaths(InVirtualPaths, MoveTemp(CompiledFilterGetter));
+
+	for (const FName& VirtualPath : InVirtualPaths)
+	{
+		TSharedPtr<FTreeItem> Item = FindTreeItem(VirtualPath);
+		if (Item.IsValid())
+		{
+			OutItems.Add(Item.ToSharedRef());
+		}
+	}
+}
+
 TSharedRef<FTreeItem> FPathViewData::AddFolderItem(FContentBrowserItemData&& InItemData)
 {
 	TSharedRef<FTreeItem> NewOrUpdatedItem = AddFolderItemInternal(MoveTemp(InItemData), nullptr);
@@ -1223,9 +1260,6 @@ void SPathView::Construct( const FArguments& InArgs )
 
 	DefaultPath = InArgs._DefaultPath;
 	bCreateDefaultPath = InArgs._CreateDefaultPath;
-
-	// Temporary turn off the delayed population to investigate navigation issue when the view is not populated
-	Populate();
 }
 
 void SPathView::ConfigureTreeView(STreeView<TSharedPtr<FTreeItem>>::FArguments& InArgs)
@@ -1302,9 +1336,6 @@ void SPathView::PluginPathFilterClicked(TSharedRef<FContentBrowserPluginFilter>
 	SetPluginPathFilterActive(Filter, !IsPluginPathFilterInUse(Filter));
 	
 	bPopulationRequestedAttribute.Assign(*this, true);
-
-	// Temporary turn off the delayed population to investigate navigation issue when the view is not populated
-	Populate();
 }
 
 bool SPathView::IsPluginPathFilterChecked(TSharedRef<FContentBrowserPluginFilter> Filter) const
@@ -1340,8 +1371,6 @@ void SPathView::ResetPluginPathFilters()
 	}
 
 	bPopulationRequestedAttribute.Assign(*this, true);
-	// Temporary turn off the delayed population to investigate navigation issue when the view is not populated
-	Populate();
 }
 
 bool SPathView::DisablePluginPathFiltersThatHideItems(TConstArrayView<FContentBrowserItem> Items)
@@ -1457,6 +1486,26 @@ void SPathView::SetSelectedPaths(const TArray<FString>& Paths)
 	LastSelectedPaths.Empty();
 	TreeViewPtr->ClearSelection();
 
+	if (ShouldDoJustInTimePopulation())
+	{
+		TArray<const FName> InPathToPopulate;
+		for (const FString& Path : Paths)
+		{
+			FName PathAsName = *Path;
+			TSharedPtr<FTreeItem> Item = TreeData->FindTreeItem(PathAsName);
+			if (!Item)
+			{
+				InPathToPopulate.Add(PathAsName);
+			}
+		}
+
+		if (!InPathToPopulate.IsEmpty())
+		{
+			TreeData->ProcessVirtualPaths(InPathToPopulate,
+				[this]() { return CreateCompiledFolderFilter(); });
+		}
+	}
+
 	for (const FString& Path : Paths)
 	{
 		TSharedPtr<FTreeItem> BestItem = TreeData->FindBestItemForPath(Path);
@@ -1771,6 +1820,7 @@ bool SPathView::InternalPathPassesBlockLists(const FStringView InInternalPath, c
 void SPathView::SyncToItems(TArrayView<const FContentBrowserItem> ItemsToSync, const bool bAllowImplicitSync)
 {
 	TSet<TSharedRef<FTreeItem>> SyncTreeItems;
+	TArray<const FName> VirtualPathsNotFound;
 	for (const FContentBrowserItem& ItemToSync : ItemsToSync)
 	{
 		FName VirtualPathToSync;
@@ -1784,11 +1834,24 @@ void SPathView::SyncToItems(TArrayView<const FContentBrowserItem> ItemsToSync, c
 			VirtualPathToSync = ItemToSync.GetVirtualPath();
 		}
 
+		bool bShouldDoJustInTimePopulation = ShouldDoJustInTimePopulation();
+
 		TSharedPtr<FTreeItem> Item = TreeData->FindTreeItem(VirtualPathToSync);
 		if (Item.IsValid())
 		{
 			SyncTreeItems.Add(Item.ToSharedRef());
 		}
+		else if (bShouldDoJustInTimePopulation)
+		{
+			VirtualPathsNotFound.Add(VirtualPathToSync);
+		}
+	}
+
+	if (!VirtualPathsNotFound.IsEmpty())
+	{
+		TreeData->ProcessVirtualPaths(VirtualPathsNotFound,
+			[this]() { return CreateCompiledFolderFilter(); },
+			SyncTreeItems);
 	}
 
 	if (Algo::AnyOf(SyncTreeItems, [](const TSharedRef<FTreeItem>& Item) { return !Item->IsVisible(); }))
@@ -1850,11 +1913,6 @@ void SPathView::SyncToItems(TArrayView<const FContentBrowserItem> ItemsToSync, c
 	}
 }
 
-bool SPathView::DoesItemExist(FName InVirtualPath) const
-{
-	return TreeData->FindTreeItem(InVirtualPath).IsValid();
-}
-
 void SPathView::ApplyHistoryData( const FHistoryData& History )
 {
 	// Prevent the selection changed delegate because it would add more history when we are just setting a state
@@ -2284,9 +2342,6 @@ void SPathView::RequestPopulation(FSimpleMulticastDelegate::FDelegate&& AfterPop
 	{
 		OnNextPopulate.Add(MoveTemp(AfterPopulation));
 	}
-
-	// Temporary turn off the delayed population to investigate navigation issue when the view is not populated
-	Populate();
 }
 
 
@@ -2521,9 +2576,6 @@ void SPathView::HandleItemDataUpdated(TArrayView<const FContentBrowserItemDataUp
 void SPathView::HandleItemDataRefreshed()
 {
 	bPopulationRequestedAttribute.Assign(*this, true);
-
-	// Temporary turn off the delayed population to investigate navigation issue when the view is not populated
-	Populate();
 }
 
 void SPathView::HandleItemDataDiscoveryComplete()
@@ -2554,9 +2606,6 @@ void SPathView::HandleSettingChanged(FName PropertyName)
 
 		// Update our path view so that it can include/exclude the permitted folder
 		bPopulationForSettingChangeRequestedAttribute.Assign(*this, true);
-
-		// Temporary turn off the delayed population to investigate navigation issue when the view is not populated
-		Populate();
 	}
 }
 
diff --git a/Engine/Source/Editor/ContentBrowser/Private/SPathView.h b/Engine/Source/Editor/ContentBrowser/Private/SPathView.h
--- a/Engine/Source/Editor/ContentBrowser/Private/SPathView.h
+++ b/Engine/Source/Editor/ContentBrowser/Private/SPathView.h
@@ -226,9 +226,6 @@ public:
 	 */
 	void SyncToItems( TArrayView<const FContentBrowserItem> ItemsToSync, const bool bAllowImplicitSync = false );
 
-	/** Returns whether the tree contains an item with the given virtual path. */
-	bool DoesItemExist(FName InVirtualPath) const;
-
 	/** Sets the state of the path view to the one described by the history data */
 	void ApplyHistoryData( const FHistoryData& History );
 
@@ -353,6 +350,11 @@ protected:
 	/** Should be call after a populate */
 	void SharedPostPopulate();
 
+	constexpr virtual bool ShouldDoJustInTimePopulation() const
+	{
+		return true;
+	}
+
 private:
 	/** Selects the given path only if it exists. Returns true if selected. */
 	bool ExplicitlyAddPathToSelection(const FName Path);
@@ -639,6 +641,12 @@ public:
 	DECLARE_DELEGATE_OneParam(FOnFolderFavoriteAdd, const TArray<FString>& /*FoldersToAdd*/)
 	void SetOnFolderFavoriteAdd(const FOnFolderFavoriteAdd& InOnFolderFavoriteAdd);
 
+protected:
+	constexpr virtual bool ShouldDoJustInTimePopulation() const override
+	{
+		return false;
+	}
+
 private:
 	virtual TSharedRef<ITableRow> GenerateTreeRow(TSharedPtr<FTreeItem> TreeItem, const TSharedRef<STableViewBase>& OwnerTable) override;
 
