Fixed UCollectionManagerScriptingSubsystem::SyncBrowserToCollections not selecting collections that were created in the same frame.

#rnx
#Synced-CL 47947116

diff --git a/Engine/Source/Editor/ContentBrowser/Private/SCollectionView.cpp b/Engine/Source/Editor/ContentBrowser/Private/SCollectionView.cpp
--- a/Engine/Source/Editor/ContentBrowser/Private/SCollectionView.cpp
+++ b/Engine/Source/Editor/ContentBrowser/Private/SCollectionView.cpp
@@ -624,10 +624,20 @@ void SCollectionView::SetSelectedCollections(const TArray<FCollectionNameType>&
 	// Prevent the selection changed delegate since the invoking code requested it
 	FScopedPreventSelectionChangedDelegate DelegatePrevention( SharedThis(this) );
 
-
 	// Clear the selection to start, then add the selected items as they are found
 	CollectionTreePtr->ClearSelection();
 
+	if (CollectionsToSelect.IsEmpty())
+	{
+		return;
+	}
+
+	if (bQueueCollectionItemsUpdate)
+	{
+		bQueueCollectionItemsUpdate = false;
+		UpdateCollectionItems();
+	}
+
 	for (const FCollectionNameType& CollectionToSelect : CollectionsToSelect)
 	{
 		TSharedPtr<FCollectionItem> CollectionItemToSelect = AvailableCollections.FindRef(CollectionToSelect);
