Content Browser - Folders - View Options test added

#rnx
#Synced-CL 47900565

diff --git a/Engine/Source/Editor/ContentBrowser/Public/SAssetView.h b/Engine/Source/Editor/ContentBrowser/Public/SAssetView.h
--- a/Engine/Source/Editor/ContentBrowser/Public/SAssetView.h
+++ b/Engine/Source/Editor/ContentBrowser/Public/SAssetView.h
@@ -536,6 +536,30 @@ public:
 	/** @return true when we are showing favorites */
 	UE_INTERNAL UE_API bool IsShowingFavorites() const;
 
+	/** @return true when we are showing C++ content */
+	UE_INTERNAL UE_API bool IsShowingCppContent() const;
+
+	/** @return true when we are showing folders */
+	UE_INTERNAL UE_API bool IsShowingFolders() const;
+
+	/** @return true when we are showing empty folders */
+	UE_INTERNAL UE_API bool IsShowingEmptyFolders() const;
+
+	/** @return true when we are showing plugin content */
+	UE_INTERNAL UE_API bool IsShowingPluginContent() const;
+
+	/** @return true when we are showing engine content */
+	UE_INTERNAL UE_API bool IsShowingEngineContent() const;
+
+	/** @return true when we are showing the developers content */
+	UE_INTERNAL UE_API bool IsShowingDevelopersContent() const;
+
+	/** @return true when we are showing the all virtual folder */
+	UE_INTERNAL UE_API bool IsShowingAllFolder() const;
+
+	/** @return true when we are organizing folders into virtual locations */
+	UE_INTERNAL UE_API bool IsOrganizingFolders() const;
+
 private:
 	/** Sets the pending selection to the current selection (used when changing views or refreshing the view). */
 	UE_API void SyncToSelection( const bool bFocusOnSync = true );
@@ -640,9 +664,6 @@ private:
 	/** Whether or not it's possible to show folders */
 	UE_API bool IsToggleShowFoldersAllowed() const;
 
-	/** @return true when we are showing folders */
-	UE_API bool IsShowingFolders() const;
-
 	/** @return true when we are showing read-only folders */
 	UE_API bool IsShowingReadOnlyFolders() const;
 
@@ -652,9 +673,6 @@ private:
 	/** Whether or not it's possible to show empty folders */
 	UE_API bool IsToggleShowEmptyFoldersAllowed() const;
 
-	/** @return true when we are showing empty folders */
-	UE_API bool IsShowingEmptyFolders() const;
-
 	/** @return true when the asset view is showing object redirectors */
 	UE_API bool IsShowingRedirectors() const;
 
@@ -679,15 +697,9 @@ private:
 	/** Toggle whether plugin content should be shown or not */
 	UE_API void ToggleShowPluginContent();
 
-	/** @return true when we are showing plugin content */
-	UE_API bool IsShowingPluginContent() const;
-
 	/** Toggle whether engine content should be shown or not */
 	UE_API void ToggleShowEngineContent();
 
-	/** @return true when we are showing engine content */
-	UE_API bool IsShowingEngineContent() const;
-
 	/** Toggle whether developers content should be shown or not */
 	UE_API void ToggleShowDevelopersContent();
 
@@ -699,9 +711,6 @@ private:
 
 	/** Whether or not it's possible to toggle plugin content */
 	UE_API bool IsToggleShowPluginContentAllowed() const;
-	
-	/** @return true when we are showing the developers content */
-	UE_API bool IsShowingDevelopersContent() const;
 
 	/** Toggle whether favorites should be shown or not */
 	UE_API void ToggleShowFavorites();
@@ -712,9 +721,6 @@ private:
 	/** Whether or not it's possible to show C++ content */
 	UE_API bool IsToggleShowCppContentAllowed() const;
 
-	/** @return true when we are showing C++ content */
-	UE_API bool IsShowingCppContent() const;
-
 	/** Toggle whether class names should be included in search criteria */
 	UE_API void ToggleIncludeClassNames();
 
@@ -745,15 +751,9 @@ private:
 	/** Toggle whether we are showing the all virtual folder */
 	UE_API void ToggleShowAllFolder();
 
-	/** @return true when we are showing the all virtual folder */
-	UE_API bool IsShowingAllFolder() const;
-
 	/** Toggle whether we are re-organizing folders into virtual locations */
 	UE_API void ToggleOrganizeFolders();
 
-	/** @return true when we are organizing folders into virtual locations */
-	UE_API bool IsOrganizingFolders() const;
-
 	/** Sets the view type and forcibly dismisses all currently open context menus */
 	UE_API void SetCurrentViewTypeFromMenu(EAssetViewType::Type NewType);
 
