Content Browser - Column View : added description of how many columns are displayed and from what asset type they are from. Without that information, it's hard for a user to understand there are some missing columns simply because they're not part of the columns of the asset type for which there's the most assets on display

#Synced-CL 48911231

diff --git a/Engine/Source/Editor/ContentBrowser/Private/AssetViewTypes.cpp b/Engine/Source/Editor/ContentBrowser/Private/AssetViewTypes.cpp
--- a/Engine/Source/Editor/ContentBrowser/Private/AssetViewTypes.cpp
+++ b/Engine/Source/Editor/ContentBrowser/Private/AssetViewTypes.cpp
@@ -5,6 +5,7 @@
 #include "AssetRegistry/AssetData.h"
 #include "Containers/UnrealString.h"
 #include "Containers/VersePath.h"
+#include "ContentBrowserDataSource.h"
 #include "ContentBrowserDelegates.h"
 #include "ContentBrowserItemData.h"
 #include "HAL/Platform.h"
@@ -152,6 +153,21 @@ void FAssetViewItem::CacheVersePathTextForTextFiltering()
 	}
 }
 
+FText FAssetViewItem::GetAssetClassText() const
+{
+	if (IsFolder())
+	{
+		return LOCTEXT("FolderName", "Folder");
+	}
+
+	FContentBrowserItemDataAttributeValue DisplayNameAttributeValue = Item.GetItemAttribute(ContentBrowserItemAttributes::ItemTypeDisplayName);
+	if (!DisplayNameAttributeValue.IsValid())
+	{
+		DisplayNameAttributeValue = Item.GetItemAttribute(ContentBrowserItemAttributes::ItemTypeName);
+	}
+	return DisplayNameAttributeValue.IsValid() ? DisplayNameAttributeValue.GetValue<FText>() : FText();
+}
+
 const FText& FAssetViewItem::GetAssetPathText() const
 {
 	checkf(IsInGameThread(), TEXT("Verse paths can only be computed on the main thread."));
diff --git a/Engine/Source/Editor/ContentBrowser/Private/AssetViewWidgets.cpp b/Engine/Source/Editor/ContentBrowser/Private/AssetViewWidgets.cpp
--- a/Engine/Source/Editor/ContentBrowser/Private/AssetViewWidgets.cpp
+++ b/Engine/Source/Editor/ContentBrowser/Private/AssetViewWidgets.cpp
@@ -617,17 +617,7 @@ FText SAssetViewItem::GetAssetClassText() const
 		return FText();
 	}
 	
-	if (AssetItem->IsFolder())
-	{
-		return LOCTEXT("FolderName", "Folder");
-	}
-
-	FContentBrowserItemDataAttributeValue DisplayNameAttributeValue = AssetItem->GetItem().GetItemAttribute(ContentBrowserItemAttributes::ItemTypeDisplayName);
-	if (!DisplayNameAttributeValue.IsValid())
-	{
-		DisplayNameAttributeValue = AssetItem->GetItem().GetItemAttribute(ContentBrowserItemAttributes::ItemTypeName);
-	}
-	return DisplayNameAttributeValue.IsValid() ? DisplayNameAttributeValue.GetValue<FText>() : FText();
+	return AssetItem->GetAssetClassText();
 }
 
 void SAssetViewItem::HandleSourceControlProviderChanged(ISourceControlProvider& OldProvider, ISourceControlProvider& NewProvider)
diff --git a/Engine/Source/Editor/ContentBrowser/Private/Menus/AssetViewMenus.cpp b/Engine/Source/Editor/ContentBrowser/Private/Menus/AssetViewMenus.cpp
--- a/Engine/Source/Editor/ContentBrowser/Private/Menus/AssetViewMenus.cpp
+++ b/Engine/Source/Editor/ContentBrowser/Private/Menus/AssetViewMenus.cpp
@@ -641,11 +641,11 @@ FDelayedAutoRegisterHelper SAssetView::NavigationBarMenuExtensionRegistration(
 
 			if (bIsNewStyleEnabled)
 			{
-				FToolMenuSection& AssetCountSection = NavigationBarMenu->FindOrAddSection("AssetCount", { }, FToolMenuInsert("History", EToolMenuInsertType::Last));
-				AssetCountSection.Alignment = EToolMenuSectionAlign::Last;
+				FToolMenuSection& AssetViewInfoSection = NavigationBarMenu->FindOrAddSection("AssetViewInfo", { }, FToolMenuInsert("History", EToolMenuInsertType::Last));
+				AssetViewInfoSection.Alignment = EToolMenuSectionAlign::Last;
 
-				AssetCountSection.AddDynamicEntry(
-					"AssetCount",
+				AssetViewInfoSection.AddDynamicEntry(
+					"AssetViewInfo",
 					FNewToolMenuSectionDelegate::CreateLambda([](FToolMenuSection& InSection)
 					{
 						if (const UContentBrowserToolbarMenuContext* Context = InSection.FindContext<UContentBrowserToolbarMenuContext>())
@@ -654,18 +654,35 @@ FDelayedAutoRegisterHelper SAssetView::NavigationBarMenuExtensionRegistration(
 							{
 								const TSharedRef<SAssetView> AssetView = Context->AssetView.Pin().ToSharedRef();
 
-								FToolMenuEntry& AssetCountEntry = InSection.AddEntry(
+								FToolMenuEntry& AssetViewInfoEntry = InSection.AddEntry(
 									FToolMenuEntry::InitWidget(
-										"AssetCount",
+										"AssetViewInfo",
 										SNew(SBox)
-										.Padding(FMargin(4, 0, 0, 0))
+										.Padding(FMargin(4, 0, 4, 0))
 										[
-											SNew(STextBlock)
-											.Text(AssetView, &SAssetView::GetAssetCountText)
+											SNew(SHorizontalBox)
+											+ SHorizontalBox::Slot()
+											.HAlign(HAlign_Left)
+											.FillWidth(1.0f)
+											[
+												SNew(STextBlock)
+												.Text(AssetView, &SAssetView::GetAssetCountText)
+											]
+											+ SHorizontalBox::Slot()
+											.HAlign(HAlign_Right)
+											.FillWidth(1.0f)
+											[
+												SNew(STextBlock)
+												.Text(AssetView, &SAssetView::GetMajorityAssetTypeText)
+												.Visibility(AssetView, &SAssetView::GetMajorityAssetTypeTextVisibility)
+												.ToolTipText(AssetView, &SAssetView::GetMajorityAssetTypeToolTipText)
+											]
 										],
 										FText::GetEmpty()));
-
-								AssetCountEntry.WidgetData.StyleParams.VerticalAlignment = VAlign_Center;
+								AssetViewInfoEntry.WidgetData.StyleParams.VerticalAlignment = VAlign_Center;
+								AssetViewInfoEntry.WidgetData.StyleParams.HorizontalAlignment = HAlign_Fill;
+								AssetViewInfoEntry.WidgetData.StyleParams.FillSize = 1.0f;
+								AssetViewInfoEntry.WidgetData.StyleParams.SizeRule = FSizeParam::SizeRule_Stretch;
 							}
 						}
 					}));
diff --git a/Engine/Source/Editor/ContentBrowser/Private/SAssetView.cpp b/Engine/Source/Editor/ContentBrowser/Private/SAssetView.cpp
--- a/Engine/Source/Editor/ContentBrowser/Private/SAssetView.cpp
+++ b/Engine/Source/Editor/ContentBrowser/Private/SAssetView.cpp
@@ -3,6 +3,7 @@
 #include "SAssetView.h"
 
 #include "Algo/AnyOf.h"
+#include "Algo/Count.h"
 #include "Algo/Compare.h"
 #include "Algo/RemoveIf.h"
 #include "Algo/Transform.h"
@@ -4097,6 +4098,7 @@ void SAssetView::SetMajorityAssetType(FName NewMajorityAssetType)
 		UE_LOG(LogContentBrowser, Verbose, TEXT("The majority of assets in the view are of type: %s"), *NewMajorityAssetType.ToString());
 
 		MajorityAssetType = NewMajorityAssetType;
+		MajorityAssetTypeText = FText::GetEmpty();
 
 		TArray<FName> AddedColumns;
 		TSharedPtr<SListView<TSharedPtr<FAssetViewItem>>> ViewToUse = ColumnView;
@@ -4196,6 +4198,8 @@ void SAssetView::SetMajorityAssetType(FName NewMajorityAssetType)
 				}
 			}
 
+			MajorityAssetTypeText = MajorityAssetItem ? MajorityAssetItem->GetAssetClassText() : FText::GetEmpty();
+
 			// Determine the columns by querying the reference item
 			if (MajorityAssetItem)
 			{
@@ -6009,6 +6013,36 @@ FText SAssetView::GetAssetCountText() const
 	return AssetCount;
 }
 
+EVisibility SAssetView::GetMajorityAssetTypeTextVisibility() const
+{
+	return (CurrentViewType == EAssetViewType::Column) ? EVisibility::Visible : EVisibility::Collapsed;
+}
+
+FText SAssetView::GetMajorityAssetTypeText() const
+{
+	if ((CurrentViewType != EAssetViewType::Column) || MajorityAssetType.IsNone() || MajorityAssetTypeText.IsEmpty())
+	{
+		return FText::GetEmpty();
+	}
+
+	const uint32 NumVisibleColumns = IntCastChecked<uint32>(Algo::CountBy(ColumnView->GetHeaderRow()->GetColumns(), /*InValue = */true, &SHeaderRow::FColumn::bIsVisible));
+
+	return FText::Format(LOCTEXT("MajorityAssetTypeText", "Showing {0} {0}|plural(one=column,other=columns) for asset type \"{1}\""), 
+		FText::AsNumber(NumVisibleColumns), MajorityAssetTypeText);
+}
+
+FText SAssetView::GetMajorityAssetTypeToolTipText() const
+{
+	if ((CurrentViewType != EAssetViewType::Column) || MajorityAssetType.IsNone() || MajorityAssetTypeText.IsEmpty())
+	{
+		return FText::GetEmpty();
+	}
+
+	return LOCTEXT("MajorityAssetTypeTooltip", "This view shows columns for the asset type with the most occurrences. \n\r"
+		"\n\r"
+		"If you need to see data for a different asset type, try filtering the view to show only assets of that type.");
+}
+
 EVisibility SAssetView::GetEditModeLabelVisibility() const
 {
 	return IsThumbnailEditMode() ? EVisibility::Visible : EVisibility::Collapsed;
diff --git a/Engine/Source/Editor/ContentBrowser/Public/AssetViewTypes.h b/Engine/Source/Editor/ContentBrowser/Public/AssetViewTypes.h
--- a/Engine/Source/Editor/ContentBrowser/Public/AssetViewTypes.h
+++ b/Engine/Source/Editor/ContentBrowser/Public/AssetViewTypes.h
@@ -61,6 +61,9 @@ public:
 	/** Caches the Verse path to use when filtering items. */
 	UE_INTERNAL void CacheVersePathTextForTextFiltering();
 
+	/** Gets the asset class name to use when viewing. */
+	UE_INTERNAL FText GetAssetClassText() const;
+
 	/** Gets the asset path to use when viewing or sorting the path column. */
 	UE_INTERNAL const FText& GetAssetPathText() const;
 
diff --git a/Engine/Source/Editor/ContentBrowser/Public/SAssetView.h b/Engine/Source/Editor/ContentBrowser/Public/SAssetView.h
--- a/Engine/Source/Editor/ContentBrowser/Public/SAssetView.h
+++ b/Engine/Source/Editor/ContentBrowser/Public/SAssetView.h
@@ -506,6 +506,15 @@ public:
 	/** Gets the text for the asset count label */
 	UE_API FText GetAssetCountText() const;
 
+	/** Gets the visibility for the majority asset type label */
+	EVisibility GetMajorityAssetTypeTextVisibility() const;
+
+	/** Gets the text for the majority asset type label */
+	UE_API FText GetMajorityAssetTypeText() const;
+
+	/** Gets the tooltip text for the majority asset type label */
+	UE_API FText GetMajorityAssetTypeToolTipText() const;
+
 	/** Gets text name for given thumbnail */
 	static UE_API FText ThumbnailSizeToDisplayName(EThumbnailSize InSize);
 
@@ -1220,6 +1229,9 @@ private:
 	/** When in columns view, this is the name of the asset type which is most commonly found in the recent results */
 	FName MajorityAssetType;
 
+	/** When in columns view, this is the display name of the asset type which is most commonly found in the recent results */
+	FText MajorityAssetTypeText;
+
 	/** The manager responsible for sorting assets in the view */
 	TSharedPtr<FAssetViewSortManager> SortManager;
 
