Content Browser: Ensure drag'n'drop of folders into favorites works when the favorite listing is empty.

#rnx
#Synced-CL 48772759

diff --git a/Engine/Source/Editor/ContentBrowser/Private/Widgets/SContentBrowserSourceTree.cpp b/Engine/Source/Editor/ContentBrowser/Private/Widgets/SContentBrowserSourceTree.cpp
--- a/Engine/Source/Editor/ContentBrowser/Private/Widgets/SContentBrowserSourceTree.cpp
+++ b/Engine/Source/Editor/ContentBrowser/Private/Widgets/SContentBrowserSourceTree.cpp
@@ -500,6 +500,31 @@ namespace UE::Editor::ContentBrowser::Private
 		}
 	}
 
+	void SContentBrowserSourceTreeArea::OnDragEnter(const FGeometry& MyGeometry, const FDragDropEvent& DragDropEvent)
+	{
+		if (BodyScrollableWidget.IsValid())
+		{
+			BodyScrollableWidget->GetScrollWidget()->OnDragEnter(MyGeometry, DragDropEvent);
+		}
+	}
+
+	void SContentBrowserSourceTreeArea::OnDragLeave(const FDragDropEvent& DragDropEvent)
+	{
+		if (BodyScrollableWidget.IsValid())
+		{
+			BodyScrollableWidget->GetScrollWidget()->OnDragLeave(DragDropEvent);
+		}
+	}
+
+	FReply SContentBrowserSourceTreeArea::OnDrop(const FGeometry& MyGeometry, const FDragDropEvent& DragDropEvent)
+	{
+		if (BodyScrollableWidget.IsValid())
+		{
+			return BodyScrollableWidget->GetScrollWidget()->OnDrop(MyGeometry, DragDropEvent);
+		}
+		return FReply::Unhandled();
+	}
+
 	TSharedPtr<SSearchToggleButton> SContentBrowserSourceTreeArea::GetSearchToggleButton() const
 	{
 		return SearchToggleButton;
diff --git a/Engine/Source/Editor/ContentBrowser/Private/Widgets/SContentBrowserSourceTree.h b/Engine/Source/Editor/ContentBrowser/Private/Widgets/SContentBrowserSourceTree.h
--- a/Engine/Source/Editor/ContentBrowser/Private/Widgets/SContentBrowserSourceTree.h
+++ b/Engine/Source/Editor/ContentBrowser/Private/Widgets/SContentBrowserSourceTree.h
@@ -201,6 +201,10 @@ namespace UE::Editor::ContentBrowser::Private
 		void LoadSettings(const FString& InIniFilename, const FString& InIniSection, const FString& InSettingsString);
 
 		bool HasSearch() const { return Search.IsValid(); }
+		
+		virtual void OnDragEnter(const FGeometry& MyGeometry, const FDragDropEvent& DragDropEvent) override;
+		virtual void OnDragLeave(const FDragDropEvent& DragDropEvent) override;
+		virtual FReply OnDrop(const FGeometry& MyGeometry, const FDragDropEvent& DragDropEvent) override;
 
 		TSharedPtr<SSearchToggleButton> GetSearchToggleButton() const;
 
