When editing the access specifier of an asset, call PostEditChange so the UContentBrowserAssetDataSource receives a FCoreUObjectDelegates::OnObjectPropertyChanged event and knows to update the item data.
Simplified the FAssetContextMenu::ExecuteBulkSetAssetAccessSpecifier logic and removed a lot of unnecessary loading.

#rnx
#Synced-CL 49027723

diff --git a/Engine/Source/Editor/ContentBrowser/Private/AssetContextMenu.cpp b/Engine/Source/Editor/ContentBrowser/Private/AssetContextMenu.cpp
--- a/Engine/Source/Editor/ContentBrowser/Private/AssetContextMenu.cpp
+++ b/Engine/Source/Editor/ContentBrowser/Private/AssetContextMenu.cpp
@@ -1278,52 +1278,30 @@ FReply FAssetContextMenu::ExecuteDeleteFolderConfirmed()
 	return FReply::Handled();
 }
 
-bool FAssetContextMenu::SetAssetAccessSpecifier(FAssetData& ItemAssetData, const EAssetAccessSpecifier InAssetAccessSpecifier, const bool bEmitEvent)
+int32 FAssetContextMenu::CompareAssetAccessSpecifiers(EAssetAccessSpecifier A, EAssetAccessSpecifier B)
 {
-	UPackage* ItemAssetPackage = ItemAssetData.GetPackage();
-	if (!ItemAssetPackage)
-	{
-		return false;
-	}
+	checkf((A == EAssetAccessSpecifier::Private || A == EAssetAccessSpecifier::EpicInternal || A == EAssetAccessSpecifier::Public) &&
+		(B == EAssetAccessSpecifier::Private || B == EAssetAccessSpecifier::EpicInternal || B == EAssetAccessSpecifier::Public),
+		TEXT("Unknown EAssetAccessSpecifier"));
 
-	const EAssetAccessSpecifier OldAssetAccessSpecifier = ItemAssetPackage->GetAssetAccessSpecifier();
-	if (OldAssetAccessSpecifier == InAssetAccessSpecifier)
+	if (A == B)
 	{
-		return false;
+		return 0;
 	}
 
-	bool bModified = false;
-	if (OldAssetAccessSpecifier == EAssetAccessSpecifier::Private)
-	{
-		bModified = ItemAssetPackage->SetAssetAccessSpecifier(InAssetAccessSpecifier);
-	}
-	else
+	if (A == EAssetAccessSpecifier::Private || B == EAssetAccessSpecifier::Public)
 	{
-		ExecuteBulkSetAssetAccessSpecifier(InAssetAccessSpecifier);
-		// Unknown if there were modifications
-		return bModified;
+		return -1;
 	}
 
-	if (bModified)
-	{
-		if (bEmitEvent)
-		{
-			OnAssetViewRefreshRequested.ExecuteIfBound();
-		}
-	}
-
-	return bModified;
+	return 1;
 }
 
 void FAssetContextMenu::ExecuteSetAssetAccessSpecifier(EAssetAccessSpecifier InAssetAccessSpecifier)
 {
 	if (ensure(SelectedFiles.Num() == 1))
 	{
-		FAssetData ItemAssetData;
-		if (SelectedFiles[0].Legacy_TryGetAssetData(ItemAssetData))
-		{
-			SetAssetAccessSpecifier(ItemAssetData, InAssetAccessSpecifier, true);
-		}
+		ExecuteBulkSetAssetAccessSpecifier(InAssetAccessSpecifier);
 	}
 }
 
@@ -1345,32 +1323,9 @@ void FAssetContextMenu::ExecuteBulkSetAssetAccessSpecifier(EAssetAccessSpecifier
 {
 	TMap<UContentBrowserDataSource*, TArray<FContentBrowserItemData>> PrivatizeSourcesAndItems;
 
-	auto AddForBulkPrivatizeItems = [&PrivatizeSourcesAndItems](const FContentBrowserItem& InItem)
-	{	
-		FContentBrowserItem::FItemDataArrayView ItemDataArray = InItem.GetInternalItems();
-		for (const FContentBrowserItemData& ItemData : ItemDataArray)
-		{
-			if (UContentBrowserDataSource* ItemDataSource = ItemData.GetOwnerDataSource())
-			{
-				FText PrivateErrorMsg;
-				// CanPrivatizeItem checks !GetAssetFolderItemPayload(), CanModifyAssetFileItem(), !IsRunningPIE(), !IsRedirector()
-				if (ItemDataSource->CanPrivatizeItem(ItemData, &PrivateErrorMsg))
-				{
-					TArray<FContentBrowserItemData>& ItemsForSource = PrivatizeSourcesAndItems.FindOrAdd(ItemDataSource);
-					ItemsForSource.Add(ItemData);
-				}
-				else
-				{
-					AssetViewUtils::ShowErrorNotifcation(PrivateErrorMsg);
-				}
-			}
-		}
-	};
-
-	if (DestScope == EAssetAccessSpecifier::Public)
 	{
 		FScopedSlowTask SlowTask((float)SelectedFiles.Num(), LOCTEXT("SetAssetScope", "Loading assets and modifying scope..."));
-		SlowTask.MakeDialog(/*bShowCancelButton*/true);
+		SlowTask.MakeDialogDelayed(0.5f, /*bShowCancelButton*/true);
 
 		for (const FContentBrowserItem& SelectedItem : SelectedFiles)
 		{
@@ -1381,72 +1336,55 @@ void FAssetContextMenu::ExecuteBulkSetAssetAccessSpecifier(EAssetAccessSpecifier
 				break;
 			}
 
-			FAssetData ItemAssetData;
-			if (SelectedItem.Legacy_TryGetAssetData(ItemAssetData))
+			FContentBrowserItem::FItemDataArrayView ItemDataArray = SelectedItem.GetInternalItems();
+			for (const FContentBrowserItemData& ItemData : ItemDataArray)
 			{
-				if (UPackage* ItemAssetPackage = ItemAssetData.GetPackage())
+				if (UContentBrowserDataSource* ItemDataSource = ItemData.GetOwnerDataSource())
 				{
-					ItemAssetPackage->SetAssetAccessSpecifier(DestScope);
-				}
-			}
-		}
-	}
-	else if (DestScope == EAssetAccessSpecifier::EpicInternal)
-	{
-		FScopedSlowTask SlowTask((float)SelectedFiles.Num(), LOCTEXT("SetAssetScope", "Loading assets and modifying scope..."));
-		SlowTask.MakeDialog(/*bShowCancelButton*/true);
-
-		for (const FContentBrowserItem& SelectedItem : SelectedFiles)
-		{
-			SlowTask.EnterProgressFrame(1);
-
-			if (SlowTask.ShouldCancel())
-			{
-				break;
-			}
-
-			FAssetData ItemAssetData;
-			if (!SelectedItem.Legacy_TryGetAssetData(ItemAssetData))
-			{
-				break;
-			}
-
-			UPackage* ItemAssetPackage = ItemAssetData.GetPackage();
-			if (!ItemAssetPackage)
-			{
-				break;
-			}
-			
-			const EAssetAccessSpecifier OldAssetAccessSpecifier = ItemAssetPackage->GetAssetAccessSpecifier();
-			if (OldAssetAccessSpecifier == DestScope)
-			{
-				continue;
-			}
+					FAssetData ItemAssetData;
+					if (!ItemDataSource->Legacy_TryGetAssetData(ItemData, ItemAssetData))
+					{
+						continue;
+					}
 
-			if (OldAssetAccessSpecifier == EAssetAccessSpecifier::Private)
-			{
-				// Ok to raise EAssetAccessSpecifier::Private to DestScope without calling BulkPrivatizeItems
-				ItemAssetPackage->SetAssetAccessSpecifier(DestScope);
-			}
-			else
-			{
-				// Lowering EAssetAccessSpecifier::Public to EAssetAccessSpecifier::EpicInternal uses BulkPrivatizeItems workflow instead of directly calling SetAssetAccessSpecifier
-				AddForBulkPrivatizeItems(SelectedItem);
+					const int32 Result = CompareAssetAccessSpecifiers(DestScope, ItemAssetData.GetAssetAccessSpecifier());
+					if (Result > 0)
+					{
+						// Ok to asset raise the access specifier to DestScope without calling BulkPrivatizeItems
+						if (UPackage* ItemAssetPackage = ItemAssetData.GetPackage())
+						{
+							if (ItemAssetPackage->SetAssetAccessSpecifier(DestScope))
+							{
+								if (UObject* Asset = ItemAssetData.GetAsset())
+								{
+									// Let the data source know we updated the asset access specifier.
+									Asset->PostEditChange();
+								}
+							}
+						}
+					}
+					else if (Result < 0)
+					{
+						// Lowering asset access specifier uses BulkPrivatizeItems workflow instead of directly calling SetAssetAccessSpecifier
+						FText PrivateErrorMsg;
+						// CanPrivatizeItem checks !GetAssetFolderItemPayload(), CanModifyAssetFileItem(), !IsRunningPIE(), !IsRedirector()
+						if (ItemDataSource->CanPrivatizeItem(ItemData, &PrivateErrorMsg))
+						{
+							if (ItemAssetData.GetPackage()) // Load under the slow task.
+							{
+								TArray<FContentBrowserItemData>& ItemsForSource = PrivatizeSourcesAndItems.FindOrAdd(ItemDataSource);
+								ItemsForSource.Add(ItemData);
+							}
+						}
+						else
+						{
+							AssetViewUtils::ShowErrorNotifcation(PrivateErrorMsg);
+						}
+					}
+				}
 			}
 		}
 	}
-	else if (DestScope == EAssetAccessSpecifier::Private)
-	{
-		// Batch these by their data sources
-		for (const FContentBrowserItem& SelectedItem : SelectedFiles)
-		{
-			AddForBulkPrivatizeItems(SelectedItem);
-		}
-	}
-	else
-	{
-		checkf(false, TEXT("Unknown EAssetAccessSpecifier dest scope: %d"), DestScope);
-	}
 
 	if (!PrivatizeSourcesAndItems.IsEmpty())
 	{
@@ -1456,6 +1394,10 @@ void FAssetContextMenu::ExecuteBulkSetAssetAccessSpecifier(EAssetAccessSpecifier
 		}
 	}
 
+	// Calling PostEditChange is enough to refresh the view, but there is an issue where if you open the context menu again on the same frame,
+	// SetSelectedItems gets called making a copy of the items before they get updated when the UContentBrowserDataSubsystem ticks,
+	// so you end up with stale data.  Refreshing the view here will cause the SContentBrowser to call SetSelectedItems again after it ticks
+	// which fixes the issue, though it creates its own problems since other decisions based on the original selection are cached (e.g. CacheCanExecuteVars).
 	OnAssetViewRefreshRequested.ExecuteIfBound();
 }
 
diff --git a/Engine/Source/Editor/ContentBrowser/Private/AssetContextMenu.h b/Engine/Source/Editor/ContentBrowser/Private/AssetContextMenu.h
--- a/Engine/Source/Editor/ContentBrowser/Private/AssetContextMenu.h
+++ b/Engine/Source/Editor/ContentBrowser/Private/AssetContextMenu.h
@@ -136,8 +136,8 @@ private:
 	/** Try to get AssetAccessSpecifier for asset **/
 	bool GetAssetAccessSpecifierFromSelection(EAssetAccessSpecifier& OutAssetAccessSpecifier);
 
-	/** Try to set AssetAccessSpecifier. Returns true if known to be modified. */
-	bool SetAssetAccessSpecifier(FAssetData& ItemAssetData, const EAssetAccessSpecifier InAssetAccessSpecifier, const bool bEmitEvent);
+	/** Compares two access specifiers.  Returns < 0 if A less accessible than B, > 0 if A more accessible than B, and 0 if they are equal. **/
+	static int32 CompareAssetAccessSpecifiers(EAssetAccessSpecifier A, EAssetAccessSpecifier B);
 
 	/** Handler for when new AssetAccessSpecifier is chosen.  */
 	void ExecuteSetAssetAccessSpecifier(EAssetAccessSpecifier InAssetAccessSpecifier);
diff --git a/Engine/Source/Editor/UnrealEd/Private/AssetPrivatizeModel.cpp b/Engine/Source/Editor/UnrealEd/Private/AssetPrivatizeModel.cpp
--- a/Engine/Source/Editor/UnrealEd/Private/AssetPrivatizeModel.cpp
+++ b/Engine/Source/Editor/UnrealEd/Private/AssetPrivatizeModel.cpp
@@ -108,6 +108,14 @@ void FPendingPrivateAsset::CheckForIllegalReferences()
 	}
 }
 
+void FPendingPrivateAsset::DoPrivatize()
+{
+	if (OwningPackage->SetAssetAccessSpecifier(TargetAssetAccessSpecifier))
+	{
+		Object->PostEditChange();
+	}
+}
+
 FAssetPrivatizeModel::FAssetPrivatizeModel(const TArray<UObject*>& InObjectsToPrivatize, const EAssetAccessSpecifier InAssetAccessSpecifier)
 	: bIsAnythingReferencedInMemoryByNonUndo(false)
 	, bIsAnythingReferencedInMemoryByUndo(false)
@@ -188,10 +196,9 @@ bool FAssetPrivatizeModel::DoPrivatize()
 		return false;
 	}
 
-	const EAssetAccessSpecifier TargetSpecifier = GetTargetAssetAccessSpecifier();
-	for (TSharedPtr<FPendingPrivateAsset>& PendingPrivateAsset : PendingPrivateAssets)
+	for (const TSharedPtr<FPendingPrivateAsset>& PendingPrivateAsset : PendingPrivateAssets)
 	{
-		PendingPrivateAsset->GetOwningPackage()->SetAssetAccessSpecifier(TargetSpecifier);
+		PendingPrivateAsset->DoPrivatize();
 		ObjectsPrivatized++;
 	}
 
@@ -206,10 +213,10 @@ bool FAssetPrivatizeModel::CanForcePrivatize()
 bool FAssetPrivatizeModel::DoForcePrivatize()
 {
 	TArray<UObject*> ObjectsToPrivatize;
+	ObjectsToPrivatize.Reserve(PendingPrivateAssets.Num());
 
-	const EAssetAccessSpecifier TargetSpecifier = GetTargetAssetAccessSpecifier();
 	TSet<UObject*> ObjectsToPrivatizeWithin;
-	for (TSharedPtr<FPendingPrivateAsset>& PendingPrivateAsset : PendingPrivateAssets)
+	for (const TSharedPtr<FPendingPrivateAsset>& PendingPrivateAsset : PendingPrivateAssets)
 	{
 		ObjectsToPrivatize.Add(PendingPrivateAsset->GetObject());
 
@@ -219,7 +226,7 @@ bool FAssetPrivatizeModel::DoForcePrivatize()
 			ObjectsToPrivatizeWithin.Add(ExternalReference.Referencer);
 		}
 
-		PendingPrivateAsset->GetOwningPackage()->SetAssetAccessSpecifier(TargetSpecifier);
+		PendingPrivateAsset->DoPrivatize();
 		ObjectsPrivatized++;
 	}
 
diff --git a/Engine/Source/Editor/UnrealEd/Public/AssetPrivatizeModel.h b/Engine/Source/Editor/UnrealEd/Public/AssetPrivatizeModel.h
--- a/Engine/Source/Editor/UnrealEd/Public/AssetPrivatizeModel.h
+++ b/Engine/Source/Editor/UnrealEd/Public/AssetPrivatizeModel.h
@@ -30,6 +30,9 @@ public:
 	/** Is the pending private object referenced in memory by something other than the undo stack */
 	bool IsReferencedInMemoryByUndo() { return bIsReferencedInMemoryByUndo; }
 
+	/** Performs the privatization action */
+	void DoPrivatize();
+
 	/** The on disk references to this object that would be illegal if the object were made private */
 	TArray<FName> IllegalDiskReferences;
 
