EditorUsability : ContentBrowser
- Allow reordering of favorites list (above or after)
- Disallow move or copy operations of folders in favorite list
- Added edit favorite dialog to set an alias and shortcut to a favorite item (through context menu)
- Reworked editor config class for saving/loading new favorite items (+ migration)
- Moved favorites to be project settings instead of editor settings

#Synced-CL 49068058

diff --git a/Engine/Source/Editor/ContentBrowser/Private/ContentBrowserCommands.cpp b/Engine/Source/Editor/ContentBrowser/Private/ContentBrowserCommands.cpp
--- a/Engine/Source/Editor/ContentBrowser/Private/ContentBrowserCommands.cpp
+++ b/Engine/Source/Editor/ContentBrowser/Private/ContentBrowserCommands.cpp
@@ -6,7 +6,7 @@
 #include "Framework/Commands/UICommandInfo.h"
 #include "GenericPlatform/GenericApplication.h"
 #include "InputCoreTypes.h"
-
+#include "Misc/EnumerateRange.h"
 
 #define LOCTEXT_NAMESPACE "ContentBrowser"
 
@@ -31,6 +31,39 @@ void FContentBrowserCommands::RegisterCommands()
 	UI_COMMAND(GridViewShortcut, "Grid", "View assets as tiles in a grid.", EUserInterfaceActionType::RadioButton, FInputChord());
 	UI_COMMAND(ListViewShortcut, "List", "View assets in a list with thumbnails.", EUserInterfaceActionType::RadioButton, FInputChord());
 	UI_COMMAND(ColumnViewShortcut, "Column", "View assets in a list with columns of details.", EUserInterfaceActionType::RadioButton, FInputChord());
+	
+	// Favorites commands
+	const TArray<FInputChord> InputChords
+	{
+		FInputChord(EModifierKey::Control, EKeys::F1),
+		FInputChord(EModifierKey::Control, EKeys::F2),
+		FInputChord(EModifierKey::Control, EKeys::F3),
+		FInputChord(EModifierKey::Control, EKeys::F4),
+		FInputChord(EModifierKey::Control, EKeys::F5),
+		FInputChord(EModifierKey::Control, EKeys::F6),
+		FInputChord(EModifierKey::Control, EKeys::F7),
+		FInputChord(EModifierKey::Control, EKeys::F8),
+		FInputChord(EModifierKey::Control, EKeys::F9),
+		FInputChord(EModifierKey::Control, EKeys::F10),
+	};
+
+	const FText FavoriteShortcutFormat = LOCTEXT("FavoriteShortcutText", "FavoriteShortcut{0}");
+	for (TConstEnumerateRef<FInputChord> InputChord : EnumerateRange(InputChords))
+	{
+		TSharedPtr<FUICommandInfo>& Command = FavoriteShortcuts.AddDefaulted_GetRef();
+		const FText CommandText = FText::Format(FavoriteShortcutFormat, FText::AsNumber(InputChord.GetIndex()));
+
+		FUICommandInfo::MakeCommandInfo(
+			AsShared(),
+			Command,
+			FName(TEXT("FavoriteShortcut") + FString::FromInt(InputChord.GetIndex())),
+			CommandText,
+			LOCTEXT("FavoriteShortcutTooltip", "Shortcut for navigation that can be assigned to a favorite item"),
+			FSlateIcon(),
+			EUserInterfaceActionType::Button,
+			*InputChord
+		);
+	}
 }
 
 
diff --git a/Engine/Source/Editor/ContentBrowser/Private/ContentBrowserCommands.h b/Engine/Source/Editor/ContentBrowser/Private/ContentBrowserCommands.h
--- a/Engine/Source/Editor/ContentBrowser/Private/ContentBrowserCommands.h
+++ b/Engine/Source/Editor/ContentBrowser/Private/ContentBrowserCommands.h
@@ -41,4 +41,5 @@ public:
 	TSharedPtr<FUICommandInfo> ListViewShortcut;
 	TSharedPtr<FUICommandInfo> ColumnViewShortcut;
 	TSharedPtr<FUICommandInfo> FilterBySelection;
+	TArray<TSharedPtr<FUICommandInfo>> FavoriteShortcuts;
 };
diff --git a/Engine/Source/Editor/ContentBrowser/Private/ContentBrowserConfig.cpp b/Engine/Source/Editor/ContentBrowser/Private/ContentBrowserConfig.cpp
--- a/Engine/Source/Editor/ContentBrowser/Private/ContentBrowserConfig.cpp
+++ b/Engine/Source/Editor/ContentBrowser/Private/ContentBrowserConfig.cpp
@@ -2,10 +2,17 @@
 
 #include "ContentBrowserConfig.h"
 
+#include "ContentBrowserLog.h"
+
 #include UE_INLINE_GENERATED_CPP_BY_NAME(ContentBrowserConfig)
 
 TObjectPtr<UContentBrowserConfig> UContentBrowserConfig::Instance = nullptr;
 
+UContentBrowserFavoriteProjectSettings* UContentBrowserFavoriteProjectSettings::Get()
+{
+	return GetMutableDefault<UContentBrowserFavoriteProjectSettings>();
+}
+
 void UContentBrowserConfig::Initialize()
 {
 	if (!Instance)
diff --git a/Engine/Source/Editor/ContentBrowser/Private/ContentBrowserConfig.h b/Engine/Source/Editor/ContentBrowser/Private/ContentBrowserConfig.h
--- a/Engine/Source/Editor/ContentBrowser/Private/ContentBrowserConfig.h
+++ b/Engine/Source/Editor/ContentBrowser/Private/ContentBrowserConfig.h
@@ -76,19 +76,48 @@ struct FContentBrowserInstanceConfig
 	bool bShowAssetAccessSpecifier = false;
 };
 
+USTRUCT()
+struct FContentBrowserFavoriteItem
+{
+	GENERATED_BODY()
+
+	/** Internal path to the folder/item */
+	UPROPERTY()
+	FString Path;
+
+	/** Friendly name given to this item */
+	UPROPERTY()
+	FString Alias;
+
+	/** Name of the command this item should respond to */
+	UPROPERTY()
+	FName ShortcutCommandName;
+
+	bool operator==(const FContentBrowserFavoriteItem& Other) const
+	{
+		return Path.Equals(Other.Path);
+	}
+
+	bool operator!=(const FContentBrowserFavoriteItem& Other) const
+	{
+		return !(*this == Other);
+	}
+	
+	friend uint32 GetTypeHash(const FContentBrowserFavoriteItem& InItem)
+	{
+		return GetTypeHash(InItem.Path);
+	}
+};
+
 UCLASS(EditorConfig="ContentBrowser")
 class UContentBrowserConfig : public UEditorConfigBase
 {
 	GENERATED_BODY()
 	
 public:
-
 	static void Initialize();
 	static UContentBrowserConfig* Get() { return Instance; }
 
-	UPROPERTY()
-	TSet<FString> Favorites;
-
 	UPROPERTY(meta=(EditorConfig))
 	TMap<FName, FContentBrowserInstanceConfig> Instances;
 
@@ -104,4 +133,17 @@ public:
 
 	UPROPERTY(config, EditAnywhere, Category=Collections, meta = (DisplayName="Exclude these collections from Content Browser"))
 	TArray<FName> ExcludedCollectionsFromView;
+};
+
+UCLASS(MinimalApi, Config = EditorPerProjectUserSettings)
+class UContentBrowserFavoriteProjectSettings : public UDeveloperSettings
+{
+	GENERATED_BODY()
+
+public:
+	static UContentBrowserFavoriteProjectSettings* Get();
+
+	/** Favorite items for a specific project */
+	UPROPERTY(Config)
+	TArray<FContentBrowserFavoriteItem> FavoriteItems;
 };
\ No newline at end of file
diff --git a/Engine/Source/Editor/ContentBrowser/Private/ContentBrowserSingleton.cpp b/Engine/Source/Editor/ContentBrowser/Private/ContentBrowserSingleton.cpp
--- a/Engine/Source/Editor/ContentBrowser/Private/ContentBrowserSingleton.cpp
+++ b/Engine/Source/Editor/ContentBrowser/Private/ContentBrowserSingleton.cpp
@@ -12,6 +12,7 @@
 #include "Containers/Set.h"
 #include "Containers/StringView.h"
 #include "ContentBrowserCommands.h"
+#include "ContentBrowserConfig.h"
 #include "ContentBrowserDataSubsystem.h"
 #include "ContentBrowserModule.h"
 #include "ContentBrowserStyle.h"
@@ -72,6 +73,7 @@
 #include "Widgets/Docking/SDockTab.h"
 #include "Widgets/Layout/SBox.h"
 #include "Widgets/Notifications/SNotificationList.h"
+#include "Widgets/SContentBrowserFavoriteItem.h"
 #include "Widgets/SWidget.h"
 #include "Widgets/SWindow.h"
 #include "WorkspaceMenuStructure.h"
@@ -714,6 +716,68 @@ const FContentBrowserPluginSettings& FContentBrowserSingleton::GetPluginSettings
 	return DefaultSettings;
 }
 
+void FContentBrowserSingleton::OpenEditFavoriteItemWindow(const FContentBrowserFavoriteItem& InItem)
+{
+	// Only a single edit window should be opened
+	if (const TSharedPtr<SWindow> EditFavoriteItemWindow = EditFavoriteItemWindowWeak.Pin())
+	{
+		EditFavoriteItemWindow->RequestDestroyWindow();
+		EditFavoriteItemWindowWeak.Reset();
+	}
+	
+	// Filter shortcuts by removing those already assigned
+	TArray<TSharedPtr<FUICommandInfo>> ShortcutCommands;
+	if (const UContentBrowserFavoriteProjectSettings* FavoriteSettings = UContentBrowserFavoriteProjectSettings::Get())
+	{
+		TSet<FName> ShortcutUsed;
+		ShortcutUsed.Reserve(FavoriteSettings->FavoriteItems.Num());
+		Algo::Transform(FavoriteSettings->FavoriteItems, ShortcutUsed, [](const FContentBrowserFavoriteItem& InItem)
+		{
+			return InItem.ShortcutCommandName;
+		});
+		
+		// Do not remove current shortcut for this item
+		ShortcutUsed.Remove(InItem.ShortcutCommandName);
+
+		FContentBrowserCommands& ContentBrowserCommands = FContentBrowserCommands::Get();
+		ShortcutCommands.Reserve(ContentBrowserCommands.FavoriteShortcuts.Num() - ShortcutUsed.Num());
+		for (const TSharedPtr<FUICommandInfo>& FavoriteShortcut : ContentBrowserCommands.FavoriteShortcuts)
+		{
+			if (FavoriteShortcut 
+				&& !ShortcutUsed.Contains(FavoriteShortcut->GetCommandName()))
+			{
+				ShortcutCommands.Emplace(FavoriteShortcut);
+			}
+		}
+	}
+
+	SContentBrowserFavoriteItem::FArguments Args;
+	Args.OnConfirm_Raw(this, &FContentBrowserSingleton::OnFavoriteItemEdited);
+	Args.CommandOptions(ShortcutCommands);
+	EditFavoriteItemWindowWeak = SContentBrowserFavoriteItem::CreateAndShowWindow(Args, InItem);
+}
+
+void FContentBrowserSingleton::OnFavoriteItemEdited(const FContentBrowserFavoriteItem& InItem, bool bInAnyChanges)
+{
+	if (!bInAnyChanges)
+	{
+		return;
+	}
+
+	const FContentBrowserItemPath ItemPath(InItem.Path, EContentBrowserPathType::Internal);
+	if (FContentBrowserFavoriteItem* FavoriteItem = ContentBrowserUtils::FindFavoriteFolder(ItemPath))
+	{
+		(*FavoriteItem) = InItem;
+
+		if (UContentBrowserFavoriteProjectSettings* FavoriteSettings = UContentBrowserFavoriteProjectSettings::Get())
+		{
+			FavoriteSettings->SaveConfig();
+		}
+
+		BroadcastFavoritesChanged(ItemPath, /** Added */false);
+	}
+}
+
 void FContentBrowserSingleton::SharedCreateAssetDialogWindow(const TSharedRef<SAssetDialog>& AssetDialog, const FSharedAssetDialogConfig& InConfig, bool bModal) const
 {
 	const FVector2D DefaultWindowSize(1152.0f, 648.0f);
@@ -1495,9 +1559,12 @@ void FContentBrowserSingleton::ToggleFolderFavorite(const TArray<FString>& Folde
 
 void FContentBrowserSingleton::ForEachFavoriteFolder(TFunctionRef<void(const FContentBrowserItemPath&)> InDelegate)
 {
-	for (const FString& Path : FavoriteFolderPaths)
+	if (const UContentBrowserFavoriteProjectSettings* FavoriteSettings = UContentBrowserFavoriteProjectSettings::Get())
 	{
-		InDelegate(FContentBrowserItemPath(Path, EContentBrowserPathType::Internal));
+		for (const FContentBrowserFavoriteItem& Item : FavoriteSettings->FavoriteItems)
+		{
+			InDelegate(FContentBrowserItemPath(Item.Path, EContentBrowserPathType::Internal));
+		}
 	}
 }
 
diff --git a/Engine/Source/Editor/ContentBrowser/Private/ContentBrowserSingleton.h b/Engine/Source/Editor/ContentBrowser/Private/ContentBrowserSingleton.h
--- a/Engine/Source/Editor/ContentBrowser/Private/ContentBrowserSingleton.h
+++ b/Engine/Source/Editor/ContentBrowser/Private/ContentBrowserSingleton.h
@@ -1,6 +1,5 @@
 // Copyright Epic Games, Inc. All Rights Reserved.
 
-
 #pragma once
 
 #include "Containers/Array.h"
@@ -20,6 +19,7 @@
 
 #include "ContentBrowserSingleton.generated.h"
 
+struct FContentBrowserFavoriteItem;
 class FCollectionAssetRegistryBridge;
 class FMenuBuilder;
 class FPathPermissionList;
@@ -152,12 +152,14 @@ PRAGMA_ENABLE_DEPRECATION_WARNINGS
 
 	/** Gets the settings for the plugin with the specified name */
 	const FContentBrowserPluginSettings& GetPluginSettings(FName PluginName) const;
-
-	/** Single storage location for content browser favorites */
-	TArray<FString> FavoriteFolderPaths;
+	
+	/** Opens window to edit a favorite item and close old one if any */
+	void OpenEditFavoriteItemWindow(const FContentBrowserFavoriteItem& InItem);
 
 private:
-
+	/** Callback for when favorite item has been edited */
+	void OnFavoriteItemEdited(const FContentBrowserFavoriteItem& InItem, bool bInAnyChanges);
+	
 	/** Util to get or create the content browser that should be used by the various Sync functions */
 	TSharedPtr<SContentBrowser> FindContentBrowserToSync(bool bAllowLockedBrowsers, const FName& InstanceName = FName(), bool bNewSpawnBrowser = false);
 
@@ -232,6 +234,9 @@ private:
 	FShowPrivateContentState ShowEpicInternalContentState;
 
 	FOnFavoritesChanged OnFavoritesChanged;
+	
+	/** Window dialog to edit a favorite item */
+	TWeakPtr<SWindow> EditFavoriteItemWindowWeak;
 
 	TMap<FName, FContentSourceFactory> ContentSourceFactories;
 	FOnContentSourceFactoriesChanged ContentSourceFactoriesChangedEvent;
diff --git a/Engine/Source/Editor/ContentBrowser/Private/ContentBrowserUtils.cpp b/Engine/Source/Editor/ContentBrowser/Private/ContentBrowserUtils.cpp
--- a/Engine/Source/Editor/ContentBrowser/Private/ContentBrowserUtils.cpp
+++ b/Engine/Source/Editor/ContentBrowser/Private/ContentBrowserUtils.cpp
@@ -1,11 +1,11 @@
 // Copyright Epic Games, Inc. All Rights Reserved.
 
-
 #include "ContentBrowserUtils.h"
 
 #include "Algo/MinElement.h"
 #include "Algo/NoneOf.h"
 #include "Algo/Unique.h"
+#include "Algo/Transform.h"
 #include "AssetRegistry/ARFilter.h"
 #include "AssetRegistry/AssetData.h"
 #include "CollectionManagerModule.h"
@@ -21,6 +21,7 @@
 #include "ContentBrowserItem.h"
 #include "ContentBrowserItemData.h"
 #include "ContentBrowserSingleton.h"
+#include "ContentBrowserStyle.h"
 #include "CoreGlobals.h"
 #include "Framework/Application/IMenu.h"
 #include "Framework/Application/MenuStack.h"
@@ -58,6 +59,7 @@
 #include "SPathView.h"
 #include "SlateOptMacros.h"
 #include "SlotBase.h"
+#include "SNavigationBar.h"
 #include "SourceControlOperations.h"
 #include "Styling/AppStyle.h"
 #include "Styling/SlateColor.h"
@@ -72,9 +74,7 @@
 #include "Widgets/SBoxPanel.h"
 #include "Widgets/SCompoundWidget.h"
 #include "Widgets/Text/STextBlock.h"
-#include "Algo/Transform.h"
-#include "ContentBrowserStyle.h"
-#include "SNavigationBar.h"
+#include "Widgets/Views/STableRow.h"
 
 class SWidget;
 struct FGeometry;
@@ -1107,7 +1107,14 @@ bool ContentBrowserUtils::IsFavoriteFolder(const FContentBrowserItemPath& Folder
 	const FName InvariantPath = ContentBrowserUtils::GetInvariantPath(FolderPath);
 	if (!InvariantPath.IsNone())
 	{
-		return FContentBrowserSingleton::Get().FavoriteFolderPaths.Contains(InvariantPath.ToString());
+		if (const UContentBrowserFavoriteProjectSettings* FavoriteSettings = UContentBrowserFavoriteProjectSettings::Get())
+		{
+			const FContentBrowserFavoriteItem Favorite
+			{
+				.Path = InvariantPath.ToString()
+			};
+			return FavoriteSettings->FavoriteItems.Contains(Favorite);
+		}
 	}
 
 	return false;
@@ -1120,52 +1127,201 @@ void ContentBrowserUtils::AddFavoriteFolder(const FString& FolderPath, bool bFlu
 
 void ContentBrowserUtils::AddFavoriteFolder(const FContentBrowserItemPath& FolderPath)
 {
-	const FName InvariantPath = ContentBrowserUtils::GetInvariantPath(FolderPath);
+	// Insert at the end
+	InsertFavoriteFolder(FolderPath, FContentBrowserItemPath(), EItemDropZone::OntoItem);
+}
+
+void ContentBrowserUtils::InsertFavoriteFolder(const FContentBrowserItemPath& InFolderPath, const FContentBrowserItemPath& InFolderContext, EItemDropZone InPositionContext)
+{
+	const FName InvariantPath = ContentBrowserUtils::GetInvariantPath(InFolderPath);
 	if (InvariantPath.IsNone())
 	{
 		return;
 	}
+	
+	const FName TargetFolderPath = ContentBrowserUtils::GetInvariantPath(InFolderContext);
+	const FString TargetFolder = TargetFolderPath.ToString();
 
-	const FString InvariantFolder = InvariantPath.ToString();
-
-	FContentBrowserSingleton::Get().FavoriteFolderPaths.AddUnique(InvariantFolder);
-
-	if (UContentBrowserConfig* EditorConfig = UContentBrowserConfig::Get())
+	// Do not move item on itself
+	if (InvariantPath == TargetFolderPath)
+	{
+		return;
+	}
+	
+	UContentBrowserFavoriteProjectSettings* FavoriteSettings = UContentBrowserFavoriteProjectSettings::Get();
+	if (!FavoriteSettings)
+	{
+		return;
+	}
+	
+	FContentBrowserFavoriteItem FavoriteItem
+	{
+		.Path = InvariantPath.ToString()
+	};
+	const int32 CurrentIndex = FavoriteSettings->FavoriteItems.Find(FavoriteItem);
+	
+	// Item already exists, do not proceed further
+	if (CurrentIndex != INDEX_NONE)
 	{
-		EditorConfig->Favorites.Add(InvariantFolder);
+		return;
+	}
+	
+	int32 TargetIndex = INDEX_NONE;
+	if (!TargetFolderPath.IsNone())
+	{
+		const FContentBrowserFavoriteItem TargetItem
+		{
+			.Path = TargetFolder
+		};
 
-		UContentBrowserConfig::Get()->SaveEditorConfig();
+		TargetIndex = FavoriteSettings->FavoriteItems.Find(TargetItem);
+		
+		// Target item was not found, do not proceed further
+		if (TargetIndex == INDEX_NONE)
+		{
+			return;
+		}
+	}
+	
+	// If no target specified, move it at the end
+	if (TargetIndex == INDEX_NONE)
+	{
+		FavoriteSettings->FavoriteItems.Add(FavoriteItem);
+	}
+	else
+	{
+		FavoriteSettings->FavoriteItems.Insert(
+			FavoriteItem,
+			InPositionContext == EItemDropZone::AboveItem ? TargetIndex : TargetIndex + 1
+		);
 	}
 
-	FContentBrowserSingleton::Get().BroadcastFavoritesChanged(FolderPath, true);
+	FavoriteSettings->SaveConfig();
+	FContentBrowserSingleton::Get().BroadcastFavoritesChanged(InFolderPath, /** Added */true);
 }
 
-void ContentBrowserUtils::RemoveFavoriteFolder(const FContentBrowserItemPath& FolderPath)
+void ContentBrowserUtils::MoveFavoriteFolder(const FContentBrowserItemPath& InFolderPath, const FContentBrowserItemPath& InFolderContext, EItemDropZone InPositionContext)
 {
-	const FName InvariantPath = ContentBrowserUtils::GetInvariantPath(FolderPath);
+	const FName InvariantPath = ContentBrowserUtils::GetInvariantPath(InFolderPath);
 	if (InvariantPath.IsNone())
 	{
 		return;
 	}
+	
+	const FName TargetFolderPath = ContentBrowserUtils::GetInvariantPath(InFolderContext);
+	const FString TargetFolder = TargetFolderPath.ToString();
+
+	// Do not move item on itself
+	if (InvariantPath == TargetFolderPath)
+	{
+		return;
+	}
+	
+	UContentBrowserFavoriteProjectSettings* FavoriteSettings = UContentBrowserFavoriteProjectSettings::Get();
+	if (!FavoriteSettings)
+	{
+		return;
+	}
+	
+	FContentBrowserFavoriteItem FavoriteItem
+	{
+		.Path = InvariantPath.ToString()
+	};
+	const int32 CurrentIndex = FavoriteSettings->FavoriteItems.Find(FavoriteItem);
+	
+	// Item does not exist, do not proceed further
+	if (CurrentIndex == INDEX_NONE)
+	{
+		return;
+	}
+
+	int32 TargetIndex = INDEX_NONE;
+	if (!TargetFolderPath.IsNone())
+	{
+		const FContentBrowserFavoriteItem TargetItem
+		{
+			.Path = TargetFolder
+		};
+
+		TargetIndex = FavoriteSettings->FavoriteItems.Find(TargetItem);
+		
+		// Target item was not found, do not proceed further
+		if (TargetIndex == INDEX_NONE)
+		{
+			return;
+		}
+	}
+
+	// In case of a move, get the actual item with all its options first then remove it to insert it again
+	FavoriteItem = FavoriteSettings->FavoriteItems[CurrentIndex];
+	FavoriteSettings->FavoriteItems.RemoveAt(CurrentIndex, EAllowShrinking::No);
+
+	// If no target specified, move it at the end
+	if (TargetIndex == INDEX_NONE)
+	{
+		FavoriteSettings->FavoriteItems.Add(FavoriteItem);
+	}
+	else
+	{
+		if (TargetIndex > CurrentIndex)
+		{
+			TargetIndex -= 1;
+		}
+
+		FavoriteSettings->FavoriteItems.Insert(
+			FavoriteItem,
+			InPositionContext == EItemDropZone::AboveItem ? TargetIndex : TargetIndex + 1
+		);
+	}
 
-	const FString InvariantFolder = InvariantPath.ToString();
-	bool bModified = false;
-	if (FContentBrowserSingleton::Get().FavoriteFolderPaths.Remove(InvariantFolder) > 0)
+	FavoriteSettings->SaveConfig();
+	FContentBrowserSingleton::Get().BroadcastFavoritesChanged(InFolderPath, /** Added */false);
+}
+
+FContentBrowserFavoriteItem* ContentBrowserUtils::FindFavoriteFolder(const FContentBrowserItemPath& InFolderPath)
+{
+	const FName InvariantPath = ContentBrowserUtils::GetInvariantPath(InFolderPath);
+	if (InvariantPath.IsNone())
 	{
-		bModified = true;
+		return nullptr;
 	}
 
-	if (UContentBrowserConfig* EditorConfig = UContentBrowserConfig::Get())
+	UContentBrowserFavoriteProjectSettings* FavoriteSettings = UContentBrowserFavoriteProjectSettings::Get();
+	if (!FavoriteSettings)
 	{
-		if (EditorConfig->Favorites.Remove(InvariantFolder) > 0)
+		return nullptr;
+	}
+	
+	return FavoriteSettings->FavoriteItems.FindByPredicate(
+		[Path = InvariantPath.ToString()](const FContentBrowserFavoriteItem& InItem)
 		{
-			bModified = true;
-			UContentBrowserConfig::Get()->SaveEditorConfig();
+			return InItem.Path.Equals(Path);
 		}
+	);
+}
+
+void ContentBrowserUtils::RemoveFavoriteFolder(const FContentBrowserItemPath& FolderPath)
+{
+	const FName InvariantPath = ContentBrowserUtils::GetInvariantPath(FolderPath);
+	if (InvariantPath.IsNone())
+	{
+		return;
+	}
+	
+	UContentBrowserFavoriteProjectSettings* FavoriteSettings = UContentBrowserFavoriteProjectSettings::Get();
+	if (!FavoriteSettings)
+	{
+		return;
 	}
 
-	if (bModified)
+	const FContentBrowserFavoriteItem FavoriteItem
 	{
+		.Path = InvariantPath.ToString()
+	};
+
+	if (FavoriteSettings->FavoriteItems.Remove(FavoriteItem) > 0)
+	{
+		FavoriteSettings->SaveConfig();
 		FContentBrowserSingleton::Get().BroadcastFavoritesChanged(FolderPath, false);
 	}
 }
@@ -1175,9 +1331,18 @@ void ContentBrowserUtils::RemoveFavoriteFolder(const FString& FolderPath, bool b
 	RemoveFavoriteFolder(FContentBrowserItemPath(FolderPath, EContentBrowserPathType::Virtual));
 }
 
-const TArray<FString>& ContentBrowserUtils::GetFavoriteFolders()
+TArray<FString> ContentBrowserUtils::GetFavoriteFolders()
 {
-	return FContentBrowserSingleton::Get().FavoriteFolderPaths;
+	TArray<FString> FavoriteFolders;
+	if (const UContentBrowserFavoriteProjectSettings* FavoriteSettings = UContentBrowserFavoriteProjectSettings::Get())
+	{
+		FavoriteFolders.Reserve(FavoriteSettings->FavoriteItems.Num());
+		for (const FContentBrowserFavoriteItem& FavoriteItem : FavoriteSettings->FavoriteItems)
+		{
+			FavoriteFolders.Emplace(FavoriteItem.Path);
+		}
+	}
+	return FavoriteFolders;
 }
 
 void ContentBrowserUtils::AddShowPrivateContentFolder(const FStringView VirtualFolderPath, const FName Owner)
diff --git a/Engine/Source/Editor/ContentBrowser/Private/ContentBrowserUtils.h b/Engine/Source/Editor/ContentBrowser/Private/ContentBrowserUtils.h
--- a/Engine/Source/Editor/ContentBrowser/Private/ContentBrowserUtils.h
+++ b/Engine/Source/Editor/ContentBrowser/Private/ContentBrowserUtils.h
@@ -26,6 +26,7 @@ class SNavigationBar;
 struct FARFilter;
 struct FAssetData;
 struct FContentBrowserDataFilter;
+struct FContentBrowserFavoriteItem;
 struct FContentBrowserInstanceConfig;
 struct FContentBrowserItem;
 struct FContentBrowserItemPath;
@@ -33,7 +34,7 @@ struct FPathViewConfig;
 enum class EContentBrowserIsFolderVisibleFlags : uint8;
 enum class EContentBrowserItemAttributeFilter : uint8;
 enum class EContentBrowserItemCategoryFilter : uint8;
-
+enum class EItemDropZone;
 
 namespace ContentBrowserUtils
 {
@@ -156,12 +157,21 @@ namespace ContentBrowserUtils
 	UE_DEPRECATED(5.3, "Use function that takes FContentBrowserItemPath instead.")
 	void AddFavoriteFolder(const FString& FolderPath, bool bFlushConfig = true);
 	void AddFavoriteFolder(const FContentBrowserItemPath& FolderPath);
+	
+	/** Adds a new favorite folder at a specific position */
+	void InsertFavoriteFolder(const FContentBrowserItemPath& InFolderPath, const FContentBrowserItemPath& InFolderContext, EItemDropZone InPositionContext);
+	
+	/** Moves an existing favorite folder at a new position */
+	void MoveFavoriteFolder(const FContentBrowserItemPath& InFolderPath, const FContentBrowserItemPath& InFolderContext, EItemDropZone InPositionContext);
 
+	/** Finds existing favorite item based on path */
+	FContentBrowserFavoriteItem* FindFavoriteFolder(const FContentBrowserItemPath& InFolderPath);
+	
 	UE_DEPRECATED(5.3, "Use function that takes FContentBrowserItemPath instead.")
 	void RemoveFavoriteFolder(const FString& FolderPath, bool bFlushConfig = true);
 	void RemoveFavoriteFolder(const FContentBrowserItemPath& FolderPath);
 
-	const TArray<FString>& GetFavoriteFolders();
+	TArray<FString> GetFavoriteFolders();
 
 	/** Adds FolderPath as a private content edit folder if it's allowed to be toggled as such */
 	void AddShowPrivateContentFolder(const FStringView VirtualFolderPath, const FName Owner);
diff --git a/Engine/Source/Editor/ContentBrowser/Private/PathContextMenu.cpp b/Engine/Source/Editor/ContentBrowser/Private/PathContextMenu.cpp
--- a/Engine/Source/Editor/ContentBrowser/Private/PathContextMenu.cpp
+++ b/Engine/Source/Editor/ContentBrowser/Private/PathContextMenu.cpp
@@ -76,6 +76,11 @@ void FPathContextMenu::SetOnFolderFavoriteToggled(const FOnFolderFavoriteToggled
 	OnFolderFavoriteToggled = InOnFolderFavoriteToggled;
 }
 
+void FPathContextMenu::SetOnFolderFavoriteEdited(const FOnFolderFavoriteEdited& InOnFolderFavoriteEdited)
+{
+	OnFolderFavoriteEdited = InOnFolderFavoriteEdited;
+}
+
 void FPathContextMenu::SetOnPrivateContentEditToggled(const FOnPrivateContentEditToggled& InOnPrivateContentEditToggled)
 {
 	OnPrivateContentEditToggled = InOnPrivateContentEditToggled;
@@ -221,6 +226,18 @@ void FPathContextMenu::MakePathViewContextMenu(UToolMenu* Menu)
 					FSlateIcon(FAppStyle::GetAppStyleSetName(), "PropertyWindow.Favorites_Disabled"),
 					FUIAction(FExecuteAction::CreateSP(this, &FPathContextMenu::ExecuteFavorite))
 				);
+				
+				if (SelectedFolders.Num() == 1)
+				{
+					// Edit favorite
+					Section.AddMenuEntry(
+						"EditFavorite",
+						LOCTEXT("EditFavorite", "Edit Favorite Item"),
+						LOCTEXT("EditFavoriteTooltip", "Edit and configure favorite item options"),
+						FSlateIcon(FAppStyle::GetAppStyleSetName(), "Icons.Edit"),
+						FUIAction(FExecuteAction::CreateSP(this, &FPathContextMenu::ExecuteEditFavorite))
+					);
+				}
 			}
 			else
 			{
@@ -424,6 +441,14 @@ void FPathContextMenu::ExecuteFavorite()
 	OnFolderFavoriteToggled.ExecuteIfBound(PathsToUpdate);
 }
 
+void FPathContextMenu::ExecuteEditFavorite()
+{
+	if (SelectedFolders.Num() == 1)
+	{
+		OnFolderFavoriteEdited.ExecuteIfBound(SelectedFolders[0].GetItemPath());
+	}
+}
+
 void FPathContextMenu::ExecutePrivateContentEdit()
 {
 	TArray<FString> PathsToUpdate;
diff --git a/Engine/Source/Editor/ContentBrowser/Private/PathContextMenu.h b/Engine/Source/Editor/ContentBrowser/Private/PathContextMenu.h
--- a/Engine/Source/Editor/ContentBrowser/Private/PathContextMenu.h
+++ b/Engine/Source/Editor/ContentBrowser/Private/PathContextMenu.h
@@ -12,6 +12,7 @@
 #include "Math/Color.h"
 #include "Templates/SharedPointer.h"
 
+struct FContentBrowserFavoriteItem;
 class FExtender;
 class FString;
 class SWidget;
@@ -37,6 +38,10 @@ public:
 	/** Delegate for when the context menu has successfully toggled the favorite status of a folder */
 	DECLARE_DELEGATE_OneParam(FOnFolderFavoriteToggled, const TArray<FString>& /*FoldersToToggle*/)
 	void SetOnFolderFavoriteToggled(const FOnFolderFavoriteToggled& InOnFolderFavoriteToggled);
+	
+	/** Delegate for when the context menu has successfully edited a favorite folder */
+	DECLARE_DELEGATE_OneParam(FOnFolderFavoriteEdited, const FContentBrowserItemPath& /*FolderToEdit*/)
+	void SetOnFolderFavoriteEdited(const FOnFolderFavoriteEdited& InOnFolderFavoriteEdited);
 
 	/* Delegate for when the context menu has succesfully toggled the private content edit mode of a folder*/
 	DECLARE_DELEGATE_OneParam(FOnPrivateContentEditToggled, const TArray<FString>& /*FolderVirtualPaths*/)
@@ -83,6 +88,9 @@ public:
 
 	/** Handler for favoriting */
 	void ExecuteFavorite();
+	
+	/** Handler for favorite editing */
+	void ExecuteEditFavorite();
 
 	/* Handler for enabling private content editing*/
 	void ExecutePrivateContentEdit();
@@ -128,5 +136,6 @@ private:
 	FOnRenameFolderRequested OnRenameFolderRequested;
 	FOnFolderDeleted OnFolderDeleted;
 	FOnFolderFavoriteToggled OnFolderFavoriteToggled;
+	FOnFolderFavoriteEdited OnFolderFavoriteEdited;
 	FOnPrivateContentEditToggled OnPrivateContentEditToggled;
 };
diff --git a/Engine/Source/Editor/ContentBrowser/Private/SContentBrowser.cpp b/Engine/Source/Editor/ContentBrowser/Private/SContentBrowser.cpp
--- a/Engine/Source/Editor/ContentBrowser/Private/SContentBrowser.cpp
+++ b/Engine/Source/Editor/ContentBrowser/Private/SContentBrowser.cpp
@@ -136,6 +136,7 @@
 #include "ContentSources/Widgets/SContentSourcesView.h"
 #include "ContentSources/Widgets/SLegacyContentSource.h"
 #include "Experimental/ContentBrowserExtensionUtils.h"
+#include "Misc/EnumerateRange.h"
 #include "Toolkits/GlobalEditorCommonCommands.h"
 #include "Trace/Detail/Channel.h"
 #include "Types/ISlateMetaData.h"
@@ -345,6 +346,9 @@ void SContentBrowser::Construct( const FArguments& InArgs, const FName& InInstan
 
 	JumpMRU.MaxItems = 30;
 
+	// Init favorites settings object
+	UContentBrowserFavoriteProjectSettings::Get();
+
 	UContentBrowserConfig::Initialize();
 	UContentBrowserConfig::Get()->LoadEditorConfig();
 	const FContentBrowserInstanceConfig* EditorConfig = CreateEditorConfigIfRequired();
@@ -478,6 +482,17 @@ void SContentBrowser::BindCommands()
 		FExecuteAction::CreateSP(this, &SContentBrowser::HandleFilterBySelectionCommandExecute),
 		FCanExecuteAction::CreateSP(this, &SContentBrowser::HandleFilterBySelectionCommandCanExecute)
 	));
+	
+	// Favorites shortcuts
+	for (const TSharedPtr<FUICommandInfo>& FavoriteShortcut : ContentBrowserCommands.FavoriteShortcuts)
+	{
+		Commands->MapAction(
+			FavoriteShortcut, 
+			FUIAction(
+				FExecuteAction::CreateSP(this, &SContentBrowser::HandleFavoriteShortcutCommandExecute, FavoriteShortcut->GetCommandName())
+			)
+		);
+	}
 
 	// Allow extenders to add commands
 	FContentBrowserModule& ContentBrowserModule = FModuleManager::GetModuleChecked<FContentBrowserModule>(TEXT("ContentBrowser"));
@@ -514,17 +529,39 @@ EVisibility SContentBrowser::GetLockButtonVisibility() const
 	return IsLocked() ? EVisibility::Visible : EVisibility::Collapsed;
 }
 
-void SContentBrowser::AddFolderFavorite(const TArray<FString>& FolderPaths)
+void SContentBrowser::UpdateFolderFavorite(const TArray<FString>& FolderPaths, const FString& InFolder, EItemDropZone InPosition)
 {
-	for (const FString& FolderPath : FolderPaths)
+	const FContentBrowserItemPath TargetPath(InFolder, EContentBrowserPathType::Virtual);
+	
+	auto InsertOrMoveFolder = [&TargetPath, &InPosition](const FString& InPath)
+		{
+			const FContentBrowserItemPath ItemPath(InPath, EContentBrowserPathType::Virtual);
+			if (ContentBrowserUtils::IsFavoriteFolder(ItemPath))
+			{
+				ContentBrowserUtils::MoveFavoriteFolder(ItemPath, TargetPath, InPosition);
+			}
+			else
+			{
+				ContentBrowserUtils::InsertFavoriteFolder(ItemPath, TargetPath, InPosition);
+			}
+		};
+
+	// When dropping onto/below an item, insert in reverse to keep order
+	if (InPosition == EItemDropZone::OntoItem || InPosition == EItemDropZone::BelowItem)
 	{
-		const FContentBrowserItemPath ItemPath(FolderPath, EContentBrowserPathType::Virtual);
-		if (!ContentBrowserUtils::IsFavoriteFolder(ItemPath))
+		for (int32 Index = FolderPaths.Num() - 1; Index >= 0; --Index)
 		{
-			ContentBrowserUtils::AddFavoriteFolder(ItemPath);
+			InsertOrMoveFolder(FolderPaths[Index]);
 		}
 	}
-
+	else
+	{
+		for (const FString& FolderPath : FolderPaths)
+		{
+			InsertOrMoveFolder(FolderPath);
+		}
+	}
+	
 	SaveAndShowNewFolderFavorites(FolderPaths);
 }
 
@@ -548,6 +585,14 @@ void SContentBrowser::ToggleFolderFavorite(const TArray<FString>& FolderPaths)
 	SaveAndShowNewFolderFavorites(FolderPathsAdded);
 }
 
+void SContentBrowser::EditFolderFavorite(const FContentBrowserItemPath& InItemPath)
+{
+	if (const FContentBrowserFavoriteItem* FavoriteItem = ContentBrowserUtils::FindFavoriteFolder(InItemPath))
+	{
+		FContentBrowserSingleton::Get().OpenEditFavoriteItemWindow(*FavoriteItem);
+	}
+}
+
 void SContentBrowser::SetFilterLayout(const EFilterBarLayout InFilterBarLayout) const
 {
 	if (const TSharedPtr<SFilterList>& FilterBar = LegacyContentSourceWidgets->FilterListPtr; FilterBar && LegacyContentSourceWidgets->AssetViewPtr)
@@ -1158,7 +1203,7 @@ TSharedRef<SWidget> SContentBrowser::CreateFavoritesView(const FContentBrowserCo
 			})
 			.EmptyBodyLabel(LOCTEXT("FavoritesEmpty", "Right click a folder to add it to your favorites."));
 
-	LegacyContentSourceWidgets->FavoritePathViewPtr->SetOnFolderFavoriteAdd(SFavoritePathView::FOnFolderFavoriteAdd::CreateSP(this, &SContentBrowser::AddFolderFavorite));
+	LegacyContentSourceWidgets->FavoritePathViewPtr->SetOnFolderFavoriteUpdated(SFavoritePathView::FOnFolderFavoriteUpdated::CreateSP(this, &SContentBrowser::UpdateFolderFavorite));
 	return ViewWidget;
 }
 
@@ -3338,6 +3383,7 @@ TSharedRef<SWidget> SContentBrowser::CreateLegacyAssetViewWidgets()
 	LegacyContentSourceWidgets->PathContextMenu->SetOnRenameFolderRequested(FPathContextMenu::FOnRenameFolderRequested::CreateSP(this, &SContentBrowser::OnRenameRequested));
 	LegacyContentSourceWidgets->PathContextMenu->SetOnFolderDeleted(FPathContextMenu::FOnFolderDeleted::CreateSP(this, &SContentBrowser::OnOpenedFolderDeleted));
 	LegacyContentSourceWidgets->PathContextMenu->SetOnFolderFavoriteToggled(FPathContextMenu::FOnFolderFavoriteToggled::CreateSP(this, &SContentBrowser::ToggleFolderFavorite));
+	LegacyContentSourceWidgets->PathContextMenu->SetOnFolderFavoriteEdited(FPathContextMenu::FOnFolderFavoriteEdited::CreateSP(this, &SContentBrowser::EditFolderFavorite));
 	LegacyContentSourceWidgets->PathContextMenu->SetOnPrivateContentEditToggled(FPathContextMenu::FOnPrivateContentEditToggled::CreateSP(this, &SContentBrowser::TogglePrivateContentEdit));
 
 	// Currently this controls the asset count
@@ -3926,6 +3972,31 @@ void SContentBrowser::HandleFilterBySelectionCommandExecute()
 	LegacyContentSourceWidgets->FilterListPtr->SetActiveFilters(SelectedAssetPaths);
 }
 
+void SContentBrowser::HandleFavoriteShortcutCommandExecute(FName InCommandName)
+{
+	if (InCommandName.IsNone())
+	{
+		return;
+	}
+
+	const UContentBrowserFavoriteProjectSettings* FavoriteSettings = UContentBrowserFavoriteProjectSettings::Get();
+	if (!FavoriteSettings)
+	{
+		return;
+	}
+
+	// Look for the item with this command assigned
+	for (const FContentBrowserFavoriteItem& FavoriteItem : FavoriteSettings->FavoriteItems)
+	{
+		if (FavoriteItem.ShortcutCommandName == InCommandName)
+		{
+			const FContentBrowserModule& ContentBrowserModule = FModuleManager::Get().LoadModuleChecked<FContentBrowserModule>("ContentBrowser");
+			ContentBrowserModule.Get().SyncBrowserToFolders({ FavoriteItem.Path });
+			break;
+		}
+	}
+}
+
 bool SContentBrowser::HandleRenameCommandCanExecute() const
 {
 	// The order of these conditions are carefully crafted to match the logic of the context menu summoning, as this callback 
diff --git a/Engine/Source/Editor/ContentBrowser/Private/SContentBrowser.h b/Engine/Source/Editor/ContentBrowser/Private/SContentBrowser.h
--- a/Engine/Source/Editor/ContentBrowser/Private/SContentBrowser.h
+++ b/Engine/Source/Editor/ContentBrowser/Private/SContentBrowser.h
@@ -82,6 +82,7 @@ struct FSlateBrush;
 struct FToolMenuContext;
 
 enum class EFilterBarLayout : uint8;
+enum class EItemDropZone;
 
 enum class EContentBrowserViewContext : uint8
 {
@@ -268,6 +269,9 @@ public:
 
 	/** Toggles the favorite status of an array of folders*/
 	void ToggleFolderFavorite(const TArray<FString>& FolderPaths);
+	
+	/** Edit favorite folder options */
+	void EditFolderFavorite(const FContentBrowserItemPath& InItemPath);
 
 	/** Sets the Content Browser filter layout to either vertical or horizontal */
 	void SetFilterLayout(EFilterBarLayout InFilterBarLayout) const;
@@ -455,6 +459,9 @@ private:
 
 	/** Handler for filter by selection */
 	void HandleFilterBySelectionCommandExecute();
+	
+	/** Handler for favorite shortcuts */
+	void HandleFavoriteShortcutCommandExecute(FName InCommandName);
 
 	/** Handler to check to see if a rename command is allowed */
 	bool HandleRenameCommandCanExecute() const;
@@ -610,8 +617,8 @@ private:
 	/** Get the visibility of the lock button */
 	EVisibility GetLockButtonVisibility() const;
 
-	/** Set favorite status of an array of folders unless they're already favorited */
-	void AddFolderFavorite(const TArray<FString>& FolderPaths);
+	/** Updates favorite status of an array of folders */
+	void UpdateFolderFavorite(const TArray<FString>& FolderPaths, const FString& InFolder, EItemDropZone InPosition);
 
 	/** Save the favorite paths to settings and optionally select and show the paths passed in */
 	void SaveAndShowNewFolderFavorites(const TArray<FString>& FolderPaths);
diff --git a/Engine/Source/Editor/ContentBrowser/Private/SPathView.cpp b/Engine/Source/Editor/ContentBrowser/Private/SPathView.cpp
--- a/Engine/Source/Editor/ContentBrowser/Private/SPathView.cpp
+++ b/Engine/Source/Editor/ContentBrowser/Private/SPathView.cpp
@@ -9,6 +9,7 @@
 #include "AssetToolsModule.h"
 #include "Containers/ContainerAllocationPolicies.h"
 #include "Containers/StringView.h"
+#include "ContentBrowserCommands.h"
 #include "ContentBrowserConfig.h"
 #include "ContentBrowserDataDragDropOp.h"
 #include "ContentBrowserDataSource.h"
@@ -2160,7 +2161,7 @@ TSharedRef<ITableRow> SPathView::GenerateTreeRow( TSharedPtr<FTreeItem> TreeItem
 
 	return
 		SNew( STableRow< TSharedPtr<FTreeItem> >, OwnerTable )
-		.OnDragDetected( this, &SPathView::OnFolderDragDetected )
+		.OnDragDetected(this, &SPathView::OnFolderDragDetected)
 		[
 			SNew(SAssetTreeItem)
 			.TreeItem(TreeItem)
@@ -2922,6 +2923,102 @@ void SFavoritePathView::ConfigureTreeView(STreeView<TSharedPtr<FTreeItem>>::FArg
 	// Don't bind some stuff that the parent class binds such as item expansion
 }
 
+void SFavoriteAssetTreeItem::Construct(const FArguments& InArgs)
+{
+	SAssetTreeItem::Construct(InArgs._Parent);
+
+	UpdateItem(InArgs._Parent._TreeItem->GetItem().GetItemPath());
+
+	// Once favorites are updated, check if this item is concerned to update alias
+	FContentBrowserSingleton::Get().RegisterOnFavoritesChangedHandler(
+		FOnFavoritesChanged::FDelegate::CreateSP(this, &SFavoriteAssetTreeItem::OnFavoriteUpdated)
+	);
+}
+
+FText SFavoriteAssetTreeItem::GetNameText() const
+{
+	if (!Alias.IsEmpty())
+	{
+		return FText::Format(LOCTEXT("FavoriteAlias", "@{0}"), Alias);
+	}
+
+	return SAssetTreeItem::GetNameText();
+}
+
+FText SFavoriteAssetTreeItem::GetToolTipText() const
+{
+	FText BaseTooltip = SAssetTreeItem::GetToolTipText();
+
+	if (!Shortcut.IsEmpty())
+	{
+		return FText::Format(LOCTEXT("FavoriteTooltip", "{0}\nShortcut : {1}"), BaseTooltip, Shortcut);
+	}
+
+	return BaseTooltip;
+}
+
+void SFavoriteAssetTreeItem::OnDragEnter(const FGeometry& InGeometry, const FDragDropEvent& InDragDropEvent)
+{
+	if (const TSharedPtr<FContentBrowserDataDragDropOp> DragDropOp = InDragDropEvent.GetOperationAs<FContentBrowserDataDragDropOp>())
+	{
+		DragDropOp->SetToolTip(LOCTEXT("ReorderFavoriteFolder", "Reorder favorite folder"), FAppStyle::GetBrush(TEXT("Graph.ConnectorFeedback.OK")));
+	}
+}
+
+void SFavoriteAssetTreeItem::OnDragLeave(const FDragDropEvent& InDragDropEvent)
+{
+	if (const TSharedPtr<FContentBrowserDataDragDropOp> DragDropOp = InDragDropEvent.GetOperationAs<FContentBrowserDataDragDropOp>())
+	{
+		DragDropOp->SetToolTip(FText::GetEmpty(), nullptr);
+	}
+}
+
+FReply SFavoriteAssetTreeItem::OnDragOver(const FGeometry& InGeometry, const FDragDropEvent& InDragDropEvent)
+{
+	// Let the tree handle this
+	return FReply::Unhandled();
+}
+
+FReply SFavoriteAssetTreeItem::OnDrop(const FGeometry& InGeometry, const FDragDropEvent& InDragDropEvent)
+{
+	// Let the tree handle this
+	return FReply::Unhandled();
+}
+
+void SFavoriteAssetTreeItem::OnFavoriteUpdated(const FContentBrowserItemPath& InItemPath, bool bInAdded)
+{
+	if (const TSharedPtr<FTreeItem> Item = TreeItem.Pin())
+	{
+		if (Item->GetItem().GetItemPath().GetVirtualPathName() == InItemPath.GetVirtualPathName())
+		{
+			UpdateItem(InItemPath);
+		}
+	}
+}
+
+void SFavoriteAssetTreeItem::UpdateItem(const FContentBrowserItemPath& InPath)
+{
+	if (const FContentBrowserFavoriteItem* FavoriteItem = ContentBrowserUtils::FindFavoriteFolder(InPath))
+	{
+		Alias = FText::FromString(FavoriteItem->Alias);
+
+		FContentBrowserCommands& ContentBrowserCommands = FContentBrowserCommands::Get();
+		const TSharedPtr<FUICommandInfo>* FavoriteShortcut = ContentBrowserCommands.FavoriteShortcuts.FindByPredicate([CommandName = FavoriteItem->ShortcutCommandName](const TSharedPtr<FUICommandInfo>& InCommand)
+		{
+			return CommandName == InCommand->GetCommandName();
+		});
+
+		if (FavoriteShortcut && FavoriteShortcut->IsValid())
+		{
+			Shortcut = (*FavoriteShortcut)->GetInputText();
+		}
+		else
+		{
+			Shortcut = FText::GetEmpty();
+		}
+	}
+}
+
 SFavoritePathView::SFavoritePathView()
 {
 	bFlat = true;
@@ -2939,55 +3036,38 @@ void SFavoritePathView::Populate(const bool bIsRefreshingFilter)
 	// Don't allow the selection changed delegate to be fired here
 	FScopedPreventTreeItemChangedDelegate DelegatePrevention(SharedThis(this));
 	TreeData->PopulateWithFavorites(CreateCompiledFolderFilter());
-	TreeData->SortRootItems();
 	TreeData->FilterFullFolderTree();
 
 	SharedPostPopulate();
 }
 
-void SFavoritePathView::SaveSettings(const FString& IniFilename, const FString& IniSection, const FString& SettingsString) const
-{
-	SPathView::SaveSettings(IniFilename, IniSection, SettingsString);
-
-	FString FavoritePathsString;
-	const TArray<FString>& FavoritePaths = ContentBrowserUtils::GetFavoriteFolders();
-	for (const FString& PathIt : FavoritePaths)
-	{
-		if (FavoritePathsString.Len() > 0)
-		{
-			FavoritePathsString += TEXT(",");
-		}
-
-		FavoritePathsString += PathIt;
-	}
-
-	GConfig->SetString(*IniSection, TEXT("FavoritePaths"), *FavoritePathsString, IniFilename);
-	GConfig->Flush(false, IniFilename);
-}
-
 void SFavoritePathView::LoadSettings(const FString& IniFilename, const FString& IniSection, const FString& SettingsString)
 {
 	TGuardValue<bool> Guard(bIsLoadingSettings, true);
 	SPathView::LoadSettings(IniFilename, IniSection, SettingsString);
 
-	// We clear the initial selection for the favorite view, as it conflicts with the main paths view and results in a phantomly selected favorite item
-	ClearSelection();
-
-	// Favorite Paths
-	FString FavoritePathsString;
-	TArray<FString> NewFavoritePaths;
-	if (GConfig->GetString(*IniSection, TEXT("FavoritePaths"), FavoritePathsString, IniFilename))
+	// Migrate favorites to new project settings
+	if (UContentBrowserFavoriteProjectSettings* FavoriteSettings = UContentBrowserFavoriteProjectSettings::Get())
 	{
-		FavoritePathsString.ParseIntoArray(NewFavoritePaths, TEXT(","), /*bCullEmpty*/true);
-	}
-
-	if (NewFavoritePaths.Num() > 0)
-	{
-		// Keep track if we changed at least one source so we know to fire the bulk selection changed delegate later
-		bool bAddedAtLeastOnePath = false;
+		if (FavoriteSettings->FavoriteItems.IsEmpty())
 		{
+			// We clear the initial selection for the favorite view, as it conflicts with the main paths view and results in a phantomly selected favorite item
+			ClearSelection();
+
+			// Favorite Paths
+			FString FavoritePathsString;
+			TArray<FString> NewFavoritePaths;
+			if (GConfig->GetString(*IniSection, TEXT("FavoritePaths"), FavoritePathsString, IniFilename))
+			{
+				FavoritePathsString.ParseIntoArray(NewFavoritePaths, TEXT(","), /*bCullEmpty*/true);
+			}
+			
+			// Keep track if we changed at least one source so we know to fire the bulk selection changed delegate later
+			bool bAddedAtLeastOnePath = false;
+
 			// If the selected paths is empty, the path was "All assets"
 			// This should handle that case properly
+			FavoriteSettings->FavoriteItems.Reserve(NewFavoritePaths.Num());
 			for (const FString& InvariantPath : NewFavoritePaths)
 			{
 				FStringView InvariantPathView(InvariantPath);
@@ -2998,20 +3078,24 @@ void SFavoritePathView::LoadSettings(const FString& IniFilename, const FString&
 					bAddedAtLeastOnePath = true;
 				}
 			}
-		}
 
-		if (bAddedAtLeastOnePath)
-		{
-			Populate();
+			UE_LOG(LogContentBrowser, Log, TEXT("Migrated %i editor favorites to project favorites"), FavoriteSettings->FavoriteItems.Num())
+
+			// Discard previous data
+			GConfig->RemoveKey(*IniSection, TEXT("FavoritePaths"), IniFilename);
+
+			if (bAddedAtLeastOnePath)
+			{
+				Populate();
+			}
 		}
 	}
 }
 
-
 TSharedPtr<FContentBrowserDataDragDropOp> SFavoritePathView::GetContentBrowserDragDropOpFromEvent(const FDragDropEvent& DragDropEvent) const
 {
 	TSharedPtr<FDragDropOperation> Operation = DragDropEvent.GetOperation();
-	if (Operation.IsValid() && OnFolderFavoriteAdd.IsBound())
+	if (Operation.IsValid() && OnFolderFavoriteUpdated.IsBound())
 	{
 		if (Operation->IsOfType<FContentBrowserDataDragDropOp>())
 		{
@@ -3030,7 +3114,6 @@ TSharedPtr<FContentBrowserDataDragDropOp> SFavoritePathView::GetContentBrowserDr
 	return {};
 }
 
-
 void SFavoritePathView::OnDragEnter(const FGeometry& MyGeometry, const FDragDropEvent& DragDropEvent)
 {
 	// If we don't have the appropriate drop content, indicate to the user that nothing will happen.
@@ -3052,49 +3135,164 @@ void SFavoritePathView::OnDragLeave(const FDragDropEvent& DragDropEvent)
 FReply SFavoritePathView::OnDrop(const FGeometry& MyGeometry, const FDragDropEvent& DragDropEvent)
 {
 	TSharedPtr<FContentBrowserDataDragDropOp> DragDropOp = GetContentBrowserDragDropOpFromEvent(DragDropEvent);
-	if (!DragDropOp.IsValid())
+	if (!DragDropOp.IsValid() || !TreeData.IsValid())
 	{
 		return FReply::Unhandled();
 	}
 
-	if (OnFolderFavoriteAdd.IsBound())
+	if (OnFolderFavoriteUpdated.IsBound())
 	{
 		TArray<FString> FolderPaths;
+		FolderPaths.Reserve(DragDropOp->GetDraggedFolders().Num());
 		for (const FContentBrowserItem& BrowserItem: DragDropOp->GetDraggedFolders())
 		{
-			FolderPaths.Add(BrowserItem.GetVirtualPath().ToString());
+			const FName VirtualPath = BrowserItem.GetVirtualPath();
+			if (!TreeData->FindTreeItem(VirtualPath))
+			{
+				FolderPaths.Add(VirtualPath.ToString());
+			}
+		}
+
+		if (!FolderPaths.IsEmpty())
+		{
+			OnFolderFavoriteUpdated.Execute(FolderPaths, FString(), EItemDropZone::BelowItem);
 		}
-		OnFolderFavoriteAdd.Execute(FolderPaths);
 	}
-	
+
 	return FReply::Handled();
 }
 
-void SFavoritePathView::SetOnFolderFavoriteAdd(const FOnFolderFavoriteAdd& InOnFolderFavoriteAdd)
+void SFavoritePathView::SetOnFolderFavoriteUpdated(const FOnFolderFavoriteUpdated& InOnFolderFavoriteUpdated)
 {
-	OnFolderFavoriteAdd = InOnFolderFavoriteAdd;
+	OnFolderFavoriteUpdated = InOnFolderFavoriteUpdated;
 }
 
-
 TSharedRef<ITableRow> SFavoritePathView::GenerateTreeRow(TSharedPtr<FTreeItem> TreeItem, const TSharedRef<STableViewBase>& OwnerTable)
 {
 	check(TreeItem.IsValid());
 
+	SAssetTreeItem::FArguments Arguments;
+	Arguments
+		.TreeItem(TreeItem)
+		.OnNameChanged(this, &SFavoritePathView::FolderNameChanged)
+		.OnVerifyNameChanged(this, &SFavoritePathView::VerifyFolderNameChanged)
+		.IsItemExpanded(false)
+		.HighlightText(this, &SFavoritePathView::GetHighlightText)
+		.IsSelected(this, &SFavoritePathView::IsTreeItemSelected, TreeItem)
+		.FontOverride(UE::ContentBrowser::Private::FContentBrowserStyle::Get().GetFontStyle("ContentBrowser.SourceTreeItemFont"));
+
 	return
-		SNew( STableRow< TSharedPtr<FTreeItem> >, OwnerTable )
-		.OnDragDetected( this, &SFavoritePathView::OnFolderDragDetected )
+		SNew(STableRow<TSharedPtr<FTreeItem>>, OwnerTable)
+		.OnDragDetected(this, &SFavoritePathView::OnFolderDragDetected)
+		.OnCanAcceptDrop(this, &SFavoritePathView::OnFolderCanAcceptDrop)
+		.OnAcceptDrop(this, &SFavoritePathView::OnFolderAcceptDrop)
 		[
-			SNew(SAssetTreeItem)
-			.TreeItem(TreeItem)
-			.OnNameChanged(this, &SFavoritePathView::FolderNameChanged)
-			.OnVerifyNameChanged(this, &SFavoritePathView::VerifyFolderNameChanged)
-			.IsItemExpanded(false)
-			.HighlightText(this, &SFavoritePathView::GetHighlightText)
-			.IsSelected(this, &SFavoritePathView::IsTreeItemSelected, TreeItem)
-			.FontOverride(UE::ContentBrowser::Private::FContentBrowserStyle::Get().GetFontStyle("ContentBrowser.SourceTreeItemFont"))
+			SNew(SFavoriteAssetTreeItem)
+			.Parent(Arguments)
 		];
 }
 
+TOptional<EItemDropZone> SFavoritePathView::OnFolderCanAcceptDrop(const FDragDropEvent& InDragDropEvent, EItemDropZone InItemDropZone, TSharedPtr<FTreeItem> InTreeItem)
+{
+	if (InItemDropZone == EItemDropZone::OntoItem)
+	{
+		return TOptional<EItemDropZone>();
+	}
+
+	TSharedPtr<FContentBrowserDataDragDropOp> DragDropOp = GetContentBrowserDragDropOpFromEvent(InDragDropEvent);
+	if (!DragDropOp.IsValid() || !InTreeItem.IsValid() || !TreeData.IsValid())
+	{
+		return TOptional<EItemDropZone>();
+	}
+
+	const int32 TargetIndex = TreeData->GetVisibleRootItems()->Find(InTreeItem);
+	if (TargetIndex == INDEX_NONE)
+	{
+		return TOptional<EItemDropZone>();
+	}
+
+	TSet<FContentBrowserItemKey> DisallowDropItemKeys
+	{
+		// Don't allow drop on itself
+		FContentBrowserItemKey(InTreeItem->GetItem())
+	};
+
+	if (InItemDropZone == EItemDropZone::BelowItem)
+	{
+		/*
+		 * If we drop onto or below a target item, make sure next item is not contained within dragged items, eg:
+		 * - A <-
+		 * - B
+		 * - C
+		 * - D
+		 * If B,C are dragged, disallow dropping after A and before B, only before A and after D
+		 */
+		if (TreeData->GetVisibleRootItems()->IsValidIndex(TargetIndex + 1))
+		{
+			if (const TSharedPtr<FTreeItem> NextItem = (*TreeData->GetVisibleRootItems())[TargetIndex + 1])
+			{
+				const FContentBrowserItemKey NextItemKey(NextItem->GetItem());
+				DisallowDropItemKeys.Add(NextItemKey);
+			}
+		}
+	}
+	else
+	{
+		/*
+		 * If we drop above a target item, make sure prev item is not contained within dragged items, eg:
+		 * - A
+		 * - B
+		 * - C
+		 * - D <-
+		 * If B,C are dragged, disallow dropping after C and before D, only before A and after D
+		 */
+		if (TreeData->GetVisibleRootItems()->IsValidIndex(TargetIndex - 1))
+		{
+			if (const TSharedPtr<FTreeItem> PrevItem = (*TreeData->GetVisibleRootItems())[TargetIndex - 1])
+			{
+				const FContentBrowserItemKey PrevItemKey(PrevItem->GetItem());
+				DisallowDropItemKeys.Add(PrevItemKey);
+			}
+		}
+	}
+
+	const bool bDropZoneDisallowed = DragDropOp->GetDraggedFolders().ContainsByPredicate([&DisallowDropItemKeys](const FContentBrowserItem& InDraggedFolder)
+		{
+			return DisallowDropItemKeys.Contains(FContentBrowserItemKey(InDraggedFolder));
+		});
+
+	if (bDropZoneDisallowed)
+	{
+		return TOptional<EItemDropZone>();
+	}
+
+	return InItemDropZone;
+}
+
+FReply SFavoritePathView::OnFolderAcceptDrop(const FDragDropEvent& InDragDropEvent, EItemDropZone InItemDropZone, TSharedPtr<FTreeItem> InTreeItem)
+{
+	TSharedPtr<FContentBrowserDataDragDropOp> DragDropOp = GetContentBrowserDragDropOpFromEvent(InDragDropEvent);
+	if (!DragDropOp.IsValid() || !InTreeItem.IsValid())
+	{
+		return FReply::Unhandled();
+	}
+
+	if (OnFolderFavoriteUpdated.IsBound())
+	{
+		TArray<FString> FolderPaths;
+		FolderPaths.Reserve(DragDropOp->GetDraggedFolders().Num());
+		for (const FContentBrowserItem& BrowserItem: DragDropOp->GetDraggedFolders())
+		{
+			FolderPaths.Add(BrowserItem.GetVirtualPath().ToString());
+		}
+
+		const FString FolderContext = InTreeItem->GetItem().GetVirtualPath().ToString();
+		OnFolderFavoriteUpdated.Execute(FolderPaths, FolderContext, InItemDropZone);
+	}
+
+	return FReply::Handled().EndDragDrop();
+}
+
 void SFavoritePathView::OnFavoriteAdded(const FContentBrowserItemPath&, bool)
 {
 	if (!bIsLoadingSettings)
diff --git a/Engine/Source/Editor/ContentBrowser/Private/SPathView.h b/Engine/Source/Editor/ContentBrowser/Private/SPathView.h
--- a/Engine/Source/Editor/ContentBrowser/Private/SPathView.h
+++ b/Engine/Source/Editor/ContentBrowser/Private/SPathView.h
@@ -3,6 +3,7 @@
 #pragma once
 
 #include "AssetViewUtils.h"
+#include "ContentBrowserConfig.h"
 #include "Containers/Array.h"
 #include "Containers/ArrayView.h"
 #include "Containers/BitArray.h"
@@ -28,6 +29,7 @@
 #include "Misc/Optional.h"
 #include "Misc/TextFilter.h"
 #include "PathViewTypes.h"
+#include "SourcesViewWidgets.h"
 #include "Templates/SharedPointer.h"
 #include "Templates/TypeHash.h"
 #include "Templates/UnrealTemplate.h"
@@ -40,6 +42,7 @@
 #include "Widgets/Views/STableViewBase.h"
 #include "Widgets/Views/STreeView.h"
 
+struct FContentBrowserFavoriteItem;
 class FContentBrowserDataDragDropOp;
 class FContentBrowserItemData;
 class FContentBrowserItemDataUpdate;
@@ -599,6 +602,34 @@ private:
 	TSharedPtr<SExpandableArea> FavoritesArea;
 };
 
+/** AssetTreeItem for the Favorite view to override some behavior */
+class SFavoriteAssetTreeItem : public SAssetTreeItem
+{
+	SLATE_BEGIN_ARGS(SFavoriteAssetTreeItem){}
+		SLATE_ARGUMENT(SAssetTreeItem::FArguments, Parent)
+	SLATE_END_ARGS()
+
+	void Construct(const FArguments& InArgs);
+
+protected:
+	//~ Begin SAssetTreeItem
+	virtual FText GetNameText() const override;
+	virtual FText GetToolTipText() const override;
+	//~ End SAssetTreeItem
+
+	//~ Begin SWidget
+	virtual void OnDragEnter(const FGeometry& InGeometry, const FDragDropEvent& InDragDropEvent) override;
+	virtual void OnDragLeave(const FDragDropEvent& InDragDropEvent) override;
+	virtual FReply OnDragOver(const FGeometry& InGeometry, const FDragDropEvent& InDragDropEvent) override;
+	virtual FReply OnDrop(const FGeometry& InGeometry, const FDragDropEvent& InDragDropEvent) override;
+	//~ End SWidget
+
+	void OnFavoriteUpdated(const FContentBrowserItemPath& InItemPath, bool bInAdded);
+	void UpdateItem(const FContentBrowserItemPath& InPath);
+
+	FText Alias;
+	FText Shortcut;
+};
 
 template<>
 struct TWidgetTypeTraits<class SFavoritePathView>
@@ -613,16 +644,13 @@ class SFavoritePathView : public SPathView
 {
 public:
 	SFavoritePathView();
-	virtual ~SFavoritePathView();
+	virtual ~SFavoritePathView() override;
 	
 	/** Constructs this widget with InArgs */
 	virtual void Construct(const FArguments& InArgs) override;
 
 	virtual void Populate(const bool bIsRefreshingFilter = false) override;
 
-	/** Saves any settings to config that should be persistent between editor sessions */
-	virtual void SaveSettings(const FString& IniFilename, const FString& IniSection, const FString& SettingsString) const override;
-
 	/** Loads any settings to config that should be persistent between editor sessions */
 	virtual void LoadSettings(const FString& IniFilename, const FString& IniSection, const FString& SettingsString) override;
 
@@ -638,8 +666,8 @@ public:
 	/** Updates favorites based on an external change. */
 	void FixupFavoritesFromExternalChange(TArrayView<const AssetViewUtils::FMovedContentFolder> MovedFolders);
 
-	DECLARE_DELEGATE_OneParam(FOnFolderFavoriteAdd, const TArray<FString>& /*FoldersToAdd*/)
-	void SetOnFolderFavoriteAdd(const FOnFolderFavoriteAdd& InOnFolderFavoriteAdd);
+	DECLARE_DELEGATE_ThreeParams(FOnFolderFavoriteUpdated, const TArray<FString>& /** UpdatedFolders */, const FString& /** FolderContext */, EItemDropZone /** PositionContext */)
+	void SetOnFolderFavoriteUpdated(const FOnFolderFavoriteUpdated& InOnFolderFavoriteUpdated);
 
 protected:
 	constexpr virtual bool ShouldDoJustInTimePopulation() const override
@@ -650,6 +678,12 @@ protected:
 private:
 	virtual TSharedRef<ITableRow> GenerateTreeRow(TSharedPtr<FTreeItem> TreeItem, const TSharedRef<STableViewBase>& OwnerTable) override;
 
+	/** Handler for tree view to accept drop */
+	TOptional<EItemDropZone> OnFolderCanAcceptDrop(const FDragDropEvent& InDragDropEvent, EItemDropZone InItemDropZone, TSharedPtr<FTreeItem> InTreeItem);
+
+	/** Handler for tree view when dropping items */
+	FReply OnFolderAcceptDrop(const FDragDropEvent& InDragDropEvent, EItemDropZone InItemDropZone, TSharedPtr<FTreeItem> InTreeItem);
+
 	void OnFavoriteAdded(const FContentBrowserItemPath& ChangedItem, bool bAdded);
 
 	/** Handles updating the view when content items are changed */
@@ -658,12 +692,11 @@ private:
 	virtual void ConfigureTreeView(STreeView<TSharedPtr<FTreeItem>>::FArguments& InArgs) override;
 
 	/** Returns an FContentBrowserDataDragDropOp object but only if it qualifies as a proper droppable content browser drag drop op */
-	TSharedPtr<FContentBrowserDataDragDropOp> GetContentBrowserDragDropOpFromEvent(const FDragDropEvent& DragDropEvent) const; 
+	TSharedPtr<FContentBrowserDataDragDropOp> GetContentBrowserDragDropOpFromEvent(const FDragDropEvent& DragDropEvent) const;
 
-private:
 	TArray<FString> RemovedByFolderMove;
 	FDelegateHandle OnFavoritesChangedHandle;
-	
-	FOnFolderFavoriteAdd OnFolderFavoriteAdd; 
+
+	FOnFolderFavoriteUpdated OnFolderFavoriteUpdated; 
 	bool bIsLoadingSettings = false;
 };
diff --git a/Engine/Source/Editor/ContentBrowser/Private/SourcesViewWidgets.h b/Engine/Source/Editor/ContentBrowser/Private/SourcesViewWidgets.h
--- a/Engine/Source/Editor/ContentBrowser/Private/SourcesViewWidgets.h
+++ b/Engine/Source/Editor/ContentBrowser/Private/SourcesViewWidgets.h
@@ -74,14 +74,19 @@ public:
 	/** Constructs this widget with InArgs */
 	void Construct( const FArguments& InArgs );
 
-	~SAssetTreeItem();
+	virtual ~SAssetTreeItem() override;
 
+	//~ Begin SWidget
 	virtual void OnDragEnter( const FGeometry& MyGeometry, const FDragDropEvent& DragDropEvent ) override;
 	virtual void OnDragLeave( const FDragDropEvent& DragDropEvent ) override;
 	virtual FReply OnDragOver( const FGeometry& MyGeometry, const FDragDropEvent& DragDropEvent ) override;
 	virtual FReply OnDrop( const FGeometry& MyGeometry, const FDragDropEvent& DragDropEvent ) override;
+	//~ End SWidget
+
+protected:
+	/** Returns the text of the folder name */
+	virtual FText GetNameText() const;
 
-private:
 	/** Handles verifying name changes */
 	bool VerifyNameChanged(const FText& InName, FText& OutError) const;
 
@@ -97,19 +102,15 @@ private:
 	/** Gets the color used to draw the folder icon */
 	FSlateColor GetFolderColor() const;
 
-	/** Returns the text of the folder name */
-	FText GetNameText() const;
-
 	/** Returns the text to use for the folder tooltip */
-	FText GetToolTipText() const;
+	virtual FText GetToolTipText() const;
 
 	/** Returns the image for the border around this item. Used for drag/drop operations */
 	const FSlateBrush* GetBorderImage() const;
 
 	/** Returns a widget that combines the results of calling all registered PathViewStateIconGenerators **/
 	TSharedRef<SWidget> GenerateStateIcons();
-	
-private:
+
 	/** The data for this item */
 	TWeakPtr<FTreeItem> TreeItem;
 
@@ -141,7 +142,6 @@ private:
 	FDelegateHandle EnterEditingModeDelegateHandle;
 };
 
-
 template <>
 struct TWidgetTypeTraits<class SCollectionTreeItem>
 {
diff --git a/Engine/Source/Editor/ContentBrowser/Private/Widgets/SContentBrowserFavoriteItem.cpp b/Engine/Source/Editor/ContentBrowser/Private/Widgets/SContentBrowserFavoriteItem.cpp
new file mode 100644
--- /dev/null
+++ b/Engine/Source/Editor/ContentBrowser/Private/Widgets/SContentBrowserFavoriteItem.cpp
@@ -0,0 +1,228 @@
+// Copyright Epic Games, Inc. All Rights Reserved.
+
+#include "Widgets/SContentBrowserFavoriteItem.h"
+
+#include "Framework/Application/SlateApplication.h"
+#include "Framework/Commands/UICommandInfo.h"
+#include "SPrimaryButton.h"
+#include "Widgets/Input/SEditableTextBox.h"
+#include "Widgets/Input/SComboBox.h"
+#include "Widgets/Input/SButton.h"
+#include "Widgets/Layout/SBorder.h"
+#include "Widgets/SBoxPanel.h"
+#include "Widgets/SWindow.h"
+#include "Widgets/Text/STextBlock.h"
+
+#define LOCTEXT_NAMESPACE "ContentBrowserFavoriteItem"
+
+TSharedRef<SWindow> SContentBrowserFavoriteItem::CreateAndShowWindow(const FArguments& InArguments, const FContentBrowserFavoriteItem& InItemToEdit)
+{
+	TSharedRef<SWindow> NewWindow = SNew(SWindow)
+		.Title(LOCTEXT("EditFavoriteItemWindow", "Edit Favorite Item"))
+		.HasCloseButton(true)
+		.SupportsMaximize(false)
+		.SupportsMinimize(false)
+		.SizingRule(ESizingRule::FixedSize)
+		.ClientSize(FVector2D(400, 150));
+
+	const TSharedRef<SContentBrowserFavoriteItem> Content = SArgumentNew(InArguments, SContentBrowserFavoriteItem, InItemToEdit);
+	Content->ParentWindowWeak = NewWindow;
+
+	NewWindow->SetContent(Content);
+	FSlateApplication::Get().AddWindow(NewWindow);
+
+	return NewWindow;
+}
+
+FText SContentBrowserFavoriteItem::GetAliasText() const
+{
+	return FText::FromString(FavoriteItemEdited.Alias);
+}
+
+void SContentBrowserFavoriteItem::OnAliasTextChanged(const FText& InText)
+{
+	FavoriteItemEdited.Alias = InText.ToString();
+	bAnyChanges = true;
+}
+
+void SContentBrowserFavoriteItem::Construct(const FArguments& InArgs, const FContentBrowserFavoriteItem& InItemToEdit)
+{
+	FavoriteItemEdited = InItemToEdit;
+	CommandOptions = InArgs._CommandOptions;
+	OnConfirmClicked = InArgs._OnConfirm;
+	OnCancelClicked = InArgs._OnCancel;
+	bAnyChanges = false;
+	
+	// Empty choice
+	CommandOptions.Add(MakeShared<FUICommandInfo>(NAME_None));
+	const TSharedPtr<FUICommandInfo>* SelectedCommandPtr = CommandOptions.FindByPredicate([this, CommandName = FavoriteItemEdited.ShortcutCommandName](const TSharedPtr<FUICommandInfo>& InCommand)
+	{
+		return CommandName == InCommand->GetCommandName();
+	});
+
+	SelectedCommand = SelectedCommandPtr ? *SelectedCommandPtr : nullptr;
+
+	constexpr float HPadding = 20.f;
+	constexpr float VPadding = 10.f;
+
+	ChildSlot
+	.HAlign(HAlign_Fill)
+	[
+		SNew(SBorder)
+		.BorderImage(FAppStyle::GetBrush("Brushes.Panel"))
+		[
+			SNew(SVerticalBox)
+			// Path (read-only)
+			+ SVerticalBox::Slot()
+			.HAlign(HAlign_Fill)
+			.Padding(HPadding, VPadding, HPadding, 0.f)
+			.AutoHeight()
+			[
+				SNew(SHorizontalBox)
+				+ SHorizontalBox::Slot()
+				.FillWidth(1.f)
+				.Padding(0.f)
+				[
+					SNew(STextBlock)
+					.Text(LOCTEXT("PathText", "Path"))
+					.ToolTipText(LOCTEXT("PathTooltip", "Path of this favorite item"))
+				]
+				+ SHorizontalBox::Slot()
+				.FillWidth(2.f)
+				.Padding(0.f)
+				[
+					SNew(STextBlock)
+					.Text(FText::FromString(FavoriteItemEdited.Path))
+				]
+			]
+			// Alias
+			+ SVerticalBox::Slot()
+			.HAlign(HAlign_Fill)
+			.Padding(HPadding, VPadding, HPadding, 0.f)
+			.AutoHeight()
+			[
+				SNew(SHorizontalBox)
+				+ SHorizontalBox::Slot()
+				.FillWidth(1.f)
+				.Padding(0.f)
+				[
+					SNew(STextBlock)
+					.Text(LOCTEXT("AliasText", "Alias"))
+					.ToolTipText(LOCTEXT("AliasTooltip", "Alias of this favorite item"))
+				]
+				+ SHorizontalBox::Slot()
+				.FillWidth(2.f)
+				.Padding(0.f)
+				[
+					SNew(SEditableTextBox)
+					.Text(this, &SContentBrowserFavoriteItem::GetAliasText)
+					.OnTextChanged(this, &SContentBrowserFavoriteItem::OnAliasTextChanged)
+				]
+			]
+			// Shortcut
+			+ SVerticalBox::Slot()
+			.HAlign(HAlign_Fill)
+			.Padding(HPadding, VPadding, HPadding, 0.f)
+			.AutoHeight()
+			[
+				SNew(SHorizontalBox)
+				+ SHorizontalBox::Slot()
+				.FillWidth(1.f)
+				.Padding(0.f)
+				[
+					SNew(STextBlock)
+					.Text(LOCTEXT("ShortcutText", "Shortcut"))
+					.ToolTipText(LOCTEXT("ShortcutTooltip", "Keyboard shortcut assigned to navigate to this favorite item"))
+				]
+				+ SHorizontalBox::Slot()
+				.FillWidth(2.f)
+				.Padding(0.f)
+				[
+					SNew(SComboBox<TSharedPtr<FUICommandInfo>>)
+					.OptionsSource(&CommandOptions)
+					.InitiallySelectedItem(SelectedCommand)
+					.OnGenerateWidget(this, &SContentBrowserFavoriteItem::OnCommandGenerateWidget)
+					.OnSelectionChanged(this, &SContentBrowserFavoriteItem::OnCommandSelectionChanged)
+					[
+						SNew(STextBlock)
+						.Text(this, &SContentBrowserFavoriteItem::GetItemCommandLabel)
+					]
+				]
+			]
+			// Buttons
+			+ SVerticalBox::Slot()
+			.VAlign(VAlign_Bottom)
+			.Padding(HPadding, VPadding, HPadding, 0.f)
+			[
+				SNew(SHorizontalBox)
+				+ SHorizontalBox::Slot()
+				.FillWidth(1.0f)
+				.Padding(5.f)
+				.HAlign(HAlign_Right)
+				[
+					SNew(SButton)
+					.Text(LOCTEXT("CancelButton", "Cancel"))
+					.OnClicked(this, &SContentBrowserFavoriteItem::OnCancelButtonClicked)
+				]
+				+ SHorizontalBox::Slot()
+				.AutoWidth()
+				.Padding(5.f)
+				.HAlign(HAlign_Right)
+				[
+					SNew(SPrimaryButton)
+					.Text(LOCTEXT("ConfirmButton", "Confirm"))
+					.OnClicked(this, &SContentBrowserFavoriteItem::OnConfirmButtonClicked)
+				]
+			]
+		]
+	];
+}
+
+TSharedRef<SWidget> SContentBrowserFavoriteItem::OnCommandGenerateWidget(TSharedPtr<FUICommandInfo> InCommand)
+{
+	return SNew(STextBlock).Text(GetCommandLabel(InCommand));
+}
+
+void SContentBrowserFavoriteItem::OnCommandSelectionChanged(TSharedPtr<FUICommandInfo> InCommand, ESelectInfo::Type InSelectType)
+{
+	SelectedCommand = InCommand;
+	FavoriteItemEdited.ShortcutCommandName = InCommand ? InCommand->GetCommandName() : NAME_None;
+	bAnyChanges = true;
+}
+
+FText SContentBrowserFavoriteItem::GetItemCommandLabel() const
+{
+	return GetCommandLabel(SelectedCommand);
+}
+
+FText SContentBrowserFavoriteItem::GetCommandLabel(TSharedPtr<FUICommandInfo> InCommand) const
+{
+	if (!InCommand || InCommand->GetCommandName().IsNone())
+	{
+		return LOCTEXT("NoFavoriteShortcutLabel", "No shortcut");
+	}
+
+	return InCommand->GetInputText();
+}
+
+FReply SContentBrowserFavoriteItem::OnConfirmButtonClicked() const
+{
+	OnConfirmClicked.ExecuteIfBound(FavoriteItemEdited, bAnyChanges);
+	if (const TSharedPtr<SWindow> ParentWindow = ParentWindowWeak.Pin())
+	{
+		ParentWindow->RequestDestroyWindow();
+	}
+	return FReply::Handled();
+}
+
+FReply SContentBrowserFavoriteItem::OnCancelButtonClicked() const
+{
+	OnCancelClicked.ExecuteIfBound();
+	if (const TSharedPtr<SWindow> ParentWindow = ParentWindowWeak.Pin())
+	{
+		ParentWindow->RequestDestroyWindow();
+	}
+	return FReply::Handled();
+}
+
+#undef LOCTEXT_NAMESPACE
diff --git a/Engine/Source/Editor/ContentBrowser/Private/Widgets/SContentBrowserFavoriteItem.h b/Engine/Source/Editor/ContentBrowser/Private/Widgets/SContentBrowserFavoriteItem.h
new file mode 100644
--- /dev/null
+++ b/Engine/Source/Editor/ContentBrowser/Private/Widgets/SContentBrowserFavoriteItem.h
@@ -0,0 +1,69 @@
+// Copyright Epic Games, Inc. All Rights Reserved.
+
+#pragma once
+
+#include "Containers/Array.h"
+#include "ContentBrowserConfig.h"
+#include "Templates/SharedPointer.h"
+#include "Types/SlateEnums.h"
+#include "Widgets/SCompoundWidget.h"
+
+class FReply;
+class FUICommandInfo;
+class SEditableTextBox;
+class SWindow;
+
+/** 
+ * Content browser favorite item dialog form to edit options
+ */
+class SContentBrowserFavoriteItem : public SCompoundWidget
+{
+public:
+	DECLARE_DELEGATE_TwoParams(FOnConfirmClicked, const FContentBrowserFavoriteItem& /* Item */, bool /** AnyChanges */);
+	DECLARE_DELEGATE(FOnCancelClicked);
+
+	SLATE_BEGIN_ARGS(SContentBrowserFavoriteItem)
+		{}
+		/** Commands to pick from for this item */
+		SLATE_ARGUMENT(TArray<TSharedPtr<FUICommandInfo>>, CommandOptions)
+		/** Delegate for when the confirm button is clicked */
+		SLATE_EVENT(FOnConfirmClicked, OnConfirm)
+		/** Delegate for when the cancel button is clicked */
+		SLATE_EVENT(FOnCancelClicked, OnCancel)
+    SLATE_END_ARGS()
+
+	/** Create a window with this widget and shows it */
+	static TSharedRef<SWindow> CreateAndShowWindow(const FArguments& InArguments, const FContentBrowserFavoriteItem& InItemToEdit);
+
+    void Construct(const FArguments& InArgs, const FContentBrowserFavoriteItem& InItemToEdit);
+
+protected:
+	TSharedRef<SWidget> OnCommandGenerateWidget(TSharedPtr<FUICommandInfo> InCommand);
+	void OnCommandSelectionChanged(TSharedPtr<FUICommandInfo> InCommand, ESelectInfo::Type InSelectType);
+	FText GetItemCommandLabel() const;
+	FText GetCommandLabel(TSharedPtr<FUICommandInfo> InCommand) const;
+	
+	FText GetAliasText() const;
+	void OnAliasTextChanged(const FText& InText);
+
+	/** Called when cancel button is clicked */
+	FReply OnConfirmButtonClicked() const;
+
+	/** Called when cancel button is clicked */
+	FReply OnCancelButtonClicked() const;
+	
+	TArray<TSharedPtr<FUICommandInfo>> CommandOptions;
+	TSharedPtr<FUICommandInfo> SelectedCommand;
+
+	FOnConfirmClicked OnConfirmClicked;
+	FOnCancelClicked OnCancelClicked;
+	
+	/** Current item being edited */
+	FContentBrowserFavoriteItem FavoriteItemEdited;
+	
+	/** Self containing window */
+	TWeakPtr<SWindow> ParentWindowWeak;
+	
+	/** Whether any changes have been done to this item */
+	bool bAnyChanges = false;
+};
\ No newline at end of file
