EditorUsability : ContentBrowser
- Fixed favorites not refreshing on load

#Synced-CL 49068058

diff --git a/Engine/Source/Editor/ContentBrowser/Private/SPathView.cpp b/Engine/Source/Editor/ContentBrowser/Private/SPathView.cpp
--- a/Engine/Source/Editor/ContentBrowser/Private/SPathView.cpp
+++ b/Engine/Source/Editor/ContentBrowser/Private/SPathView.cpp
@@ -3049,11 +3049,11 @@ void SFavoritePathView::LoadSettings(const FString& IniFilename, const FString&
 	// Migrate favorites to new project settings
 	if (UContentBrowserFavoriteProjectSettings* FavoriteSettings = UContentBrowserFavoriteProjectSettings::Get())
 	{
+		// We clear the initial selection for the favorite view, as it conflicts with the main paths view and results in a phantomly selected favorite item
+		ClearSelection();
+
 		if (FavoriteSettings->FavoriteItems.IsEmpty())
 		{
-			// We clear the initial selection for the favorite view, as it conflicts with the main paths view and results in a phantomly selected favorite item
-			ClearSelection();
-
 			// Favorite Paths
 			FString FavoritePathsString;
 			TArray<FString> NewFavoritePaths;
@@ -3061,9 +3061,6 @@ void SFavoritePathView::LoadSettings(const FString& IniFilename, const FString&
 			{
 				FavoritePathsString.ParseIntoArray(NewFavoritePaths, TEXT(","), /*bCullEmpty*/true);
 			}
-			
-			// Keep track if we changed at least one source so we know to fire the bulk selection changed delegate later
-			bool bAddedAtLeastOnePath = false;
 
 			// If the selected paths is empty, the path was "All assets"
 			// This should handle that case properly
@@ -3075,7 +3072,6 @@ void SFavoritePathView::LoadSettings(const FString& IniFilename, const FString&
 				if (!InvariantPathView.IsEmpty() && InvariantPathView != TEXT("None"))
 				{
 					ContentBrowserUtils::AddFavoriteFolder(FContentBrowserItemPath(InvariantPathView, EContentBrowserPathType::Internal));
-					bAddedAtLeastOnePath = true;
 				}
 			}
 
@@ -3083,11 +3079,11 @@ void SFavoritePathView::LoadSettings(const FString& IniFilename, const FString&
 
 			// Discard previous data
 			GConfig->RemoveKey(*IniSection, TEXT("FavoritePaths"), IniFilename);
+		}
 
-			if (bAddedAtLeastOnePath)
-			{
-				Populate();
-			}
+		if (!FavoriteSettings->FavoriteItems.IsEmpty())
+		{
+			Populate();
 		}
 	}
 }
