Fixing the content browser not selecting the previous path propely on reopen. UE-309355 is no longer an issue since we always delay the population.

#Synced-CL 49167278

diff --git a/Engine/Source/Editor/ContentBrowser/Private/SPathView.cpp b/Engine/Source/Editor/ContentBrowser/Private/SPathView.cpp
--- a/Engine/Source/Editor/ContentBrowser/Private/SPathView.cpp
+++ b/Engine/Source/Editor/ContentBrowser/Private/SPathView.cpp
@@ -2361,9 +2361,9 @@ void SPathView::Populate(const bool bIsRefreshingFilter)
 	const bool bFilteringByText = !TreeData->GetFolderPathTextFilter().GetRawFilterText().IsEmpty();
 
 	// Batch the selection changed event
-	// Only emit events when the user isn't filtering or doing the first population.
+	// Only emit events when the user isn't filtering
 	FScopedSelectionChangedEvent ScopedSelectionChangedEvent(SharedThis(this),
-		!bFilteringByText && !bIsRefreshingFilter && !bInitialPopulationRequestedAttribute.Get());
+		!bFilteringByText && !bIsRefreshingFilter);
 	TreeData->PopulateFullFolderTree(CreateCompiledFolderFilter());
 	TreeData->FilterFullFolderTree();
 	TreeData->SortRootItems();
