Content Browser - View Options - Thumbnails tests added

#rnx

#Synced-CL 49334002

diff --git a/Engine/Source/Editor/ContentBrowser/Public/SAssetView.h b/Engine/Source/Editor/ContentBrowser/Public/SAssetView.h
--- a/Engine/Source/Editor/ContentBrowser/Public/SAssetView.h
+++ b/Engine/Source/Editor/ContentBrowser/Public/SAssetView.h
@@ -569,6 +569,12 @@ public:
 	/** @return true when we are organizing folders into virtual locations */
 	UE_INTERNAL UE_API bool IsOrganizingFolders() const;
 
+	/** Is the current thumbnail size menu option checked */
+	UE_INTERNAL UE_API bool IsThumbnailSizeChecked(EThumbnailSize InThumbnailSize) const;
+
+	/** Returns true if ViewType is the current view type */
+	UE_INTERNAL UE_API bool IsCurrentViewType(EAssetViewType::Type ViewType) const;
+
 private:
 	/** Sets the pending selection to the current selection (used when changing views or refreshing the view). */
 	UE_API void SyncToSelection( const bool bFocusOnSync = true );
@@ -774,9 +780,6 @@ private:
 
 	UE_API TSharedRef<SWidget> CreateShadowOverlay( TSharedRef<STableViewBase> Table );
 
-	/** Returns true if ViewType is the current view type */
-	UE_API bool IsCurrentViewType(EAssetViewType::Type ViewType) const;
-
 	/** Set the keyboard focus to the correct list view that should be active */
 	UE_API void FocusList() const;
 
@@ -876,9 +879,6 @@ private:
 	/** Called when  thumbnail size is changed */
 	UE_API void OnThumbnailSizeChanged(EThumbnailSize NewThumbnailSize);
 
-	/** Is the current thumbnail size menu option checked */
-	UE_API bool IsThumbnailSizeChecked(EThumbnailSize InThumbnailSize) const;
-
 	/** Are thumbnails allowed to be scaled by the user */
 	UE_API bool IsThumbnailScalingAllowed() const;
 
