EditorUsability : ContentBrowser
- Fixed allow copy/move drop operation to be handled in favorites for files and folders
- Disabled SFavoritePathView widget invalidation causing drop indicators to not be painted

* if CTRL is pressed and folder are dragged then copy/move folders is performed on drop
* if CTRL is not pressed and folders are dragged then add/reordering favorites is performed on drop
* if favorites folders are dragged then reordering is performed on drop
* if file assets are dragged then copy/move items is performed on drop

#Synced-CL 49549735

diff --git a/Engine/Source/Editor/ContentBrowser/Private/SPathView.cpp b/Engine/Source/Editor/ContentBrowser/Private/SPathView.cpp
--- a/Engine/Source/Editor/ContentBrowser/Private/SPathView.cpp
+++ b/Engine/Source/Editor/ContentBrowser/Private/SPathView.cpp
@@ -2160,7 +2160,7 @@ TSharedRef<ITableRow> SPathView::GenerateTreeRow( TSharedPtr<FTreeItem> TreeItem
 	check(TreeItem.IsValid());
 
 	return
-		SNew( STableRow< TSharedPtr<FTreeItem> >, OwnerTable )
+		SNew(STableRow<TSharedPtr<FTreeItem>>, OwnerTable)
 		.OnDragDetected(this, &SPathView::OnFolderDragDetected)
 		[
 			SNew(SAssetTreeItem)
@@ -2959,28 +2959,47 @@ FText SFavoriteAssetTreeItem::GetToolTipText() const
 
 void SFavoriteAssetTreeItem::OnDragEnter(const FGeometry& InGeometry, const FDragDropEvent& InDragDropEvent)
 {
-	if (const TSharedPtr<FContentBrowserDataDragDropOp> DragDropOp = InDragDropEvent.GetOperationAs<FContentBrowserDataDragDropOp>())
+	if (const TSharedPtr<FContentBrowserDataDragDropOp> DragDropOp = GetFavoriteMoveOperation(InDragDropEvent))
+	{
+		DragDropOp->SetToolTip(LOCTEXT("AddOrReorderFavoriteFolder", "Add or Reorder favorite folder"), FAppStyle::GetBrush(TEXT("Graph.ConnectorFeedback.OK")));
+	}
+	else
 	{
-		DragDropOp->SetToolTip(LOCTEXT("ReorderFavoriteFolder", "Reorder favorite folder"), FAppStyle::GetBrush(TEXT("Graph.ConnectorFeedback.OK")));
+		SAssetTreeItem::OnDragEnter(InGeometry, InDragDropEvent);
 	}
 }
 
 void SFavoriteAssetTreeItem::OnDragLeave(const FDragDropEvent& InDragDropEvent)
 {
-	if (const TSharedPtr<FContentBrowserDataDragDropOp> DragDropOp = InDragDropEvent.GetOperationAs<FContentBrowserDataDragDropOp>())
+	if (const TSharedPtr<FContentBrowserDataDragDropOp> DragDropOp = GetFavoriteMoveOperation(InDragDropEvent))
 	{
 		DragDropOp->SetToolTip(FText::GetEmpty(), nullptr);
+		bDraggedOver = false;
+	}
+	else
+	{
+		SAssetTreeItem::OnDragLeave(InDragDropEvent);
 	}
 }
 
 FReply SFavoriteAssetTreeItem::OnDragOver(const FGeometry& InGeometry, const FDragDropEvent& InDragDropEvent)
 {
+	if (!GetFavoriteMoveOperation(InDragDropEvent))
+	{
+		return SAssetTreeItem::OnDragOver(InGeometry, InDragDropEvent);
+	}
+
 	// Let the tree handle this
 	return FReply::Unhandled();
 }
 
 FReply SFavoriteAssetTreeItem::OnDrop(const FGeometry& InGeometry, const FDragDropEvent& InDragDropEvent)
 {
+	if (!GetFavoriteMoveOperation(InDragDropEvent))
+	{
+		return SAssetTreeItem::OnDrop(InGeometry, InDragDropEvent);
+	}
+
 	// Let the tree handle this
 	return FReply::Unhandled();
 }
@@ -3019,6 +3038,23 @@ void SFavoriteAssetTreeItem::UpdateItem(const FContentBrowserItemPath& InPath)
 	}
 }
 
+TSharedPtr<FContentBrowserDataDragDropOp> SFavoriteAssetTreeItem::GetFavoriteMoveOperation(const FDragDropEvent& InDragDropEvent)
+{
+	if (const TSharedPtr<FContentBrowserDataDragDropOp> DragDropOp = InDragDropEvent.GetOperationAs<FContentBrowserDataDragDropOp>())
+	{
+		// Drag operation contains only folders and CTRL is not pressed => Add/Reordering
+		if (!FSlateApplication::Get().GetModifierKeys().IsControlDown()
+			&& !DragDropOp->GetDraggedFolders().IsEmpty() 
+			&& DragDropOp->GetDraggedFiles().IsEmpty())
+		{
+			return DragDropOp;
+		}
+	}
+
+	// Drag operation contains assets files or folders while CTRL is pressed => Move/Copy
+	return nullptr;
+}
+
 SFavoritePathView::SFavoritePathView()
 {
 	bFlat = true;
diff --git a/Engine/Source/Editor/ContentBrowser/Private/SPathView.h b/Engine/Source/Editor/ContentBrowser/Private/SPathView.h
--- a/Engine/Source/Editor/ContentBrowser/Private/SPathView.h
+++ b/Engine/Source/Editor/ContentBrowser/Private/SPathView.h
@@ -633,6 +633,10 @@ protected:
 	void OnFavoriteUpdated(const FContentBrowserItemPath& InItemPath, bool bInAdded);
 	void UpdateItem(const FContentBrowserItemPath& InPath);
 
+private:
+	/** Retrieves only drag operation that favorites view will handle */
+	static TSharedPtr<FContentBrowserDataDragDropOp> GetFavoriteMoveOperation(const FDragDropEvent& InDragDropEvent);
+
 	FText Alias;
 	FText Shortcut;
 };
@@ -640,7 +644,8 @@ protected:
 template<>
 struct TWidgetTypeTraits<class SFavoritePathView>
 {
-	static constexpr bool SupportsInvalidation() { return true; }
+	// Disable invalidation otherwise drop indicators are not drawn
+	static constexpr bool SupportsInvalidation() { return false; }
 };
 
 /**
