Fixed the last path selection being lost when opening the editor with a CB that not displayed and then close again.
We add similar issue where the last selected path would be lost if the CB is open but not displayed until the asset scan is finished

#Synced-CL 49998195

diff --git a/Engine/Source/Editor/ContentBrowser/Private/SPathView.cpp b/Engine/Source/Editor/ContentBrowser/Private/SPathView.cpp
--- a/Engine/Source/Editor/ContentBrowser/Private/SPathView.cpp
+++ b/Engine/Source/Editor/ContentBrowser/Private/SPathView.cpp
@@ -1080,7 +1080,6 @@ void SPathView::Construct( const FArguments& InArgs )
 	UContentBrowserDataSubsystem* ContentBrowserData = IContentBrowserDataModule::Get().GetSubsystem();
 	ContentBrowserData->OnItemDataUpdated().AddSP(this, &SPathView::HandleItemDataUpdated);
 	ContentBrowserData->OnItemDataRefreshed().AddSP(this, &SPathView::HandleItemDataRefreshed);
-	ContentBrowserData->OnItemDataDiscoveryComplete().AddSP(this, &SPathView::HandleItemDataDiscoveryComplete);
 
 	FAssetToolsModule& AssetToolsModule = FModuleManager::LoadModuleChecked<FAssetToolsModule>("AssetTools");
 	FolderPermissionList = AssetToolsModule.Get().GetFolderPermissionList();
@@ -1938,18 +1937,41 @@ void SPathView::ApplyHistoryData( const FHistoryData& History )
 void SPathView::SaveSettings(const FString& IniFilename, const FString& IniSection, const FString& InstanceName) const
 {
 	FString SelectedPathsString;
-	TArray< TSharedPtr<FTreeItem> > PathItems = TreeViewPtr->GetSelectedItems();
 
-	for (const TSharedPtr<FTreeItem>& Item : PathItems)
+	UContentBrowserDataSubsystem* DataSubsystem = IContentBrowserDataModule::Get().GetSubsystem();
 	{
-		if (SelectedPathsString.Len() > 0)
+		TArray< TSharedPtr<FTreeItem> > PathItems = TreeViewPtr->GetSelectedItems();
+
+		for (const TSharedPtr<FTreeItem>& Item : PathItems)
 		{
-			SelectedPathsString += TEXT(",");
+			if (!SelectedPathsString.IsEmpty())
+			{
+				SelectedPathsString += TEXT(",");
+			}
+
+			FName InvariantPath;
+			DataSubsystem->TryConvertVirtualPath(Item->GetItem().GetVirtualPath(), InvariantPath);
+			InvariantPath.AppendString(SelectedPathsString);
 		}
+	}
 
-		FName InvariantPath;
-		IContentBrowserDataModule::Get().GetSubsystem()->TryConvertVirtualPath(Item->GetItem().GetVirtualPath(), InvariantPath);
-		InvariantPath.AppendString(SelectedPathsString);
+	if (SelectedPathsString.IsEmpty())
+	{
+		/**
+		 * We didn't finish selecting the initial paths so save those instead
+		 * This can happen if the we save the settings before the view is populated or the ar scan is finished and the user didn't change the selected paths.
+		 */
+		for (FName SelectedPath : PendingInitialPaths)
+		{
+			if (!SelectedPathsString.IsEmpty())
+			{
+				SelectedPathsString += TEXT(",");
+			}
+
+			FName InvariantPath;
+			DataSubsystem->TryConvertVirtualPath(SelectedPath, InvariantPath);
+			InvariantPath.AppendString(SelectedPathsString);
+		}
 	}
 
 	GConfig->SetString(*IniSection, *(InstanceName + TEXT(".SelectedPaths")), *SelectedPathsString, IniFilename);
@@ -2586,13 +2608,6 @@ void SPathView::HandleItemDataRefreshed()
 	bPopulationRequestedAttribute.Assign(*this, true);
 }
 
-void SPathView::HandleItemDataDiscoveryComplete()
-{
-	// If there were any more initial paths, they no longer exist so clear them now.
-	UE_LOG(LogPathView, VeryVerbose, TEXT("[%s] Resetting pending initial paths at end of asset data discovery"), *WriteToString<256>(OwningContentBrowserName));
-	PendingInitialPaths.Empty();
-}
-
 void SPathView::HandleSettingChanged(FName PropertyName)
 {
 	if ((PropertyName == GET_MEMBER_NAME_CHECKED(UContentBrowserSettings, DisplayEmptyFolders))
diff --git a/Engine/Source/Editor/ContentBrowser/Private/SPathView.h b/Engine/Source/Editor/ContentBrowser/Private/SPathView.h
--- a/Engine/Source/Editor/ContentBrowser/Private/SPathView.h
+++ b/Engine/Source/Editor/ContentBrowser/Private/SPathView.h
@@ -304,9 +304,6 @@ protected:
 	/** Handles updating the view when content items are refreshed */
 	void HandleItemDataRefreshed();
 
-	/** Notification for when the content browser has completed it's initial search */
-	void HandleItemDataDiscoveryComplete();
-
 	/** Creates a list item for the tree view */
 	virtual TSharedRef<ITableRow> GenerateTreeRow(TSharedPtr<FTreeItem> TreeItem, const TSharedRef<STableViewBase>& OwnerTable);
 
