[BUGFIX] CB showing wrong hint text when the folder gets renamed

#rnx
#Synced-CL 49809867

diff --git a/Engine/Source/Editor/ContentBrowser/Private/SContentBrowser.cpp b/Engine/Source/Editor/ContentBrowser/Private/SContentBrowser.cpp
--- a/Engine/Source/Editor/ContentBrowser/Private/SContentBrowser.cpp
+++ b/Engine/Source/Editor/ContentBrowser/Private/SContentBrowser.cpp
@@ -4844,8 +4844,13 @@ FText SContentBrowser::GetSearchAssetsHintText() const
 			SearchHint += TEXT(" ");
 			for(int32 i = 0; i < Paths.Num(); i++)
 			{
-				SearchHint += Paths[i].GetDisplayName().ToString();
+				const FString PathDisplayNameString = Paths[i].GetDisplayName().ToString();
+				if (PathDisplayNameString.IsEmpty())
+				{
+					continue;
+				}
 
+				SearchHint += PathDisplayNameString;
 				if (i + 1 < Paths.Num())
 				{
 					SearchHint += ", ";
