[Core] Enable range-based for loop support for `TFilterCollection`.

#Synced-CL 50036996

diff --git a/Engine/Plugins/AudioInsights/Source/AudioInsights/Private/Views/AudioEventLogDashboardViewFactory.cpp b/Engine/Plugins/AudioInsights/Source/AudioInsights/Private/Views/AudioEventLogDashboardViewFactory.cpp
--- a/Engine/Plugins/AudioInsights/Source/AudioInsights/Private/Views/AudioEventLogDashboardViewFactory.cpp
+++ b/Engine/Plugins/AudioInsights/Source/AudioInsights/Private/Views/AudioEventLogDashboardViewFactory.cpp
@@ -1046,9 +1046,9 @@ namespace UE::Audio::Insights
 			return true;
 		}
 		
-		for (int32 FilterIndex = 0; FilterIndex < ActiveFilters->Num(); ++FilterIndex)
+		for (const TSharedPtr<IFilter<FEventLogFilterID>>& Filter : *ActiveFilters.Get())
 		{
-			if (ActiveFilters->GetFilterAtIndex(FilterIndex)->PassesFilter(EventLogID))
+			if (Filter->PassesFilter(EventLogID))
 			{
 				return true;
 			}
diff --git a/Engine/Plugins/Compositing/Composite/Source/CompositeEditor/Private/UI/SCompositePassTree.cpp b/Engine/Plugins/Compositing/Composite/Source/CompositeEditor/Private/UI/SCompositePassTree.cpp
--- a/Engine/Plugins/Compositing/Composite/Source/CompositeEditor/Private/UI/SCompositePassTree.cpp
+++ b/Engine/Plugins/Compositing/Composite/Source/CompositeEditor/Private/UI/SCompositePassTree.cpp
@@ -227,9 +227,9 @@ public:
 				return true;
 			}
 		
-			for (int32 Index = 0; Index < NumFilters; ++Index)
+			for (const TSharedPtr<IFilter<FFilterType>>& Filter : *FilterCollection.Get())
 			{
-				if (FilterCollection->GetFilterAtIndex(Index)->PassesFilter(InItem))
+				if (Filter->PassesFilter(InItem))
 				{
 					return true;
 				}
diff --git a/Engine/Plugins/Developer/Concert/ConcertUI/ConcertClientSharedSlate/Source/ConcertClientSharedSlate/Private/Replication/Editor/View/PropertyTree/SFilteredPropertyTreeView.cpp b/Engine/Plugins/Developer/Concert/ConcertUI/ConcertClientSharedSlate/Source/ConcertClientSharedSlate/Private/Replication/Editor/View/PropertyTree/SFilteredPropertyTreeView.cpp
--- a/Engine/Plugins/Developer/Concert/ConcertUI/ConcertClientSharedSlate/Source/ConcertClientSharedSlate/Private/Replication/Editor/View/PropertyTree/SFilteredPropertyTreeView.cpp
+++ b/Engine/Plugins/Developer/Concert/ConcertUI/ConcertClientSharedSlate/Source/ConcertClientSharedSlate/Private/Replication/Editor/View/PropertyTree/SFilteredPropertyTreeView.cpp
@@ -150,9 +150,9 @@ namespace UE::ConcertClientSharedSlate
 	bool SFilteredPropertyTreeView::PassesAnyFilters(const ConcertSharedSlate::FPropertyData& ReplicatedPropertyData) const
 	{
 		TSharedPtr<TFilterCollection<const ConcertSharedSlate::FPropertyData&>> FilterCollection = FilterBar->GetAllActiveFilters();
-		for (int32 Index = 0; Index < FilterCollection->Num(); Index++)
+		for (const TSharedPtr<IFilter<const ConcertSharedSlate::FPropertyData&>>& Filter : *FilterCollection.Get())
 		{
-			if (FilterCollection->GetFilterAtIndex(Index)->PassesFilter(ReplicatedPropertyData))
+			if (Filter->PassesFilter(ReplicatedPropertyData))
 			{
 				return true;
 			}
diff --git a/Engine/Plugins/Experimental/EditorTelemetry/Source/Private/EditorTelemetry.cpp b/Engine/Plugins/Experimental/EditorTelemetry/Source/Private/EditorTelemetry.cpp
--- a/Engine/Plugins/Experimental/EditorTelemetry/Source/Private/EditorTelemetry.cpp
+++ b/Engine/Plugins/Experimental/EditorTelemetry/Source/Private/EditorTelemetry.cpp
@@ -998,9 +998,8 @@ void FEditorTelemetry::StartSession()
 			{
 				Private::FAnalyticsJsonWriter J(&FilterText);
 				J.WriteArrayStart();
-				for (int32 i=0; i < Data.FrontendFilters->Num(); ++i)
+				for (const TSharedPtr<IFilter<FAssetFilterType>>& Filter : *Data.FrontendFilters.Get())
 				{
-					TSharedPtr<IFilter<FAssetFilterType>> Filter = Data.FrontendFilters->GetFilterAtIndex(i);
 					J.WriteValue(Filter->GetName());
 				}
 				J.WriteArrayEnd();
diff --git a/Engine/Plugins/Media/MediaFrameworkUtilities/Source/MediaFrameworkUtilitiesEditor/Private/AssetEditor/SMediaProfileSourcesTreeView.cpp b/Engine/Plugins/Media/MediaFrameworkUtilities/Source/MediaFrameworkUtilitiesEditor/Private/AssetEditor/SMediaProfileSourcesTreeView.cpp
--- a/Engine/Plugins/Media/MediaFrameworkUtilities/Source/MediaFrameworkUtilitiesEditor/Private/AssetEditor/SMediaProfileSourcesTreeView.cpp
+++ b/Engine/Plugins/Media/MediaFrameworkUtilities/Source/MediaFrameworkUtilitiesEditor/Private/AssetEditor/SMediaProfileSourcesTreeView.cpp
@@ -456,9 +456,9 @@ public:
 				return true;
 			}
 		
-			for (int32 Index = 0; Index < NumFilters; ++Index)
+			for (const TSharedPtr<IFilter<FFilterType>>& Filter : *FilterCollection.Get())
 			{
-				if (FilterCollection->GetFilterAtIndex(Index)->PassesFilter(InItem))
+				if (Filter->PassesFilter(InItem))
 				{
 					return true;
 				}
diff --git a/Engine/Source/Editor/ContentBrowser/Private/SAssetPicker.cpp b/Engine/Source/Editor/ContentBrowser/Private/SAssetPicker.cpp
--- a/Engine/Source/Editor/ContentBrowser/Private/SAssetPicker.cpp
+++ b/Engine/Source/Editor/ContentBrowser/Private/SAssetPicker.cpp
@@ -736,9 +736,9 @@ bool SAssetPicker::IsShowingOtherDevelopersContent() const
 	}
 	else
 	{
-		for (int32 i=0; i < FrontendFilters->Num(); ++i)
+		for (const TSharedPtr<IFilter<FAssetFilterType>>& Filter : *FrontendFilters.Get())
 		{
-			if (FrontendFilters->GetFilterAtIndex(i) == HideOtherDevelopersFilter)
+			if (Filter == HideOtherDevelopersFilter)
 			{
 				return false;
 			}
diff --git a/Engine/Source/Editor/ContentBrowser/Private/SAssetView.cpp b/Engine/Source/Editor/ContentBrowser/Private/SAssetView.cpp
--- a/Engine/Source/Editor/ContentBrowser/Private/SAssetView.cpp
+++ b/Engine/Source/Editor/ContentBrowser/Private/SAssetView.cpp
@@ -3746,12 +3746,11 @@ bool SAssetView::ShouldFilterRecursively() const
 		// Otherwise, check if there are any non-inverse frontend filters selected
 		if (FrontendFilters.IsValid())
 		{
-			for (int32 FilterIndex = 0; FilterIndex < FrontendFilters->Num(); ++FilterIndex)
+			for (const TSharedPtr<IFilter<FAssetFilterType>>& Filter : *FrontendFilters.Get())
 			{
-				const auto* Filter = static_cast<FFrontendFilter*>(FrontendFilters->GetFilterAtIndex(FilterIndex).Get());
-				if (Filter)
+				if (const auto* FrontendFilter = static_cast<FFrontendFilter*>(Filter.Get()))
 				{
-					if (!Filter->IsInverseFilter())
+					if (!FrontendFilter->IsInverseFilter())
 					{
 						return true;
 					}
@@ -3803,13 +3802,13 @@ void SAssetView::RefreshFilteredItems()
 		const bool bInvalidateFilterCache = false;
 		const FContentBrowserDataFilter DataFilter = CreateBackendDataFilter(bInvalidateFilterCache);
 
-		for (int32 FilterIdx = 0; FilterIdx < FrontendFilters->Num(); ++FilterIdx)
+		for (const TSharedPtr<IFilter<FAssetFilterType>>& Filter : *FrontendFilters.Get())
 		{
 			// There are only FFrontendFilters in this collection
-			const TSharedPtr<FFrontendFilter>& Filter = StaticCastSharedPtr<FFrontendFilter>(FrontendFilters->GetFilterAtIndex(FilterIdx));
-			if (Filter.IsValid())
+			const TSharedPtr<FFrontendFilter>& FrontendFilter = StaticCastSharedPtr<FFrontendFilter>(Filter);
+			if (FrontendFilter.IsValid())
 			{
-				Filter->SetCurrentFilter(DataSourcePaths, DataFilter);
+				FrontendFilter->SetCurrentFilter(DataSourcePaths, DataFilter);
 			}
 		}
 	}
@@ -4353,12 +4352,12 @@ void SAssetView::OnFrontendFiltersChanged()
 	if (TextFilter.IsValid() && FrontendFilters.IsValid())
 	{ 
 		TArray<FText> CustomTextFilters;
-		for (int32 i = 0; i < FrontendFilters->Num(); ++i)
+		for (const TSharedPtr<IFilter<FAssetFilterType>>& Filter : *FrontendFilters.Get())
 		{
-			TSharedPtr<FFrontendFilter> Filter = StaticCastSharedPtr<FFrontendFilter>(FrontendFilters->GetFilterAtIndex(i));
+			const TSharedPtr<FFrontendFilter> FrontendFilter = StaticCastSharedPtr<FFrontendFilter>(Filter);
 			if (Filter.IsValid())
 			{
-				TOptional<FText> Text = Filter->GetAsCustomTextFilter();
+				TOptional<FText> Text = FrontendFilter->GetAsCustomTextFilter();
 				if (Text.IsSet())
 				{
 					CustomTextFilters.Add(Text.GetValue());
diff --git a/Engine/Source/Editor/ContentBrowser/Private/SPathView.cpp b/Engine/Source/Editor/ContentBrowser/Private/SPathView.cpp
--- a/Engine/Source/Editor/ContentBrowser/Private/SPathView.cpp
+++ b/Engine/Source/Editor/ContentBrowser/Private/SPathView.cpp
@@ -1358,9 +1358,9 @@ bool SPathView::IsPluginPathFilterInUse(TSharedRef<FContentBrowserPluginFilter>
 {
 	if (PluginPathFilters.IsValid())
 	{
-		for (int32 i=0; i < PluginPathFilters->Num(); ++i)
+		for (const TSharedPtr<IFilter<FPluginFilterType>>& PluginFilter : *PluginPathFilters.Get())
 		{
-			if (PluginPathFilters->GetFilterAtIndex(i) == Filter)
+			if (PluginFilter == Filter)
 			{
 				return true;
 			}
@@ -1979,15 +1979,15 @@ void SPathView::SaveSettings(const FString& IniFilename, const FString& IniSecti
 	FString PluginFiltersString;
 	if (PluginPathFilters.IsValid())
 	{
-		for (int32 FilterIdx = 0; FilterIdx < PluginPathFilters->Num(); ++FilterIdx)
+		for (const TSharedPtr<IFilter<FPluginFilterType>>& Filter : *PluginPathFilters.Get())
 		{
 			if (PluginFiltersString.Len() > 0)
 			{
 				PluginFiltersString += TEXT(",");
 			}
 
-			TSharedPtr<FContentBrowserPluginFilter> Filter = StaticCastSharedPtr<FContentBrowserPluginFilter>(PluginPathFilters->GetFilterAtIndex(FilterIdx));
-			PluginFiltersString += Filter->GetName();
+			TSharedPtr<FContentBrowserPluginFilter> PluginFilter = StaticCastSharedPtr<FContentBrowserPluginFilter>(Filter);
+			PluginFiltersString += PluginFilter->GetName();
 		}
 		GConfig->SetString(*IniSection, *(InstanceName + TEXT(".PluginFilters")), *PluginFiltersString, IniFilename);
 	}
diff --git a/Engine/Source/Runtime/Core/Public/Misc/FilterCollection.h b/Engine/Source/Runtime/Core/Public/Misc/FilterCollection.h
--- a/Engine/Source/Runtime/Core/Public/Misc/FilterCollection.h
+++ b/Engine/Source/Runtime/Core/Public/Misc/FilterCollection.h
@@ -108,9 +108,9 @@ public:
 	 */
 	bool PassesAllFilters( ItemType InItem ) const
 	{
-		for (int Index = 0; Index < ChildFilters.Num(); Index++)
+		for (const TSharedPtr<IFilter<ItemType>>& Filter : ChildFilters)
 		{
-			if( !ChildFilters[Index]->PassesFilter( InItem ) )
+			if( !Filter->PassesFilter( InItem ) )
 			{
 				return false;
 			}
@@ -123,6 +123,15 @@ public:
 	DECLARE_EVENT( TFilterCollection<ItemType>, FChangedEvent );
 	FChangedEvent& OnChanged() { return ChangedEvent; }
 
+	using RangedForIteratorType = TArray<TSharedPtr<IFilter<ItemType>>>::RangedForIteratorType;
+	using RangedForConstIteratorType = TArray<TSharedPtr<IFilter<ItemType>>>::RangedForConstIteratorType;
+
+	/** STL-like iterators to enable range-based for loop support. */
+	[[nodiscard]] UE_REWRITE constexpr RangedForIteratorType      begin()       { return ChildFilters.begin(); }
+	[[nodiscard]] UE_REWRITE constexpr RangedForConstIteratorType begin() const { return ChildFilters.begin(); }
+	[[nodiscard]] UE_REWRITE constexpr RangedForIteratorType      end()         { return ChildFilters.end(); }
+	[[nodiscard]] UE_REWRITE constexpr RangedForConstIteratorType end()   const { return ChildFilters.end(); }
+
 protected:
 
 	/** 
