[Fixed] a bug where the selected path view filters list grows unbounded

[Added] A check to see if the filter name is already inside the array

#13781

#Synced-CL 50045910

diff --git a/Engine/Source/Editor/ContentBrowser/Private/SPathView.cpp b/Engine/Source/Editor/ContentBrowser/Private/SPathView.cpp
--- a/Engine/Source/Editor/ContentBrowser/Private/SPathView.cpp
+++ b/Engine/Source/Editor/ContentBrowser/Private/SPathView.cpp
@@ -1446,16 +1446,20 @@ void SPathView::SetPluginPathFilterActive(const TSharedRef<FContentBrowserPlugin
 
 	if (FPathViewConfig* PathViewConfig = GetPathViewConfig())
 	{
+		const FString FilterName = Filter->GetName();
 		if (bActive)
 		{
-			PathViewConfig->PluginFilters.Add(Filter->GetName());
+			if (!PathViewConfig->PluginFilters.Contains(FilterName))
+			{
+				PathViewConfig->PluginFilters.Add(FilterName);
+				UContentBrowserConfig::Get()->SaveEditorConfig();
+			}
 		}
 		else
 		{
-			PathViewConfig->PluginFilters.Remove(Filter->GetName());
+			PathViewConfig->PluginFilters.Remove(FilterName);
+			UContentBrowserConfig::Get()->SaveEditorConfig();
 		}
-		
-		UContentBrowserConfig::Get()->SaveEditorConfig();
 	}
 }
 
