[Fixed] Editor does not apply path view filters on boot

#13782

#Synced-CL 50350999

diff --git a/Engine/Source/Editor/ContentBrowser/Private/SPathView.cpp b/Engine/Source/Editor/ContentBrowser/Private/SPathView.cpp
--- a/Engine/Source/Editor/ContentBrowser/Private/SPathView.cpp
+++ b/Engine/Source/Editor/ContentBrowser/Private/SPathView.cpp
@@ -2080,10 +2080,18 @@ void SPathView::LoadSettings(const FString& IniFilename, const FString& IniSecti
 			}
 		}
 
+		bool bAnyFiltersActive = false;
 		for (const TSharedRef<FContentBrowserPluginFilter>& Filter : AllPluginPathFilters)
 		{
 			bool bFilterActive = NewSelectedFilters.Contains(Filter->GetName());
 			SetPluginPathFilterActive(Filter, bFilterActive);
+
+			bAnyFiltersActive |= bFilterActive;
+		}
+
+		if (bAnyFiltersActive)
+		{
+			Populate();
 		}
 	}
 }
