EditorUsability : Favorites
- Fixed Favorites not refreshing
- Added UE_EXPERIMENTAL to the public API refresh function

#Synced-CL 50149474

diff --git a/Engine/Source/Editor/ContentBrowser/Private/ContentBrowserSingleton.cpp b/Engine/Source/Editor/ContentBrowser/Private/ContentBrowserSingleton.cpp
--- a/Engine/Source/Editor/ContentBrowser/Private/ContentBrowserSingleton.cpp
+++ b/Engine/Source/Editor/ContentBrowser/Private/ContentBrowserSingleton.cpp
@@ -1319,6 +1319,11 @@ void FContentBrowserSingleton::UnregisterOnFavoritesChangedDelegate(FDelegateHan
 	OnFavoritesChanged.Remove(Handle);
 }
 
+void FContentBrowserSingleton::RefreshFavoritesPathView()
+{
+	BroadcastFavoritesChanged(FContentBrowserItemPath(), /** Added */false);
+}
+
 void FContentBrowserSingleton::BroadcastFavoritesChanged(const FContentBrowserItemPath& ChangedItem, bool bAdded) const
 {
 	OnFavoritesChanged.Broadcast(ChangedItem, bAdded);
diff --git a/Engine/Source/Editor/ContentBrowser/Private/ContentBrowserSingleton.h b/Engine/Source/Editor/ContentBrowser/Private/ContentBrowserSingleton.h
--- a/Engine/Source/Editor/ContentBrowser/Private/ContentBrowserSingleton.h
+++ b/Engine/Source/Editor/ContentBrowser/Private/ContentBrowserSingleton.h
@@ -121,6 +121,7 @@ PRAGMA_ENABLE_DEPRECATION_WARNINGS
 	UE_DEPRECATED(5.7, "Use the overload that takes in an FOnFavoritesChanged delegate")
 	virtual FDelegateHandle RegisterOnFavoritesChangedHandler(FSimpleDelegate OnFavoritesChanged) override;
 	virtual void UnregisterOnFavoritesChangedDelegate(FDelegateHandle Handle) override;
+	virtual void RefreshFavoritesPathView() override;
 	virtual TArray<FString> GetAliasesForPath(const FSoftObjectPath& InPath) const override;
 	virtual void AddCustomTextFilter(const FCustomTextFilterData& FilterData, bool bApplyFilter) override;
 	virtual void DeleteCustomTextFilterByLabel(const FText& FilterLabel) override;
diff --git a/Engine/Source/Editor/ContentBrowser/Public/IContentBrowserSingleton.h b/Engine/Source/Editor/ContentBrowser/Public/IContentBrowserSingleton.h
--- a/Engine/Source/Editor/ContentBrowser/Public/IContentBrowserSingleton.h
+++ b/Engine/Source/Editor/ContentBrowser/Public/IContentBrowserSingleton.h
@@ -849,6 +849,10 @@ public:
 
 	/** Unregister a previously-registered handler for when Favorites changes. */
 	virtual void UnregisterOnFavoritesChangedDelegate(FDelegateHandle Handle) = 0;
+	
+	/** Triggers a refresh of favorite path views in content browser */
+	UE_EXPERIMENTAL(5.8, "Experimental way to refresh the content browser favorites, API may change in the future")
+	virtual void RefreshFavoritesPathView() = 0;
 
 	/**
 	 * Get a list of other paths that the data source may be using to represent a specific path
