Fixed memory leak in content browser when removing collection
[RN] Fixed memory leak in content browser when removing collection

#Synced-CL 50041969

diff --git a/Engine/Source/Editor/ContentBrowser/Private/SContentBrowser.cpp b/Engine/Source/Editor/ContentBrowser/Private/SContentBrowser.cpp
--- a/Engine/Source/Editor/ContentBrowser/Private/SContentBrowser.cpp
+++ b/Engine/Source/Editor/ContentBrowser/Private/SContentBrowser.cpp
@@ -1159,7 +1159,14 @@ void SContentBrowser::RemoveSlotForCollectionContainer(const TSharedRef<ICollect
 		float SplitterSize = Splitter.SlotAt(SlotIndex).GetSizeValue();
 		GConfig->SetFloat(*SettingsIniSection, *Key, SplitterSize, Filename);
 
-		Splitter.RemoveAt(SlotIndex);
+		if (UE::Editor::ContentBrowser::IsNewStyleEnabled())
+		{
+			LegacyContentSourceWidgets->SourceTreePtr->RemoveSlotAt(SlotIndex);
+		}
+		else
+		{
+			Splitter.RemoveAt(SlotIndex);
+		}
 
 		CollectionSources.RemoveAt(Index);
 	}
diff --git a/Engine/Source/Editor/ContentBrowser/Private/Widgets/SContentBrowserSourceTree.cpp b/Engine/Source/Editor/ContentBrowser/Private/Widgets/SContentBrowserSourceTree.cpp
--- a/Engine/Source/Editor/ContentBrowser/Private/Widgets/SContentBrowserSourceTree.cpp
+++ b/Engine/Source/Editor/ContentBrowser/Private/Widgets/SContentBrowserSourceTree.cpp
@@ -155,6 +155,27 @@ namespace UE::Editor::ContentBrowser::Private
 		return INDEX_NONE;
 	}
 
+	int32 SContentBrowserSourceTree::RemoveSlotAt(int32 InIndex)
+	{
+		if (!Splitter.IsValid())
+		{
+			return INDEX_NONE;
+		}
+
+		if (Slots.IsValidIndex(InIndex))
+		{
+			Slots.RemoveAt(InIndex);
+
+			Splitter->RemoveAt(InIndex);
+
+			UpdateTotalHeaderHeight();
+
+			return InIndex;
+		}
+
+		return INDEX_NONE;
+	}
+
 	void SContentBrowserSourceTree::ClearChildren()
 	{
 		Slots.Empty();
diff --git a/Engine/Source/Editor/ContentBrowser/Private/Widgets/SContentBrowserSourceTree.h b/Engine/Source/Editor/ContentBrowser/Private/Widgets/SContentBrowserSourceTree.h
--- a/Engine/Source/Editor/ContentBrowser/Private/Widgets/SContentBrowserSourceTree.h
+++ b/Engine/Source/Editor/ContentBrowser/Private/Widgets/SContentBrowserSourceTree.h
@@ -116,10 +116,17 @@ namespace UE::Editor::ContentBrowser::Private
 		/** Removes a slot from this panel which contains the specified SWidget
 		 *
 		 * @param InSlotWidget The widget to match when searching through the slots
-		 * @returns The index in the children array where the slot was removed and -1 if no slot was found matching the widget
+		 * @returns The index in the children array where the slot was removed and INDEX_NONE if no slot was found matching the widget
 		 */
 		int32 RemoveSlot(const TSharedRef<SWidget>& InSlotWidget);
 
+		/** Removes a slot from this panel at the specified slot index
+		 *
+		 * @param InIndex The index of the slot to remove
+		 * @returns The index in the children array where the slot was removed and INDEX_NONE if there was no slot at that index
+		 */
+		int32 RemoveSlotAt(int32 InIndex);
+
 		/** Removes all children from the panel. */
 		void ClearChildren();
 
