PR #14347: Avoids potential crash in GetCombinedBackendFilter

#Synced-CL 50415902

diff --git a/Engine/Source/Editor/ContentBrowser/Private/SFilterList.cpp b/Engine/Source/Editor/ContentBrowser/Private/SFilterList.cpp
--- a/Engine/Source/Editor/ContentBrowser/Private/SFilterList.cpp
+++ b/Engine/Source/Editor/ContentBrowser/Private/SFilterList.cpp
@@ -195,7 +195,7 @@ TSharedRef<SWidget> SFilterList::ExternalMakeAddFilterMenu()
 FARFilter SFilterList::GetCombinedBackendFilter(TArray<TSharedRef<const FPathPermissionList>>& OutPermissionLists) const
 {
 	TSharedPtr<FFilter_HideOtherDevelopers> OtherDevelopersFilter = StaticCastSharedPtr<FFilter_HideOtherDevelopers>(GetFrontendFilter("HideOtherDevelopersBackend"));
-	if (OtherDevelopersFilter->IsActive())
+	if (OtherDevelopersFilter.IsValid() && OtherDevelopersFilter->IsActive())
 	{
 		OutPermissionLists.Add(OtherDevelopersFilter->GetPathPermissionList());
 	}
