Adding content browser suggestion filtering so only paths starting with the currently typed text are shown

#Synced-CL 51339719

diff --git a/Engine/Source/Editor/ContentBrowser/Private/SContentBrowser.cpp b/Engine/Source/Editor/ContentBrowser/Private/SContentBrowser.cpp
--- a/Engine/Source/Editor/ContentBrowser/Private/SContentBrowser.cpp
+++ b/Engine/Source/Editor/ContentBrowser/Private/SContentBrowser.cpp
@@ -2654,11 +2654,14 @@ FText SContentBrowser::OnGetEditPathAsText(const FString& Path) const
 TArray<FNavigationBarComboOption> SContentBrowser::OnCompletePathPrefix(const FString& Prefix) const
 {
 	FStringView PrefixView = Prefix;
+	FStringView SuffixView;
 	
 	// Strip to last path separator, but keep the very first path separator
 	if (int32 Index = UE::String::FindLastChar(PrefixView, '/'); Index != INDEX_NONE)
 	{	
 		PrefixView.LeftInline(FMath::Max(1, Index));
+		SuffixView = Prefix;
+		SuffixView.RightChopInline(PrefixView.Len() + 1);
 	}
 
 	// Find PrefixView in the available tree of data sources, get its direct children, and filter them by SuffixView 
@@ -2672,7 +2675,12 @@ TArray<FNavigationBarComboOption> SContentBrowser::OnCompletePathPrefix(const FS
 	TArray<FNavigationBarComboOption> Results;
 	for (const FString& ChildPath : ChildPaths)
 	{
-		Results.Emplace(FText::AsCultureInvariant(ChildPath), ChildPath);
+		FStringView ChildSuffix = ChildPath;
+		ChildSuffix.RightChopInline(PrefixView.Len() + 1);
+		if (ChildSuffix.StartsWith(SuffixView))
+		{
+			Results.Emplace(FText::AsCultureInvariant(ChildPath), ChildPath);
+		}
 	}
 	return Results;
 }
