EditorUsability : UnrealEd
- Fixed non-displayable references list is not interactable when deleting pending assets in dialog

#Synced-CL 51264175

diff --git a/Engine/Source/Editor/ContentBrowser/Private/SAssetView.cpp b/Engine/Source/Editor/ContentBrowser/Private/SAssetView.cpp
--- a/Engine/Source/Editor/ContentBrowser/Private/SAssetView.cpp
+++ b/Engine/Source/Editor/ContentBrowser/Private/SAssetView.cpp
@@ -78,7 +78,7 @@
 #include "Widgets/Notifications/SProgressBar.h"
 #include "Widgets/SOverlay.h"
 #include "Widgets/Text/STextBlock.h"
-
+#include "Widgets/Text/SMultiLineEditableText.h"
 #include "ISourceControlModule.h"
 #include "RevisionControlStyle/RevisionControlStyle.h"
 
@@ -1528,13 +1528,14 @@ void SAssetView::Construct( const FArguments& InArgs )
 				SNew(SScrollBox)
 				.Visibility( this, &SAssetView::IsAssetShowWarningTextVisible )
 				+SScrollBox::Slot()
+				.Padding(8.f)
 				[
-					// A warning to display when there are no assets to show
-					SNew( STextBlock )
-					.Justification( ETextJustify::Center )
-					.Text( this, &SAssetView::GetAssetShowWarningText )
-					.Visibility( this, &SAssetView::IsAssetShowWarningTextVisible )
-					.AutoWrapText( true )
+					// A warning to display when there are no assets to show, allow text copy
+					SNew(SMultiLineEditableText)
+					.IsReadOnly(true)
+					.Justification(ETextJustify::Left)
+					.Text(this, &SAssetView::GetAssetShowWarningText)
+					.AutoWrapText(true)
 				]
 			]
 
@@ -6348,7 +6349,7 @@ void SAssetView::SetSortParameters(
 
 EVisibility SAssetView::IsAssetShowWarningTextVisible() const
 {
-	return (FilteredAssetItems.Num() > 0 || bQuickFrontendListRefreshRequested) ? EVisibility::Collapsed : EVisibility::HitTestInvisible;
+	return (FilteredAssetItems.Num() > 0 || bQuickFrontendListRefreshRequested) ? EVisibility::Collapsed : EVisibility::Visible;
 }
 
 FText SAssetView::GetAssetShowWarningText() const
diff --git a/Engine/Source/Editor/UnrealEd/Private/Dialogs/SDeleteAssetsDialog.cpp b/Engine/Source/Editor/UnrealEd/Private/Dialogs/SDeleteAssetsDialog.cpp
--- a/Engine/Source/Editor/UnrealEd/Private/Dialogs/SDeleteAssetsDialog.cpp
+++ b/Engine/Source/Editor/UnrealEd/Private/Dialogs/SDeleteAssetsDialog.cpp
@@ -800,7 +800,7 @@ TSharedRef<SWidget> SDeleteAssetsDialog::MakeAssetViewForReferencerAssets()
 
 	AssetPickerConfig.AssetShowWarningText = TAttribute< FText >( this, &SDeleteAssetsDialog::GetReferencingAssetsEmptyText );
 
-	AssetPickerConfig.InitialAssetViewType = EAssetViewType::Tile;
+	AssetPickerConfig.InitialAssetViewType = EAssetViewType::List;
 	AssetPickerConfig.OnAssetsActivated = FOnAssetsActivated::CreateSP(this, &SDeleteAssetsDialog::OnAssetsActivated);
 	AssetPickerConfig.OnShouldFilterAsset = FOnShouldFilterAsset::CreateSP(this, &SDeleteAssetsDialog::OnShouldFilterAsset);
 	AssetPickerConfig.OnGetAssetContextMenu = FOnGetAssetContextMenu::CreateSP( this, &SDeleteAssetsDialog::OnGetAssetContextMenu );
