Fixed a crash we see on live where we start a rename when there is already one happening.
I don't have repo unfortunately, but it seem the  SAssetView::RenameItem and the asset rename in the tick function could end up causing the crash.

Refactor a bit code to use a new function CancelPendingRename to avoid duplicating some logic.

#Synced-CL 51761360

diff --git a/Engine/Source/Editor/ContentBrowser/Private/SAssetView.cpp b/Engine/Source/Editor/ContentBrowser/Private/SAssetView.cpp
--- a/Engine/Source/Editor/ContentBrowser/Private/SAssetView.cpp
+++ b/Engine/Source/Editor/ContentBrowser/Private/SAssetView.cpp
@@ -1966,6 +1966,8 @@ void SAssetView::RenameItem(const FContentBrowserItem& ItemToRename)
 {
 	if (const TSharedPtr<FAssetViewItem> Item = Items->FindItemForRename(ItemToRename))
 	{
+		CancelPendingRename();
+
 		AwaitingScrollIntoViewForRename = Item;
 		
 		SetSelection(Item);
@@ -2576,6 +2578,7 @@ void SAssetView::Tick( const FGeometry& AllottedGeometry, const double InCurrent
 		}
 		else if (OwnerWindow->HasAnyUserFocusOrFocusedDescendants())
 		{
+			CancelPendingRename();
 			AssetAwaitingRename->OnRenameRequested().ExecuteIfBound();
 			AwaitingRename = nullptr;
 		}
@@ -5505,23 +5508,7 @@ void SAssetView::AssetItemWidgetDestroyed(const TSharedPtr<FAssetViewItem>& Item
 {
 	if(RenamingAsset.Pin().Get() == Item.Get())
 	{
-		/* Check if the item is in a temp state and if it is, commit using the default name so that it does not entirely vanish on the user.
-		   This keeps the functionality consistent for content to never be in a temporary state */
-
-		if (Item && Item->IsTemporary())
-		{
-			if (Item->IsFile())
-			{
-				FText OutErrorText;
-				EndCreateDeferredItem(Item, Item->GetItem().GetItemName().ToString(), /*bFinalize*/true, OutErrorText);
-			}
-			else
-			{
-				DeferredItemToCreate.Reset();
-			}
-		}
-
-		RenamingAsset.Reset();
+		CancelPendingRename();
 	}
 
 	if ( VisibleItems.Remove(Item) != INDEX_NONE )
@@ -5688,11 +5675,7 @@ TSharedPtr<SWidget> SAssetView::OnGetContextMenuContent()
 {
 	if (CanOpenContextMenu())
 	{
-		if (IsRenamingAsset())
-		{
-			RenamingAsset.Pin()->OnRenameCanceled().ExecuteIfBound();
-			RenamingAsset.Reset();
-		}
+		CancelPendingRename();
 
 		OnInteractDuringFiltering();
 		const TArray<FContentBrowserItem> SelectedItems = GetSelectedItems();
@@ -6930,6 +6913,40 @@ void SAssetView::OnInteractDuringFiltering()
 	}
 }
 
+void SAssetView::CancelPendingRename()
+{
+	if (IsRenamingAsset())
+	{
+		if (TSharedPtr<FAssetViewItem> Item = RenamingAsset.Pin())
+		{
+			/**
+			 * Check if the item is in a temp state and if it is, commit using the default name so that it does not entirely vanish on the user.
+			 * This keeps the functionality consistent for content to never be in a temporary state
+			 */
+
+			if (Item && Item->IsTemporary())
+			{
+				if (Item->IsFile())
+				{
+					FText OutErrorText;
+					EndCreateDeferredItem(Item, Item->GetItem().GetItemName().ToString(), /*bFinalize*/true, OutErrorText);
+				}
+				else
+				{
+					DeferredItemToCreate.Reset();
+					Item->OnRenameCanceled().ExecuteIfBound();
+				}
+			}
+			else
+			{
+				Item->OnRenameCanceled().ExecuteIfBound();
+			}
+
+			RenamingAsset.Reset();
+		}
+	}
+}
+
 #undef checkAssetList
 #undef ASSET_VIEW_PARANOIA_LIST_CHECKS
 
diff --git a/Engine/Source/Editor/ContentBrowser/Public/SAssetView.h b/Engine/Source/Editor/ContentBrowser/Public/SAssetView.h
--- a/Engine/Source/Editor/ContentBrowser/Public/SAssetView.h
+++ b/Engine/Source/Editor/ContentBrowser/Public/SAssetView.h
@@ -1026,6 +1026,9 @@ private:
 	UE_API bool IsCustomViewSet() const;
 
 private:
+	/** Cancel the ongoing rename if there is one before starting another one. */
+	void CancelPendingRename();
+
 	friend class FAssetViewFrontendFilterHelper;
 
 	/** Private type managing data retrieved from the backend and async filtering thereof */
