Changed asset view's "no results, check your filter" message alignment to match the details view one

#Synced-CL 51883779

diff --git a/Engine/Source/Editor/ContentBrowser/Private/SAssetPicker.cpp b/Engine/Source/Editor/ContentBrowser/Private/SAssetPicker.cpp
--- a/Engine/Source/Editor/ContentBrowser/Private/SAssetPicker.cpp
+++ b/Engine/Source/Editor/ContentBrowser/Private/SAssetPicker.cpp
@@ -380,6 +380,7 @@ void SAssetPicker::Construct( const FArguments& InArgs )
 		.HighlightedText( HighlightText )
 		.ThumbnailLabel( ThumbnailLabel )
 		.AssetShowWarningText( InArgs._AssetPickerConfig.AssetShowWarningText)
+		.AssetShowWarningTextJustification( InArgs._AssetPickerConfig.AssetShowWarningTextJustification )
 		.AllowFocusOnSync(false)	// Stop the asset view from stealing focus (we're in control of that)
 		.HiddenColumnNames(InArgs._AssetPickerConfig.HiddenColumnNames)
 		.ListHiddenColumnNames(InArgs._AssetPickerConfig.ListHiddenColumnNames)
diff --git a/Engine/Source/Editor/ContentBrowser/Private/SAssetView.cpp b/Engine/Source/Editor/ContentBrowser/Private/SAssetView.cpp
--- a/Engine/Source/Editor/ContentBrowser/Private/SAssetView.cpp
+++ b/Engine/Source/Editor/ContentBrowser/Private/SAssetView.cpp
@@ -1421,6 +1421,7 @@ void SAssetView::Construct( const FArguments& InArgs )
 	AllowThumbnailHintLabel = InArgs._AllowThumbnailHintLabel;
 	InitialCategoryFilter = InArgs._InitialCategoryFilter;
 	AssetShowWarningText = InArgs._AssetShowWarningText;
+	AssetShowWarningTextJustification = InArgs._AssetShowWarningTextJustification;
 	bAllowDragging = InArgs._AllowDragging;
 	bAllowFocusOnSync = InArgs._AllowFocusOnSync;
 	HiddenColumnNames = DefaultHiddenColumnNames = InArgs._HiddenColumnNames;
@@ -1522,7 +1523,6 @@ void SAssetView::Construct( const FArguments& InArgs )
 
 			+ SOverlay::Slot()
 			.HAlign(HAlign_Fill)
-			.VAlign(VAlign_Center)
 			.Padding(FMargin(0, 14, 0, 0))
 			[
 				SNew(SScrollBox)
@@ -1533,7 +1533,7 @@ void SAssetView::Construct( const FArguments& InArgs )
 					// A warning to display when there are no assets to show, allow text copy
 					SNew(SMultiLineEditableText)
 					.IsReadOnly(true)
-					.Justification(ETextJustify::Left)
+					.Justification(AssetShowWarningTextJustification)
 					.Text(this, &SAssetView::GetAssetShowWarningText)
 					.AutoWrapText(true)
 				]
@@ -6369,7 +6369,7 @@ FText SAssetView::GetAssetShowWarningText() const
 		DropText = LOCTEXT( "DropFilesOrRightClick", "Drop files here or right click to create content." );
 	}
 	
-	return NothingToShowText.IsEmpty() ? DropText : FText::Format(LOCTEXT("NothingToShowPattern", "{0}\n\n{1}"), NothingToShowText, DropText);
+	return NothingToShowText.IsEmpty() ? DropText : FText::Format(LOCTEXT("NothingToShowPattern", "{0}\n{1}"), NothingToShowText, DropText);
 }
 
 bool SAssetView::HasSingleCollectionSource() const
diff --git a/Engine/Source/Editor/ContentBrowser/Public/IContentBrowserSingleton.h b/Engine/Source/Editor/ContentBrowser/Public/IContentBrowserSingleton.h
--- a/Engine/Source/Editor/ContentBrowser/Public/IContentBrowserSingleton.h
+++ b/Engine/Source/Editor/ContentBrowser/Public/IContentBrowserSingleton.h
@@ -277,8 +277,11 @@ struct FAssetPickerConfig
 	/** The default view mode */
 	EAssetViewType::Type InitialAssetViewType;
 
-	/** The text to show when there are no assets to show */
-	TAttribute< FText > AssetShowWarningText;
+	/** The warning text to show as heading (e.g. when there are no assets to show) */
+	TAttribute<FText> AssetShowWarningText;
+
+	/** The justification to use for the AssetShowWarning text (see AssetShowWarningText) */
+	TAttribute<ETextJustify::Type> AssetShowWarningTextJustification;
 
 	/** If set, view settings will be saved and loaded for the asset view using this name in ini files */
 	FString SaveSettingsName;
diff --git a/Engine/Source/Editor/ContentBrowser/Public/SAssetView.h b/Engine/Source/Editor/ContentBrowser/Public/SAssetView.h
--- a/Engine/Source/Editor/ContentBrowser/Public/SAssetView.h
+++ b/Engine/Source/Editor/ContentBrowser/Public/SAssetView.h
@@ -122,6 +122,7 @@ class SAssetView
 public:
 	SLATE_BEGIN_ARGS( SAssetView )
 		: _InitialCategoryFilter(EContentBrowserItemCategoryFilter::IncludeAll)
+		, _AssetShowWarningTextJustification(ETextJustify::Center)
 		, _ThumbnailLabel( EThumbnailLabel::ClassName )
 		, _AllowThumbnailHintLabel(true)
 		, _bShowPathViewFilters(false)
@@ -196,9 +197,12 @@ public:
 		/** Initial set of item categories that this view should show - may be adjusted further by things like CanShowClasses or legacy delegate bindings */
 		SLATE_ARGUMENT( EContentBrowserItemCategoryFilter, InitialCategoryFilter )
 
-		/** The warning text to display when there are no assets to show */
+		/** The warning text to show as heading (e.g. when there are no assets to show) */
 		SLATE_ATTRIBUTE( FText, AssetShowWarningText )
 
+		/** The justification to use for the AssetShowWarning text (see AssetShowWarningText) */
+		SLATE_ATTRIBUTE( ETextJustify::Type, AssetShowWarningTextJustification )
+
 		/** Attribute to determine what text should be highlighted */
 		SLATE_ATTRIBUTE( FText, HighlightedText )
 
@@ -936,7 +940,7 @@ private:
 	/** Is the no assets to show warning visible? */
 	UE_API EVisibility IsAssetShowWarningTextVisible() const;
 
-	/** Gets the text for displaying no assets to show warning */
+	/** Gets the warning text to show as heading (e.g. when there are no assets to show) */
 	UE_API FText GetAssetShowWarningText() const;
 
 	/** Whether we have a single source collection selected */
@@ -1337,8 +1341,11 @@ private:
 	/** The current thumbnail hint color and opacity*/
 	FLinearColor ThumbnailHintColorAndOpacity;
 
-	/** The text to show when there are no assets to show */
-	TAttribute< FText > AssetShowWarningText;
+	/** The warning text to show as heading (e.g. when there are no assets to show) */
+	TAttribute<FText> AssetShowWarningText;
+
+	/** The justification to use for the AssetShowWarning text (see AssetShowWarningText) */
+	TAttribute<ETextJustify::Type> AssetShowWarningTextJustification;
 
 	/** Initial set of item categories that this view should show - may be adjusted further by things like CanShowClasses or legacy delegate bindings */
 	EContentBrowserItemCategoryFilter InitialCategoryFilter;
diff --git a/Engine/Source/Editor/UnrealEd/Private/Dialogs/SDeleteAssetsDialog.cpp b/Engine/Source/Editor/UnrealEd/Private/Dialogs/SDeleteAssetsDialog.cpp
--- a/Engine/Source/Editor/UnrealEd/Private/Dialogs/SDeleteAssetsDialog.cpp
+++ b/Engine/Source/Editor/UnrealEd/Private/Dialogs/SDeleteAssetsDialog.cpp
@@ -798,7 +798,8 @@ TSharedRef<SWidget> SDeleteAssetsDialog::MakeAssetViewForReferencerAssets()
 	AssetPickerConfig.bShowBottomToolbar = false;
 	AssetPickerConfig.bAutohideSearchBar = true;
 
-	AssetPickerConfig.AssetShowWarningText = TAttribute< FText >( this, &SDeleteAssetsDialog::GetReferencingAssetsEmptyText );
+	AssetPickerConfig.AssetShowWarningText = TAttribute<FText>( this, &SDeleteAssetsDialog::GetReferencingAssetsEmptyText );
+	AssetPickerConfig.AssetShowWarningTextJustification = ETextJustify::Left;
 
 	AssetPickerConfig.InitialAssetViewType = EAssetViewType::List;
 	AssetPickerConfig.OnAssetsActivated = FOnAssetsActivated::CreateSP(this, &SDeleteAssetsDialog::OnAssetsActivated);
