EditorUsability : ContentBrowser
- Added AssetView attribute to allow warning text to be selectable or not through the asset picker config option
- Fixed warning text from stealing events in content browser (context menu) but still allowing delete dialog references to be selectable and scrollable
- Warning text is back to being centered vertically

#Synced-CL 52060601

diff --git a/Engine/Source/Editor/ContentBrowser/Private/SAssetPicker.cpp b/Engine/Source/Editor/ContentBrowser/Private/SAssetPicker.cpp
--- a/Engine/Source/Editor/ContentBrowser/Private/SAssetPicker.cpp
+++ b/Engine/Source/Editor/ContentBrowser/Private/SAssetPicker.cpp
@@ -381,6 +381,7 @@ void SAssetPicker::Construct( const FArguments& InArgs )
 		.ThumbnailLabel( ThumbnailLabel )
 		.AssetShowWarningText( InArgs._AssetPickerConfig.AssetShowWarningText)
 		.AssetShowWarningTextJustification( InArgs._AssetPickerConfig.AssetShowWarningTextJustification )
+		.AssetShowWarningTextSelectable( InArgs._AssetPickerConfig.AssetShowWarningTextSelectable )
 		.AllowFocusOnSync(false)	// Stop the asset view from stealing focus (we're in control of that)
 		.HiddenColumnNames(InArgs._AssetPickerConfig.HiddenColumnNames)
 		.ListHiddenColumnNames(InArgs._AssetPickerConfig.ListHiddenColumnNames)
diff --git a/Engine/Source/Editor/ContentBrowser/Private/SAssetView.cpp b/Engine/Source/Editor/ContentBrowser/Private/SAssetView.cpp
--- a/Engine/Source/Editor/ContentBrowser/Private/SAssetView.cpp
+++ b/Engine/Source/Editor/ContentBrowser/Private/SAssetView.cpp
@@ -1422,6 +1422,7 @@ void SAssetView::Construct( const FArguments& InArgs )
 	InitialCategoryFilter = InArgs._InitialCategoryFilter;
 	AssetShowWarningText = InArgs._AssetShowWarningText;
 	AssetShowWarningTextJustification = InArgs._AssetShowWarningTextJustification;
+	AssetShowWarningTextSelectable = InArgs._AssetShowWarningTextSelectable;
 	bAllowDragging = InArgs._AllowDragging;
 	bAllowFocusOnSync = InArgs._AllowFocusOnSync;
 	HiddenColumnNames = DefaultHiddenColumnNames = InArgs._HiddenColumnNames;
@@ -1523,6 +1524,7 @@ void SAssetView::Construct( const FArguments& InArgs )
 
 			+ SOverlay::Slot()
 			.HAlign(HAlign_Fill)
+			.VAlign(VAlign_Center)
 			.Padding(FMargin(0, 14, 0, 0))
 			[
 				SNew(SScrollBox)
@@ -1536,6 +1538,7 @@ void SAssetView::Construct( const FArguments& InArgs )
 					.Justification(AssetShowWarningTextJustification)
 					.Text(this, &SAssetView::GetAssetShowWarningText)
 					.AutoWrapText(true)
+					.IsEnabled(AssetShowWarningTextSelectable)
 				]
 			]
 
@@ -6332,7 +6335,8 @@ void SAssetView::SetSortParameters(
 
 EVisibility SAssetView::IsAssetShowWarningTextVisible() const
 {
-	return (FilteredAssetItems.Num() > 0 || bQuickFrontendListRefreshRequested) ? EVisibility::Collapsed : EVisibility::Visible;
+	const bool bTextSelectable = AssetShowWarningTextSelectable.Get(/* Default **/false);
+	return (FilteredAssetItems.Num() > 0 || bQuickFrontendListRefreshRequested) ? EVisibility::Collapsed : (bTextSelectable ? EVisibility::SelfHitTestInvisible : EVisibility::HitTestInvisible);
 }
 
 FText SAssetView::GetAssetShowWarningText() const
diff --git a/Engine/Source/Editor/ContentBrowser/Public/IContentBrowserSingleton.h b/Engine/Source/Editor/ContentBrowser/Public/IContentBrowserSingleton.h
--- a/Engine/Source/Editor/ContentBrowser/Public/IContentBrowserSingleton.h
+++ b/Engine/Source/Editor/ContentBrowser/Public/IContentBrowserSingleton.h
@@ -282,6 +282,9 @@ struct FAssetPickerConfig
 
 	/** The justification to use for the AssetShowWarning text (see AssetShowWarningText) */
 	TAttribute<ETextJustify::Type> AssetShowWarningTextJustification;
+	
+	/** Whether the warning text is selectable or not (see AssetShowWarningText) */
+	TAttribute<bool> AssetShowWarningTextSelectable;
 
 	/** If set, view settings will be saved and loaded for the asset view using this name in ini files */
 	FString SaveSettingsName;
diff --git a/Engine/Source/Editor/ContentBrowser/Public/SAssetView.h b/Engine/Source/Editor/ContentBrowser/Public/SAssetView.h
--- a/Engine/Source/Editor/ContentBrowser/Public/SAssetView.h
+++ b/Engine/Source/Editor/ContentBrowser/Public/SAssetView.h
@@ -123,6 +123,7 @@ public:
 	SLATE_BEGIN_ARGS( SAssetView )
 		: _InitialCategoryFilter(EContentBrowserItemCategoryFilter::IncludeAll)
 		, _AssetShowWarningTextJustification(ETextJustify::Center)
+		, _AssetShowWarningTextSelectable(false)
 		, _ThumbnailLabel( EThumbnailLabel::ClassName )
 		, _AllowThumbnailHintLabel(true)
 		, _bShowPathViewFilters(false)
@@ -203,6 +204,9 @@ public:
 		/** The justification to use for the AssetShowWarning text (see AssetShowWarningText) */
 		SLATE_ATTRIBUTE( ETextJustify::Type, AssetShowWarningTextJustification )
 
+		/** Whether the warning text is selectable or not (see AssetShowWarningText) */
+		SLATE_ATTRIBUTE( bool, AssetShowWarningTextSelectable )
+
 		/** Attribute to determine what text should be highlighted */
 		SLATE_ATTRIBUTE( FText, HighlightedText )
 
@@ -1346,6 +1350,9 @@ private:
 
 	/** The justification to use for the AssetShowWarning text (see AssetShowWarningText) */
 	TAttribute<ETextJustify::Type> AssetShowWarningTextJustification;
+	
+	/** Whether to allow selection of the warning text or not */
+	TAttribute<bool> AssetShowWarningTextSelectable;
 
 	/** Initial set of item categories that this view should show - may be adjusted further by things like CanShowClasses or legacy delegate bindings */
 	EContentBrowserItemCategoryFilter InitialCategoryFilter;
diff --git a/Engine/Source/Editor/UnrealEd/Private/Dialogs/SDeleteAssetsDialog.cpp b/Engine/Source/Editor/UnrealEd/Private/Dialogs/SDeleteAssetsDialog.cpp
--- a/Engine/Source/Editor/UnrealEd/Private/Dialogs/SDeleteAssetsDialog.cpp
+++ b/Engine/Source/Editor/UnrealEd/Private/Dialogs/SDeleteAssetsDialog.cpp
@@ -800,6 +800,7 @@ TSharedRef<SWidget> SDeleteAssetsDialog::MakeAssetViewForReferencerAssets()
 
 	AssetPickerConfig.AssetShowWarningText = TAttribute<FText>( this, &SDeleteAssetsDialog::GetReferencingAssetsEmptyText );
 	AssetPickerConfig.AssetShowWarningTextJustification = ETextJustify::Left;
+	AssetPickerConfig.AssetShowWarningTextSelectable = true;
 
 	AssetPickerConfig.InitialAssetViewType = EAssetViewType::List;
 	AssetPickerConfig.OnAssetsActivated = FOnAssetsActivated::CreateSP(this, &SDeleteAssetsDialog::OnAssetsActivated);
