Add an accessor to ContentBrowserSingleton that returns all selected assets across all content browser instances.

#Synced-CL 52159831

diff --git a/Engine/Source/Editor/ContentBrowser/Private/ContentBrowserSingleton.cpp b/Engine/Source/Editor/ContentBrowser/Private/ContentBrowserSingleton.cpp
--- a/Engine/Source/Editor/ContentBrowser/Private/ContentBrowserSingleton.cpp
+++ b/Engine/Source/Editor/ContentBrowser/Private/ContentBrowserSingleton.cpp
@@ -606,6 +606,17 @@ void FContentBrowserSingleton::GetSelectedAssets(TArray<FAssetData>& SelectedAss
 	}
 }
 
+void FContentBrowserSingleton::GetAllSelectedAssets(TArray<FAssetData>& SelectedAssets)
+{
+	for (const TWeakPtr<SContentBrowser>& Browser : AllContentBrowsers)
+	{
+		if (TSharedPtr<SContentBrowser> BrowserPin = Browser.Pin())
+		{
+			BrowserPin->GetSelectedAssets(SelectedAssets);
+		}
+	}
+}
+
 void FContentBrowserSingleton::GetSelectedFolders(TArray<FString>& SelectedFolders)
 {
 	if (PrimaryContentBrowser.IsValid())
diff --git a/Engine/Source/Editor/ContentBrowser/Private/ContentBrowserSingleton.h b/Engine/Source/Editor/ContentBrowser/Private/ContentBrowserSingleton.h
--- a/Engine/Source/Editor/ContentBrowser/Private/ContentBrowserSingleton.h
+++ b/Engine/Source/Editor/ContentBrowser/Private/ContentBrowserSingleton.h
@@ -96,6 +96,7 @@ public:
 	virtual void SyncBrowserToItems(const TArray<FContentBrowserItem>& ItemsToSync, bool bAllowLockedBrowsers = false, bool bFocusContentBrowser = true, const FName& InstanceName = FName(), bool bNewSpawnBrowser = false) override;
 	virtual void SyncBrowserTo(const FContentBrowserSelection& ItemSelection, bool bAllowLockedBrowsers = false, bool bFocusContentBrowser = true, const FName& InstanceName = FName(), bool bNewSpawnBrowser = false) override;
 	virtual void GetSelectedAssets(TArray<FAssetData>& SelectedAssets) override;
+	virtual void GetAllSelectedAssets(TArray<FAssetData>& SelectedAssets) override;
 	virtual void GetSelectedFolders(TArray<FString>& SelectedFolders) override;
 	virtual void GetSelectedPathViewFolders(TArray<FString>& SelectedFolders) override;
 PRAGMA_DISABLE_DEPRECATION_WARNINGS
diff --git a/Engine/Source/Editor/ContentBrowser/Public/IContentBrowserSingleton.h b/Engine/Source/Editor/ContentBrowser/Public/IContentBrowserSingleton.h
--- a/Engine/Source/Editor/ContentBrowser/Public/IContentBrowserSingleton.h
+++ b/Engine/Source/Editor/ContentBrowser/Public/IContentBrowserSingleton.h
@@ -765,6 +765,9 @@ public:
 	/** Generates a list of assets that are selected in the primary content browser */
 	virtual void GetSelectedAssets(TArray<FAssetData>& SelectedAssets) = 0;
 
+	/** Generates a list of assets that are selected across all open content browsers */
+	virtual void GetAllSelectedAssets(TArray<FAssetData>& SelectedAssets) = 0;
+
 	/** Generates a list of folders that are selected in the primary content browser */
 	virtual void GetSelectedFolders(TArray<FString>& SelectedFolders) = 0;
 
