[Unified Favorites] Content Browser Favorites Asset Filter

[Added] A front end filter for the CB asset favorites
[Note] Forced the favorite filter to be before the common filters

#Synced-CL 52005427

diff --git a/Engine/Source/Editor/ContentBrowser/Private/FrontendFilters.cpp b/Engine/Source/Editor/ContentBrowser/Private/FrontendFilters.cpp
--- a/Engine/Source/Editor/ContentBrowser/Private/FrontendFilters.cpp
+++ b/Engine/Source/Editor/ContentBrowser/Private/FrontendFilters.cpp
@@ -51,6 +51,9 @@
 #include "Templates/UnrealTemplate.h"
 #include "TextFilterKeyValueHandlers.h"
 #include "TextFilterValueHandlers.h"
+#include "DataStorage/Features.h"
+#include "DataStorage/MapKey.h"
+#include "Elements/Columns/TypedElementMiscColumns.h"
 #include "Textures/SlateIcon.h"
 #include "UObject/Class.h"
 #include "UObject/Object.h"
@@ -1678,3 +1681,48 @@ bool FFrontendFilter_Unsupported::PassesFilter(FAssetFilterType InItem) const
 {
 	return !InItem.IsSupported();
 }
+
+/////////////////////////////////////////
+// FFrontendFilter_Favorites
+/////////////////////////////////////////
+
+FFrontendFilter_Favorites::FFrontendFilter_Favorites(TSharedPtr<FFrontendFilterCategory> InCategory)
+	:FFrontendFilter(InCategory)
+{
+	using namespace UE::Editor::DataStorage;
+	DataStorage = GetMutableDataStorageFeature<ICoreProvider>(StorageFeatureName);
+}
+
+bool FFrontendFilter_Favorites::PassesFilter(FAssetFilterType InItem) const
+{
+	using namespace UE::Editor::DataStorage;
+
+	if (!DataStorage)
+	{
+		return false;
+	}
+
+	FMapKey Path;
+
+	if (InItem.GetItemType() == EContentBrowserItemFlags::Type_File)
+	{
+		if (FAssetData AssetData; InItem.Legacy_TryGetAssetData(AssetData))
+		{
+			Path = FMapKey(AssetData.GetSoftObjectPath());
+		}
+	}
+
+	if (!Path.IsSet())
+	{
+		return false;
+	}
+
+	const RowHandle Row = DataStorage->LookupMappedRow("Asset", Path);
+
+	if (DataStorage->IsRowAvailable(Row))
+	{
+		return DataStorage->HasColumns<FFavoriteTag>(Row);
+	}
+
+	return false;
+}
diff --git a/Engine/Source/Editor/ContentBrowser/Private/SFilterList.cpp b/Engine/Source/Editor/ContentBrowser/Private/SFilterList.cpp
--- a/Engine/Source/Editor/ContentBrowser/Private/SFilterList.cpp
+++ b/Engine/Source/Editor/ContentBrowser/Private/SFilterList.cpp
@@ -58,6 +58,19 @@ struct FGeometry;
 
 #define LOCTEXT_NAMESPACE "ContentBrowser"
 
+namespace UE::Editor::ContentBrowser
+{
+	namespace Private
+	{
+		static bool bUseFavoritesFilters = false;
+		static FAutoConsoleVariableRef CVarUseFavoritesFilters(
+			TEXT("UnifiedFavorites.ContentBrowser.Filter"),
+			Private::bUseFavoritesFilters,
+			TEXT("Use unified favorites filters in the Content Browser")
+		);
+	}
+}
+
 /////////////////////
 // SFilterList
 /////////////////////
@@ -94,6 +107,7 @@ void SFilterList::Construct( const FArguments& InArgs )
 	AllFrontendFilters_Internal.Add( MakeShareable(new FFrontendFilter_NotSourceControlled(DefaultCategory)) );
 	AllFrontendFilters_Internal.Add(MakeShareable(new FFrontendFilter_VirtualizedData(DefaultCategory)));
 	AllFrontendFilters_Internal.Add(MakeShared<FFrontendFilter_Unsupported>(DefaultCategory));
+	AllFrontendFilters_Internal.Add(MakeShareable(new FFrontendFilter_Favorites(DefaultCategory)));
 
 	// Add any global user-defined frontend filters
 	for (TObjectIterator<UContentBrowserFrontEndFilterExtension> ExtensionIt(RF_NoFlags); ExtensionIt; ++ExtensionIt)
@@ -537,6 +551,36 @@ void SFilterList::PopulateAddFilterMenu(UToolMenu* Menu, TSharedPtr<FFilterCateg
 					})),
 				EUserInterfaceActionType::RadioButton
 			);
+
+		}
+
+		if (UE::Editor::ContentBrowser::Private::bUseFavoritesFilters)
+		{
+			// Add Favorites section before the BasicFilterBarFiltersMenu filters
+			FToolMenuSection& FavoritesSection = Menu->FindOrAddSection(
+				"FilterBarFavorites",
+				LOCTEXT("FilterBarFavoritesSection", "Favorites Filter"),
+				FToolMenuInsert("FilterBarDisplay", EToolMenuInsertType::After));
+
+			// Find the Favorites filter and add it directly to the favorites section
+			for (const TSharedRef<FFilterBase<FAssetFilterType>>& Filter : AllFrontendFilters)
+			{
+				if (Filter->GetName() == "Favorites")
+				{
+					FavoritesSection.AddMenuEntry(
+						"FavoritesFilter",
+						Filter->GetDisplayName(),
+						Filter->GetToolTipText(),
+						FSlateIcon(FAppStyle::GetAppStyleSetName(), "Icons.Star"),
+						FUIAction(
+							FExecuteAction::CreateSP(this, &SFilterList::FrontendFilterClicked, Filter),
+							FCanExecuteAction(),
+							FIsActionChecked::CreateSP(this, &SFilterList::IsFrontendFilterInUse, Filter)),
+						EUserInterfaceActionType::ToggleButton
+					);
+					break;
+				}
+			}
 		}
 	}
 
diff --git a/Engine/Source/Editor/ContentBrowser/Public/FrontendFilters.h b/Engine/Source/Editor/ContentBrowser/Public/FrontendFilters.h
--- a/Engine/Source/Editor/ContentBrowser/Public/FrontendFilters.h
+++ b/Engine/Source/Editor/ContentBrowser/Public/FrontendFilters.h
@@ -20,6 +20,7 @@
 #include "Misc/TextFilterExpressionEvaluator.h"
 #include "Misc/TextFilterUtils.h"
 #include "SourceControlOperations.h"
+#include "Elements/Interfaces/TypedElementDataStorageInterface.h"
 #include "Templates/SharedPointer.h"
 #include "Types/SlateEnums.h"
 #include "UObject/NameTypes.h"
@@ -484,6 +485,26 @@ public:
 	UE_API virtual bool PassesFilter(FAssetFilterType InItem) const override;
 };
 
+/** A filter that displays only assets the are mark as Favorite */
+class FFrontendFilter_Favorites : public FFrontendFilter
+{
+public:
+	UE_API FFrontendFilter_Favorites(TSharedPtr<FFrontendFilterCategory> InCategory);
+	~FFrontendFilter_Favorites() = default;
+
+	// FFrontendFilter implementation
+	virtual FString GetName() const override { return "Favorites"; }
+	virtual FText GetDisplayName() const override { return LOCTEXT("FrontendFilter_Favorite", "Favorites"); }
+	virtual FText GetToolTipText() const override { return LOCTEXT("FrontendFilter_FavoriteTooltip", "Filter by favorite assets."); }
+	virtual FLinearColor GetColor() const override { return FStyleColors::AccentYellow.GetSpecifiedColor(); }
+
+	// IFilter implementation
+	UE_API virtual bool PassesFilter(FAssetFilterType InItem) const override;
+
+private:
+	UE::Editor::DataStorage::ICoreProvider* DataStorage;
+};
+
 #undef LOCTEXT_NAMESPACE
 
 #undef UE_API
