Refactor search override toggle handling to preserve and forward existing search text instead of unconditionally clearing it
Move sorting guard inside the rename check to avoid skipping the pending sort flag reset
Consolidate override toggle callback from asset view into content browser for more accurate state-aware behavior
Fixed an  issue in the CB where a force refresh request wouldn't  always refrest the Item widget list

#Synced-CL 52352994

diff --git a/Engine/Source/Editor/ContentBrowser/Private/SAssetView.cpp b/Engine/Source/Editor/ContentBrowser/Private/SAssetView.cpp
--- a/Engine/Source/Editor/ContentBrowser/Private/SAssetView.cpp
+++ b/Engine/Source/Editor/ContentBrowser/Private/SAssetView.cpp
@@ -2404,6 +2404,11 @@ void SAssetView::Tick( const FGeometry& AllottedGeometry, const double InCurrent
 
 	if (HasItemsPendingFilter() || bForceViewUpdate)
 	{
+		if (bForceViewUpdate)
+		{
+			RefreshList();
+		}
+
 		bForceViewUpdate = false;
 
 		const double TickStartTime = FPlatformTime::Seconds();
@@ -4423,17 +4428,18 @@ void SAssetView::SortList(bool bSyncToSelection)
 {
 	TRACE_CPUPROFILER_EVENT_SCOPE(SAssetView::SortList);
 
-	PRAGMA_DISABLE_EXPERIMENTAL_WARNINGS
-	if (AssetSearchOverride.IsValid() && AssetSearchOverride->IsOverrideActive())
-	{
-		bPendingSortFilteredItems = false;
-		return;
-	}
-	PRAGMA_ENABLE_EXPERIMENTAL_WARNINGS
-
 	if ( !IsRenamingAsset() )
 	{
-		SortManager->SortList(FilteredAssetItems, MajorityAssetType, CustomColumns);
+		PRAGMA_DISABLE_EXPERIMENTAL_WARNINGS
+		if (AssetSearchOverride.IsValid() && AssetSearchOverride->IsOverrideActive())
+		{
+			// Do nothing a the moment
+		}
+		PRAGMA_ENABLE_EXPERIMENTAL_WARNINGS
+		else
+		{
+			SortManager->SortList(FilteredAssetItems, MajorityAssetType, CustomColumns);
+		}
 
 		// Update the thumbnails we were using since the order has changed
 		bPendingUpdateThumbnails = true;
@@ -6925,18 +6931,12 @@ void SAssetView::SetAssetSearchOverride(TSharedPtr<UE::Editor::ContentBrowser::E
 	AssetSearchOverride = MoveTemp(InOverride);
 	if (AssetSearchOverride.IsValid())
 	{
-		AssetSearchOverride->GetInternalApi().OnOverrideActiveChanged.AddSP(this, &SAssetView::OnSearchOverrideToggled);
 		AssetSearchOverride->GetInternalApi().OnRequestSetUserSearching.BindSP(this, &SAssetView::SetUserSearching);
 		AssetSearchOverride->GetInternalApi().OnRequestQuickFrontendListRefresh.BindSP(this, &SAssetView::RequestQuickFrontendListRefresh);
 	}
 	RequestQuickFrontendListRefresh();
 }
 
-void SAssetView::OnSearchOverrideToggled()
-{
-	RequestQuickFrontendListRefresh();
-}
-
 TWeakPtr<FAssetViewSortManager> SAssetView::GetSortManager() const
 {
 	return SortManager;
diff --git a/Engine/Source/Editor/ContentBrowser/Private/SContentBrowser.cpp b/Engine/Source/Editor/ContentBrowser/Private/SContentBrowser.cpp
--- a/Engine/Source/Editor/ContentBrowser/Private/SContentBrowser.cpp
+++ b/Engine/Source/Editor/ContentBrowser/Private/SContentBrowser.cpp
@@ -1031,10 +1031,7 @@ TSharedRef<SWidget> SContentBrowser::CreateAssetView(const FContentBrowserConfig
 	if (UE::Editor::ContentBrowser::Extension::IAssetSearchOverride* ActiveOverride = FModuleManager::LoadModuleChecked<FContentBrowserModule>("ContentBrowser").GetActiveAssetSearchOverride())
 	{
 		LegacyContentSourceWidgets->AssetSearchOverrideInstance = ActiveOverride->MakePerViewOverride();
-
-		// When the override is toggled, clear the search text and reset the filter pipeline.
-		LegacyContentSourceWidgets->AssetSearchOverrideInstance->GetInternalApi().OnOverrideActiveChanged.AddSP(this, &SContentBrowser::ClearSearchText);
-
+		LegacyContentSourceWidgets->AssetSearchOverrideInstance->GetInternalApi().OnOverrideActiveChanged.AddSP(this, &SContentBrowser::OnSearchOverrideToggled);
 		// Pass the instance to the asset view so it can take over the last step of the filtering and the sort.
 		LegacyContentSourceWidgets->AssetViewPtr->SetAssetSearchOverride(LegacyContentSourceWidgets->AssetSearchOverrideInstance);
 	}
@@ -2424,15 +2421,6 @@ void SContentBrowser::SetSearchText(const FText& InSearchText)
 	}
 }
 
-void SContentBrowser::ClearSearchText()
-{
-	SetSearchBoxText(FText::GetEmpty());
-	if (LegacyContentSourceWidgets->SearchBoxPtr)
-	{
-		LegacyContentSourceWidgets->SearchBoxPtr->SetText(FText::GetEmpty());
-	}
-}
-
 void SContentBrowser::SetAssetViewUserSearching(bool bIsSearching)
 {
 	LegacyContentSourceWidgets->AssetViewPtr->SetUserSearching(bIsSearching);
@@ -4008,6 +3996,22 @@ TSharedRef<SWidget> SContentBrowser::CreateLegacyAssetViewWidgets()
 	return FinalWidget;
 }
 
+void SContentBrowser::OnSearchOverrideToggled()
+{
+	if (LegacyContentSourceWidgets->AssetSearchOverrideInstance && LegacyContentSourceWidgets->AssetSearchOverrideInstance->IsOverrideActive())
+	{
+		// Use enter since this was based on a user action
+		LegacyContentSourceWidgets->AssetSearchOverrideInstance->OnSearchTextCommitted(LegacyContentSourceWidgets->SearchBoxPtr->GetText(), ETextCommit::OnEnter);
+		SetSearchBoxText(FText::GetEmpty());
+	}
+	else
+	{
+		SetSearchBoxText(LegacyContentSourceWidgets->SearchBoxPtr->GetText());
+	}
+	
+	LegacyContentSourceWidgets->AssetViewPtr->RequestQuickFrontendListRefresh();
+}
+
 FReply SContentBrowser::BackClicked()
 {
 	HistoryManager.GoBack();
diff --git a/Engine/Source/Editor/ContentBrowser/Private/SContentBrowser.h b/Engine/Source/Editor/ContentBrowser/Private/SContentBrowser.h
--- a/Engine/Source/Editor/ContentBrowser/Private/SContentBrowser.h
+++ b/Engine/Source/Editor/ContentBrowser/Private/SContentBrowser.h
@@ -340,9 +340,6 @@ private:
 	/** Called when the editable text needs to be set or cleared */
 	void SetSearchBoxText(const FText& InSearchText);
 
-	/** Clears the search box text and resets the filter pipeline; used when an asset search override is toggled. */
-	void ClearSearchText();
-
 	/** Forwards the user-searching state to the asset view */
 	void SetAssetViewUserSearching(bool bIsSearching);
 
@@ -705,6 +702,8 @@ private:
 	/** Create the content browser widgets (SPathView, SAssetView, search and filter bars etc) */
 	TSharedRef<SWidget> CreateLegacyAssetViewWidgets();
 	
+	void OnSearchOverrideToggled();
+
 private:
 
 	/**
diff --git a/Engine/Source/Editor/ContentBrowser/Public/SAssetView.h b/Engine/Source/Editor/ContentBrowser/Public/SAssetView.h
--- a/Engine/Source/Editor/ContentBrowser/Public/SAssetView.h
+++ b/Engine/Source/Editor/ContentBrowser/Public/SAssetView.h
@@ -1141,9 +1141,6 @@ private:
 	TSharedPtr<UE::Editor::ContentBrowser::Extension::ICBSearchOverrideInstance> AssetSearchOverride;
 	PRAGMA_ENABLE_EXPERIMENTAL_WARNINGS
 
-	/** Called when the override is toggled; triggers a fresh frontend filter pass. */
-	void OnSearchOverrideToggled();
-
 	/** Called when the asset view is asked to start to create a temporary item */
 	FOnAssetViewNewItemRequested OnNewItemRequested;
 
