Fix FDirectoryPath content picker closing immediately when property already has a valid path

SPathPicker::OnItemSelectionChanged was firing OnPathSelected for programmatic selections
(ESelectInfo::Direct) triggered during SPathView::Populate. When the DefaultPath was set,
FScopedSelectionChangedEvent detected a selection change on destruction and called
TreeSelectionChanged with ESelectInfo::Direct, which propagated to
SPathPicker::OnItemSelectionChanged and fired OnPathSelected. This caused
FDirectoryPathStructCustomization::OnPathPicked to immediately dismiss the picker menu
before the user could interact with it.

The fix skips OnPathSelected for ESelectInfo::Direct, allowing programmatic default path
selection without triggering the path-picked callback that dismisses the menu.

#Synced-CL 52342074

diff --git a/Engine/Source/Editor/ContentBrowser/Private/SPathPicker.cpp b/Engine/Source/Editor/ContentBrowser/Private/SPathPicker.cpp
--- a/Engine/Source/Editor/ContentBrowser/Private/SPathPicker.cpp
+++ b/Engine/Source/Editor/ContentBrowser/Private/SPathPicker.cpp
@@ -87,6 +87,11 @@ void SPathPicker::Construct( const FArguments& InArgs )
 
 void SPathPicker::OnItemSelectionChanged(const FContentBrowserItem& SelectedItem, ESelectInfo::Type SelectInfo)
 {
+	if (SelectInfo == ESelectInfo::Direct)
+	{
+		return;
+	}
+
 	FName SelectedPackagePath;
 	if (SelectedItem.IsFolder())
 	{
