When right clicking asset, show reason asset cannot be made public in tooltip if asset class does not allow it.
Also fixed rename being blocked by not being able to set private assets to public in /Game
#rnx

#Synced-CL 52465190

diff --git a/Engine/Plugins/Editor/ContentBrowser/ContentBrowserAssetDataSource/Source/ContentBrowserAssetDataSource/Private/ContentBrowserAssetDataCore.cpp b/Engine/Plugins/Editor/ContentBrowser/ContentBrowserAssetDataSource/Source/ContentBrowserAssetDataSource/Private/ContentBrowserAssetDataCore.cpp
--- a/Engine/Plugins/Editor/ContentBrowser/ContentBrowserAssetDataSource/Source/ContentBrowserAssetDataSource/Private/ContentBrowserAssetDataCore.cpp
+++ b/Engine/Plugins/Editor/ContentBrowser/ContentBrowserAssetDataSource/Source/ContentBrowserAssetDataSource/Private/ContentBrowserAssetDataCore.cpp
@@ -2070,7 +2070,7 @@ bool CanModifyAssetAccessSpecifier(IAssetTools* InAssetTools, const UContentBrow
 		FNameBuilder AssetPath;
 		FTopLevelAssetPath(AssetPayload->GetAssetData().PackageName, AssetPayload->GetAssetData().AssetName).AppendString(AssetPath);
 
-		return InAssetTools->CanAssetBePublic(AssetPath);
+		return InAssetTools->CanAssetBePublic(AssetPath, nullptr);
 	}
 
 	return false;
diff --git a/Engine/Source/Developer/AssetTools/Private/AssetRenameManager.cpp b/Engine/Source/Developer/AssetTools/Private/AssetRenameManager.cpp
--- a/Engine/Source/Developer/AssetTools/Private/AssetRenameManager.cpp
+++ b/Engine/Source/Developer/AssetTools/Private/AssetRenameManager.cpp
@@ -1503,7 +1503,8 @@ void FAssetRenameManager::SetupPublicAssets(TArray<FAssetRenameDataWithReference
 		if (!ReferencingAssetInDifferentMountPoint.IsNone())
 		{
 			const FString OldAssetPath = RenameData.OldObjectPath.GetAssetPathString();
-			if (AssetTools.CanAssetBePublic(OldAssetPath))
+			FText DisallowedSetPublicReason;
+			if (AssetTools.CanAssetBePublic(OldAssetPath, &DisallowedSetPublicReason))
 			{
 				PackagesToMakePublic.Add(Package);
 			}
@@ -1514,15 +1515,17 @@ void FAssetRenameManager::SetupPublicAssets(TArray<FAssetRenameDataWithReference
 				if (RenameData.bCreateRedirector)
 				{
 					RenameData.FailureReason = FText::Format(
-						LOCTEXT("AssetCannotBePublicForRedirector", "Cannot move asset to {0} because it cannot be made public in order to be referenced from a different plugin by its redirector"),
-						FText::FromStringView(AssetTools.GetUserFacingPluginName(NewPackageMountPoint)));
+						LOCTEXT("AssetCannotBePublicForRedirector", "Cannot move asset to {0} because it cannot be made public in order to be referenced from a different plugin by its redirector. Additional Reason: {1}"),
+						FText::FromStringView(AssetTools.GetUserFacingPluginName(NewPackageMountPoint)),
+						DisallowedSetPublicReason);
 				}
 				else
 				{
 					RenameData.FailureReason = FText::Format(
-						LOCTEXT("AssetCannotBePublic", "Cannot move asset to {0} because it cannot be made public in order to be referenced from a different plugin by {1}"),
+						LOCTEXT("AssetCannotBePublic", "Cannot move asset to {0} because it cannot be made public in order to be referenced from a different plugin by {1}. Additional Reason: {2}"),
 						FText::FromStringView(AssetTools.GetUserFacingPluginName(NewPackageMountPoint)),
-						FText::FromString(AssetTools.GetUserFacingLongPackageName(ReferencingAssetInDifferentMountPoint)));
+						FText::FromString(AssetTools.GetUserFacingLongPackageName(ReferencingAssetInDifferentMountPoint)),
+						DisallowedSetPublicReason);
 				}
 			}
 		}
diff --git a/Engine/Source/Developer/AssetTools/Private/AssetTools.cpp b/Engine/Source/Developer/AssetTools/Private/AssetTools.cpp
--- a/Engine/Source/Developer/AssetTools/Private/AssetTools.cpp
+++ b/Engine/Source/Developer/AssetTools/Private/AssetTools.cpp
@@ -467,11 +467,11 @@ namespace UE::AssetTools::Private
 		TEXT("Enables the Experimental Public Asset Feature (False: disabled, True:enabled")
 	);
 
-	static bool bEnableEpicInternalAssetFeature = true;
+	static bool bEnableEpicInternalAssetFeature = false;
 	static FAutoConsoleVariableRef CVarEnableEpicInternalAssetFeature(
 		TEXT("AssetTools.EnableEpicInternalAssetFeature"),
 		bEnableEpicInternalAssetFeature,
-		TEXT("Enables the Experimental Epic Internal Asset Feature UI (False: disabled, True: enabled")
+		TEXT("Enables the Epic Internal Asset Feature UI (False: disabled, True: enabled")
 	);
 
 	/** 
@@ -7252,13 +7252,13 @@ void UAssetToolsImpl::UnregisterCanMigrateAsset(const FName OwnerName)
 	CanMigrateAssetDelegates.Remove(OwnerName);
 }
 
-bool UAssetToolsImpl::CanAssetBePublic(FStringView AssetPath) const
+bool UAssetToolsImpl::CanAssetBePublic(FStringView AssetPath, FText* DisallowReason) const
 {
 	for (const TPair<FName, UE::AssetTools::FCanAssetBePublic>& Pair : CanAssetBePublicDelegates)
 	{
 		if (Pair.Value.IsBound())
 		{
-			if (!Pair.Value.Execute(AssetPath))
+			if (!Pair.Value.Execute(AssetPath, DisallowReason))
 			{
 				return false;
 			}
diff --git a/Engine/Source/Developer/AssetTools/Private/AssetTools.h b/Engine/Source/Developer/AssetTools/Private/AssetTools.h
--- a/Engine/Source/Developer/AssetTools/Private/AssetTools.h
+++ b/Engine/Source/Developer/AssetTools/Private/AssetTools.h
@@ -184,7 +184,7 @@ public:
 	virtual void UnregisterDefaultAssetNameForClass(const UClass* Class) override;
 	virtual void RegisterCanMigrateAsset(const FName OwnerName, UE::AssetTools::FCanMigrateAsset Delegate) override;
 	virtual void UnregisterCanMigrateAsset(const FName OwnerName) override;
-	virtual bool CanAssetBePublic(FStringView AssetPath) const override;
+	virtual bool CanAssetBePublic(FStringView AssetPath, FText* DisallowReason = nullptr) const override;
 	virtual void RegisterCanAssetBePublic(const FName OwnerName, UE::AssetTools::FCanAssetBePublic Delegate) override;
 	virtual void UnregisterCanAssetBePublic(const FName OwnerName) override;
 
diff --git a/Engine/Source/Developer/AssetTools/Public/IAssetTools.h b/Engine/Source/Developer/AssetTools/Public/IAssetTools.h
--- a/Engine/Source/Developer/AssetTools/Public/IAssetTools.h
+++ b/Engine/Source/Developer/AssetTools/Public/IAssetTools.h
@@ -39,7 +39,7 @@ namespace UE::AssetTools
 
 	DECLARE_DELEGATE_RetVal_OneParam(bool, FCanMigrateAsset, FName);
 
-	DECLARE_DELEGATE_RetVal_OneParam(bool, FCanAssetBePublic, FStringView /*AssetPath*/);
+	DECLARE_DELEGATE_RetVal_TwoParams(bool, FCanAssetBePublic, FStringView /*AssetPath*/, FText* /*DisallowReason*/);
 
 	DECLARE_MULTICAST_DELEGATE_TwoParams(FShouldCreateAssetsAsExternallyReferenceableForPath, FStringView /*AssetPath*/, TOptional<bool>& /*ExternallyReferenceable*/);
 
@@ -873,7 +873,7 @@ public:
 	virtual void UnregisterCanMigrateAsset(const FName OwnerName) = 0;
 
 	/** Returns whether the specified asset can be public (referenceable from another mount point / plugin) */
-	virtual bool CanAssetBePublic(FStringView AssetPath) const = 0;
+	virtual bool CanAssetBePublic(FStringView AssetPath, FText* DisallowReason = nullptr) const = 0;
 
 	/**
 	 * Register/unregister delegates to specify whether an asset can be made public (referenceable from another mount point / plugin)
diff --git a/Engine/Source/Editor/ContentBrowser/Private/AssetContextMenu.cpp b/Engine/Source/Editor/ContentBrowser/Private/AssetContextMenu.cpp
--- a/Engine/Source/Editor/ContentBrowser/Private/AssetContextMenu.cpp
+++ b/Engine/Source/Editor/ContentBrowser/Private/AssetContextMenu.cpp
@@ -639,7 +639,8 @@ void FAssetContextMenu::AddExploreMenuOptions(UToolMenu* Menu)
 
 bool FAssetContextMenu::AddPublicStateMenuOptions(UToolMenu* Menu)
 {
-	if (!(bCanExecuteSetPublicAsset || bCanExecuteSetEpicInternalAsset || bCanExecuteSetPrivateAsset || bCanExecuteBulkSetPublicAsset || bCanExecuteBulkSetEpicInternalAsset || bCanExecuteBulkSetPrivateAsset))
+	bool bHasDisallowSetPublicReason = !DisallowSetPublicReason.IsEmpty();
+	if (!(bHasDisallowSetPublicReason || bCanShowSetAssetAccessMenu || bCanExecuteSetPublicAsset || bCanExecuteSetEpicInternalAsset || bCanExecuteSetPrivateAsset || bCanExecuteBulkSetPublicAsset || bCanExecuteBulkSetEpicInternalAsset || bCanExecuteBulkSetPrivateAsset))
 	{
 		return false;
 	}
@@ -650,12 +651,12 @@ bool FAssetContextMenu::AddPublicStateMenuOptions(UToolMenu* Menu)
 
 		if (SelectedFiles.Num() == 1)
 		{
-			if (bCanExecuteSetPublicAsset || bCanExecuteSetEpicInternalAsset || bCanExecuteSetPrivateAsset)
+			if (bCanExecuteSetPublicAsset || bCanExecuteSetEpicInternalAsset || bCanExecuteSetPrivateAsset || bHasDisallowSetPublicReason || bCanShowSetAssetAccessMenu)
 			{
 				Section.AddMenuEntry(
 					"PublicAsset",
 					LOCTEXT("PublicAssetToggle", "Public Asset"),
-					LOCTEXT("PublicAssetToggleTooltip", "Sets the asset to be referencable by other Plugins"),
+					bHasDisallowSetPublicReason ? DisallowSetPublicReason : LOCTEXT("PublicAssetToggleTooltip", "Sets the asset to be referencable by other Plugins"),
 					FSlateIcon(FAppStyle::GetAppStyleSetName(), FAssetViewItem::GetPublicAssetAccessSpecifierIconStyleName(TEXT(".Small"))),
 					FUIAction(
 						FExecuteAction::CreateSP(this, &FAssetContextMenu::ExecuteSetAssetAccessSpecifier, EAssetAccessSpecifier::Public),
@@ -670,7 +671,7 @@ bool FAssetContextMenu::AddPublicStateMenuOptions(UToolMenu* Menu)
 					Section.AddMenuEntry(
 						"EpicInternalAsset",
 						LOCTEXT("SetInternalAsset", "Epic Internal Asset"),
-						LOCTEXT("SetInternalAssetTooltip", "Sets the asset to be referencable by Epic internal plugins and mount points"),
+						bHasDisallowSetPublicReason ? DisallowSetPublicReason : LOCTEXT("SetInternalAssetTooltip", "Sets the asset to be referencable by Epic internal plugins and mount points"),
 						FSlateIcon(FAppStyle::GetAppStyleSetName(), FAssetViewItem::GetEpicInternalAssetAccessSpecifierIconStyleName(TEXT(".Small"))),
 						FUIAction(
 							FExecuteAction::CreateSP(this, &FAssetContextMenu::ExecuteSetAssetAccessSpecifier, EAssetAccessSpecifier::EpicInternal),
@@ -1519,61 +1520,80 @@ void FAssetContextMenu::CacheCanExecuteVars()
 {
 	bCanExecuteFindInExplorer = false;
 
+	bCanShowSetAssetAccessMenu = false;
+
+	bCanExecuteSetPublicAsset = false;
+	bCanExecuteSetEpicInternalAsset = false;
+	bCanExecuteSetPrivateAsset = false;
+
+	bCanExecuteBulkSetPublicAsset = false;
+	bCanExecuteBulkSetEpicInternalAsset = false;
+	bCanExecuteBulkSetPrivateAsset = false;
+
+	DisallowSetPublicReason = FText::GetEmpty();
+
 	FContentBrowserModule& ContentBrowserModule = FModuleManager::GetModuleChecked<FContentBrowserModule>(TEXT("ContentBrowser"));
 
+	const FOnContentBrowserCanShowSetAssetAccessMenu& OnContentBrowserCanShowSetAssetAccessMenu = ContentBrowserModule.GetOnContentBrowserCanShowSetAssetAccessMenu();
+
 	const bool bIsEpicInternalAssetFeatureEnabled = AssetContextMenu::IsEpicInternalAssetFeatureEnabled();
 	const TSharedPtr<FPathPermissionList>& ShowEpicInternalContentPermissionList = ContentBrowserModule.Get().GetShowEpicInternalContentPermissionList();
 	bool bAnyPathPassesShowEpicInternal = !ShowEpicInternalContentPermissionList->HasFiltering();
 
 	const TSharedPtr<FPathPermissionList>& ShowPrivateContentPermissionList = ContentBrowserModule.Get().GetShowPrivateContentPermissionList();
 
-	if (SelectedFiles.Num() == 1)
+	auto CheckCanExecuteFindInExplorer = [](const FContentBrowserItem& SelectedItem) -> bool
 	{
-		const FContentBrowserItem& SelectedItem = SelectedFiles[0];
-
 		FString ItemFilename;
-		if (!bCanExecuteFindInExplorer && SelectedItem.GetItemPhysicalPath(ItemFilename) && FPaths::FileExists(ItemFilename))
-		{
-			bCanExecuteFindInExplorer = true;
-		}
+		return SelectedItem.GetItemPhysicalPath(ItemFilename) && FPaths::FileExists(ItemFilename);
+	};
 
-		bCanExecuteSetPublicAsset = false;
-		bCanExecuteSetEpicInternalAsset = false;
-		bCanExecuteSetPrivateAsset = false;
+	auto CheckCanAssetBePublic = [](const FStringView AssetPath, FText* DisallowReason) -> bool
+	{
+		return IAssetTools::Get().CanAssetBePublic(AssetPath, DisallowReason);
+	};
+
+	auto CheckCanShowSetAssetAccessMenu = [&OnContentBrowserCanShowSetAssetAccessMenu](const FStringView AssetPath, const FAssetData& Asset) -> bool
+	{
+		return !OnContentBrowserCanShowSetAssetAccessMenu.IsBound() || OnContentBrowserCanShowSetAssetAccessMenu.Execute(AssetPath, Asset);
+	};
+
+	if (SelectedFiles.Num() == 1)
+	{
+		const FContentBrowserItem& SelectedItem = SelectedFiles[0];
+		bCanExecuteFindInExplorer = CheckCanExecuteFindInExplorer(SelectedItem);
 
 		if (SelectedItem.CanEdit())
 		{
 			const FNameBuilder ItemInternalPath(SelectedItem.GetInternalPath());
 			const FStringView AssetPath(ItemInternalPath);
-
-			// Only show menu entries when asset can be made externally referenceable
-			if (IAssetTools::Get().CanAssetBePublic(AssetPath))
+			const FAssetData AssetData = IAssetRegistry::GetChecked().GetAssetByObjectPath(FSoftObjectPath(AssetPath));
+			if (AssetData.IsValid() && CheckCanShowSetAssetAccessMenu(AssetPath, AssetData))
 			{
-				bCanExecuteSetPublicAsset = true;
-				if (bIsEpicInternalAssetFeatureEnabled && (ShowEpicInternalContentPermissionList->PassesStartsWithFilter(AssetPath) || ShowPrivateContentPermissionList->PassesStartsWithFilter(AssetPath)))
+				bCanShowSetAssetAccessMenu = true;
+
+				// Only show menu entries when asset can be made externally referenceable 
+				if (CheckCanAssetBePublic(AssetPath, &DisallowSetPublicReason))
 				{
-					bCanExecuteSetEpicInternalAsset = true;
+					bCanExecuteSetPublicAsset = true;
+					bCanExecuteSetEpicInternalAsset = bIsEpicInternalAssetFeatureEnabled && (ShowEpicInternalContentPermissionList->PassesStartsWithFilter(AssetPath) || ShowPrivateContentPermissionList->PassesStartsWithFilter(AssetPath));
+					bCanExecuteSetPrivateAsset = true;
+				}
+				else if (AssetData.GetAssetAccessSpecifier() != EAssetAccessSpecifier::Private)
+				{
+					// Allow user to manually fix asset that should be private
+					bCanExecuteSetPrivateAsset = true;
 				}
-				bCanExecuteSetPrivateAsset = true;
 			}
 		}
-
-		bCanExecuteBulkSetPublicAsset = false;
-		bCanExecuteBulkSetEpicInternalAsset = false;
-		bCanExecuteBulkSetPrivateAsset = false;
 	}
 	else
 	{
-		bCanExecuteBulkSetPublicAsset = false;
-		bCanExecuteBulkSetEpicInternalAsset = false;
-		bCanExecuteBulkSetPrivateAsset = false;
-
 		for (const FContentBrowserItem& SelectedItem : SelectedFiles)
 		{
-			FString ItemFilename;
-			if (!bCanExecuteFindInExplorer && SelectedItem.GetItemPhysicalPath(ItemFilename) && FPaths::FileExists(ItemFilename))
+			if (!bCanExecuteFindInExplorer)
 			{
-				bCanExecuteFindInExplorer = true;
+				bCanExecuteFindInExplorer = CheckCanExecuteFindInExplorer(SelectedItem);
 			}
 
 			if (SelectedItem.CanEdit())
@@ -1584,7 +1604,7 @@ void FAssetContextMenu::CacheCanExecuteVars()
 				if (AssetData.IsValid())
 				{
 					// If any selected asset can be made public, show menu entries to bulk set as Public/Internal/Private
-					if (!bCanExecuteBulkSetPublicAsset && IAssetTools::Get().CanAssetBePublic(AssetPath))
+					if (!bCanExecuteBulkSetPublicAsset && CheckCanAssetBePublic(AssetPath, nullptr))
 					{
 						bCanExecuteBulkSetPublicAsset = true;
 						bCanExecuteBulkSetEpicInternalAsset = bIsEpicInternalAssetFeatureEnabled;
@@ -1594,7 +1614,7 @@ void FAssetContextMenu::CacheCanExecuteVars()
 					// If any selected can be set to Epic internal
 					if (!bAnyPathPassesShowEpicInternal && bIsEpicInternalAssetFeatureEnabled)
 					{
-						if (IAssetTools::Get().CanAssetBePublic(AssetPath))
+						if (CheckCanAssetBePublic(AssetPath, nullptr))
 						{
 							if (ShowEpicInternalContentPermissionList->PassesStartsWithFilter(AssetPath) || (ShowPrivateContentPermissionList->HasFiltering() && ShowPrivateContentPermissionList->PassesStartsWithFilter(AssetPath)))
 							{
@@ -1635,6 +1655,8 @@ void FAssetContextMenu::CacheCanExecuteVars()
 	// CVar bypass to always show. For repairing broken assets stuck with the wrong flags.
 	if (AssetContextMenu::IsAlwaysShowSetAssetScopeMenus())
 	{
+		bCanShowSetAssetAccessMenu = true;
+
 		if (SelectedFiles.Num() == 1)
 		{
 			bCanExecuteSetPublicAsset = true;
diff --git a/Engine/Source/Editor/ContentBrowser/Private/AssetContextMenu.h b/Engine/Source/Editor/ContentBrowser/Private/AssetContextMenu.h
--- a/Engine/Source/Editor/ContentBrowser/Private/AssetContextMenu.h
+++ b/Engine/Source/Editor/ContentBrowser/Private/AssetContextMenu.h
@@ -192,6 +192,8 @@ private:
 	/** Cached CanExecute vars */
 	bool bCanExecuteFindInExplorer = false;
 
+	bool bCanShowSetAssetAccessMenu = false;
+
 	bool bCanExecuteSetPublicAsset = false;
 	bool bCanExecuteSetEpicInternalAsset = false;
 	bool bCanExecuteSetPrivateAsset = false;
@@ -199,4 +201,6 @@ private:
 	bool bCanExecuteBulkSetPublicAsset = false;
 	bool bCanExecuteBulkSetEpicInternalAsset = false;
 	bool bCanExecuteBulkSetPrivateAsset = false;
-};
+
+	FText DisallowSetPublicReason;
+};
\ No newline at end of file
diff --git a/Engine/Source/Editor/ContentBrowser/Public/ContentBrowserDelegates.h b/Engine/Source/Editor/ContentBrowser/Public/ContentBrowserDelegates.h
--- a/Engine/Source/Editor/ContentBrowser/Public/ContentBrowserDelegates.h
+++ b/Engine/Source/Editor/ContentBrowser/Public/ContentBrowserDelegates.h
@@ -114,6 +114,9 @@ DECLARE_DELEGATE_RetVal(TSharedRef<FExtender>, FContentBrowserMenuExtender);
 DECLARE_DELEGATE_RetVal_OneParam(TSharedRef<FExtender>, FContentBrowserMenuExtender_SelectedAssets, const TArray<FAssetData>& /*SelectedAssets*/);
 DECLARE_DELEGATE_RetVal_OneParam(TSharedRef<FExtender>, FContentBrowserMenuExtender_SelectedPaths, const TArray<FString>& /*SelectedPaths*/);
 
+/** Called to check if should show menu entries to set asset public/private */
+DECLARE_DELEGATE_RetVal_TwoParams(bool, FOnContentBrowserCanShowSetAssetAccessMenu, const FStringView /*AssetPath*/, const FAssetData& /*Asset*/);
+
 /** Called to request the menu when right clicking on an asset */
 DECLARE_DELEGATE_RetVal_ThreeParams(TSharedPtr<SWidget>, FOnGetFolderContextMenu, const TArray<FString>& /*SelectedPaths*/, FContentBrowserMenuExtender_SelectedPaths /*MenuExtender*/, FOnCreateNewFolder /*CreationDelegate*/);
 
diff --git a/Engine/Source/Editor/ContentBrowser/Public/ContentBrowserModule.h b/Engine/Source/Editor/ContentBrowser/Public/ContentBrowserModule.h
--- a/Engine/Source/Editor/ContentBrowser/Public/ContentBrowserModule.h
+++ b/Engine/Source/Editor/ContentBrowser/Public/ContentBrowserModule.h
@@ -197,6 +197,7 @@ public:
 	FOnAssetPathChanged& GetOnAssetPathChanged() { return OnAssetPathChanged; }
 	FOnContentBrowserSettingChanged& GetOnContentBrowserSettingChanged() { return OnContentBrowserSettingChanged; }
 	FOnDeleteFolders& GetOnDeleteFolders() { return OnDeleteFolders; }
+	FOnContentBrowserCanShowSetAssetAccessMenu& GetOnContentBrowserCanShowSetAssetAccessMenu() { return OnContentBrowserCanShowSetAssetAccessMenu; }
 
 	/** Override list of paths to select by default */
 	FDefaultSelectedPathsDelegate& GetDefaultSelectedPathsDelegate() { return DefaultSelectedPathsDelegate; }
@@ -310,4 +311,7 @@ private:
 
 	/** The currently active asset search override, if any. Only one may be registered at a time. */
 	TUniquePtr<UE::Editor::ContentBrowser::Extension::IAssetSearchOverride> ActiveAssetSearchOverride = nullptr;
+
+	/** Override to selectively hide menu entries for viewing and setting asset as Public/Private */
+	FOnContentBrowserCanShowSetAssetAccessMenu OnContentBrowserCanShowSetAssetAccessMenu;
 };
