Add an `OnAllKnownItemsAvailable` notification to the Content Browser search override interface, fired once per filter session when the frontend filter has finished processing all known items
Dispatch the notification from the asset view after filtering, guard it against duplicate calls in `FInternalApiWrapper` consistent with the `OnFilteringReset` pattern, and add a stub implementation in the Semantic Search integration
Also wire up the previously no-op sort override to forward sorting to the search override

#Synced-CL 52476245

#Synced-CL 52476245
#Synced-CL 52476245

diff --git a/Engine/Plugins/Experimental/SemanticSearch/Source/SemanticSearchEditorIntegrations/Private/ContentBrowser/SemanticSearchCB.cpp b/Engine/Plugins/Experimental/SemanticSearch/Source/SemanticSearchEditorIntegrations/Private/ContentBrowser/SemanticSearchCB.cpp
--- a/Engine/Plugins/Experimental/SemanticSearch/Source/SemanticSearchEditorIntegrations/Private/ContentBrowser/SemanticSearchCB.cpp
+++ b/Engine/Plugins/Experimental/SemanticSearch/Source/SemanticSearchEditorIntegrations/Private/ContentBrowser/SemanticSearchCB.cpp
@@ -73,6 +73,11 @@ namespace Private
 
 		}
 
+		virtual void OnAllKnownItemsAvailableImplementation() override
+		{
+
+		}
+
 	protected:
 
 		// Incrementing the epoch silently invalidates all in-flight callbacks:
diff --git a/Engine/Source/Editor/ContentBrowser/Private/Experimental/ContentBrowserSearchOverride.cpp b/Engine/Source/Editor/ContentBrowser/Private/Experimental/ContentBrowserSearchOverride.cpp
--- a/Engine/Source/Editor/ContentBrowser/Private/Experimental/ContentBrowserSearchOverride.cpp
+++ b/Engine/Source/Editor/ContentBrowser/Private/Experimental/ContentBrowserSearchOverride.cpp
@@ -51,7 +51,17 @@ FInternalApiWrapper& ICBSearchOverrideInstance::GetInternalApi()
 void FInternalApiWrapper::OnFilteringReset()
 {
 	PublishedResults.Reset();
+	bAllKnownItemsAvailableNotified = false;
 	Instance.OnFilteringResetImplementation();
 }
 
+void FInternalApiWrapper::OnAllKnownItemsAvailable()
+{
+	if (!bAllKnownItemsAvailableNotified)
+	{
+		bAllKnownItemsAvailableNotified = true;
+		Instance.OnAllKnownItemsAvailableImplementation();
+	}
+}
+
 } // namespace UE::Editor::ContentBrowser::Extension
diff --git a/Engine/Source/Editor/ContentBrowser/Private/Experimental/ContentBrowserSearchOverrideInternal.h b/Engine/Source/Editor/ContentBrowser/Private/Experimental/ContentBrowserSearchOverrideInternal.h
--- a/Engine/Source/Editor/ContentBrowser/Private/Experimental/ContentBrowserSearchOverrideInternal.h
+++ b/Engine/Source/Editor/ContentBrowser/Private/Experimental/ContentBrowserSearchOverrideInternal.h
@@ -18,6 +18,7 @@ namespace UE::Editor::ContentBrowser::Extension
 		{}
 
 		void OnFilteringReset();
+		void OnAllKnownItemsAvailable();
 
 		/** Fires when the override is toggled on or off. */
 		FSimpleMulticastDelegate OnOverrideActiveChanged;
@@ -31,6 +32,8 @@ namespace UE::Editor::ContentBrowser::Extension
 		/** Results set by PublishResults(); consumed by the asset view on the next tick. */
 		TArray<TSharedPtr<FAssetViewItem>> PublishedResults;
 
+		bool bAllKnownItemsAvailableNotified = false;
+
 	private:
 		ICBSearchOverrideInstance& Instance;
 	};
diff --git a/Engine/Source/Editor/ContentBrowser/Private/SAssetView.cpp b/Engine/Source/Editor/ContentBrowser/Private/SAssetView.cpp
--- a/Engine/Source/Editor/ContentBrowser/Private/SAssetView.cpp
+++ b/Engine/Source/Editor/ContentBrowser/Private/SAssetView.cpp
@@ -2701,6 +2701,11 @@ void SAssetView::ProcessItemsPendingFilter(const double TickStartTime)
 
 			FilterResults.Append(MoveTemp(AssetSearchOverride->GetInternalApi().PublishedResults));
 			AssetSearchOverride->GetInternalApi().PublishedResults.Reset();
+
+			if (!Items->HasItemsPendingFilter())
+			{
+				AssetSearchOverride->GetInternalApi().OnAllKnownItemsAvailable();
+			}
 		}
 	PRAGMA_ENABLE_EXPERIMENTAL_WARNINGS
 
@@ -4434,7 +4439,7 @@ void SAssetView::SortList(bool bSyncToSelection)
 		PRAGMA_DISABLE_EXPERIMENTAL_WARNINGS
 		if (AssetSearchOverride.IsValid() && AssetSearchOverride->IsOverrideActive())
 		{
-			// Do nothing a the moment
+			AssetSearchOverride->SortItemList(FilteredAssetItems);
 		}
 		PRAGMA_ENABLE_EXPERIMENTAL_WARNINGS
 		else
diff --git a/Engine/Source/Editor/ContentBrowser/Public/Experimental/ContentBrowserSearchOverride.h b/Engine/Source/Editor/ContentBrowser/Public/Experimental/ContentBrowserSearchOverride.h
--- a/Engine/Source/Editor/ContentBrowser/Public/Experimental/ContentBrowserSearchOverride.h
+++ b/Engine/Source/Editor/ContentBrowser/Public/Experimental/ContentBrowserSearchOverride.h
@@ -54,7 +54,6 @@ namespace UE::Editor::ContentBrowser::Extension
 		 */
 		virtual void OnItemsAvailable(TArrayView<const TSharedPtr<FAssetViewItem>> Items) = 0;
 
-
 		virtual void SortItemList(TArrayView<TSharedPtr<FAssetViewItem>> Items) = 0;
 
 		/**
@@ -82,6 +81,11 @@ namespace UE::Editor::ContentBrowser::Extension
 		 */
 		virtual void OnFilteringResetImplementation() = 0;
 
+		/**
+		 * Called when the extension has received all the known items that it has to process.
+		 */
+		virtual void OnAllKnownItemsAvailableImplementation() = 0;
+
 	private:
 		bool bIsOverrideActive = false;
 		TUniquePtr<FInternalApiWrapper> InternalApi;
