UAF Browser: Plugin Filters + Search Suggestions
- Move SFilterList & SAssetPicker to internal.
- Expose methods to access filter list & toggle filters
- Add callback for on text changed.
- Add nullable API to get asset picker from tagged asset browser (Nullable implies result may not exist / be null).

#tests editor
#Synced-CL 52588081

diff --git a/Engine/Plugins/Experimental/UAF/UAF/Source/UAFEditor/Private/Common/SUAFBrowser.cpp b/Engine/Plugins/Experimental/UAF/UAF/Source/UAFEditor/Private/Common/SUAFBrowser.cpp
--- a/Engine/Plugins/Experimental/UAF/UAF/Source/UAFEditor/Private/Common/SUAFBrowser.cpp
+++ b/Engine/Plugins/Experimental/UAF/UAF/Source/UAFEditor/Private/Common/SUAFBrowser.cpp
@@ -11,23 +11,28 @@
 #include "Common/AssetPreview/AnimSequenceAssetPreview.h"
 #include "Common/AssetPreview/IUAFAssetPreview.h"
 #include "Common/AssetPreview/UAFAssetPreviewFactorySubsystem.h"
+#include "Common/SUAFBrowserFilterSuggestionStrip.h"
 #include "Common/UAFBrowserMenuContext.h"
 #include "Common/UAFBrowserTabSummoner.h"
 #include "ContentBrowserModule.h"
 #include "Editor.h"
 #include "Engine/Blueprint.h"
 #include "Factories/BlueprintFactory.h"
+#include "Filters/CustomClassFilterData.h"
 #include "Framework/Application/SlateApplication.h"
 #include "Kismet2/KismetEditorUtilities.h"
 #include "Kismet2/SClassPickerDialog.h"
 #include "IAssetTools.h"
 #include "IContentBrowserSingleton.h"
+#include "Interfaces/IPluginManager.h"
 #include "IWorkspaceEditor.h"
 #include "IWorkspaceEditorModule.h"
 #include "Misc/ConfigCacheIni.h"
 #include "Misc/PackageName.h"
 #include "NewAssetContextMenu.h"
+#include "SAssetPicker.h"
 #include "SAssetView.h"
+#include "SFilterList.h"
 #include "Sound/SoundWave.h"
 #include "Settings/UAFEditorProjectSettings.h"
 #include "Settings/UAFEditorUserSettings.h"
@@ -397,22 +402,25 @@ TSharedRef<SWidget> SUAFBrowser::Construct_GetTaggedAssetBrowser()
 
 	if (UTaggedAssetBrowserConfiguration* CategoryConfiguration = Cast<UTaggedAssetBrowserConfiguration>(CategoryConfigurationObject))
 	{
+		SAssignNew(SuggestionStrip, SUAFBrowserFilterSuggestionStrip)
+			.OnFilterSuggestionSelected(TDelegate<void()>::CreateSP(this, &SUAFBrowser::OnFilterSuggestionSelected));
+
 		// We intentionally do not restrict AvailableClasses as this will filter types we may not know about from beind displayed
 		// Ex: Even in contexts such as filtering by tag
-		TSharedRef<STaggedAssetBrowser> TaggedBrowser = SNew(STaggedAssetBrowser, *CategoryConfiguration)
+		SAssignNew(TaggedAssetBrowser, STaggedAssetBrowser, *CategoryConfiguration)
 			.AvailableClasses({})
 			.DefaultDetailsTabConfiguration(DefaultDetailsTabConfiguration)
 			.InterfaceOverrideProfiles(OverrideProfiles)
 			.OnOverrideAssetPickerConfig(STaggedAssetBrowser::FOnOverrideAssetPickerConfig::CreateSP(this, &SUAFBrowser::OnOverrideAssetPickerConfig))
 			.OnActiveSectionChanged(STaggedAssetBrowser::FOnActiveSectionChanged::CreateSP(this, &SUAFBrowser::OnActiveSectionChanged))
-			.OnPrimaryFilterSelectionChanged(STaggedAssetBrowser::FOnPrimaryFilterSelectionChanged::CreateSP(this, &SUAFBrowser::OnPrimaryFilterSelectionChanged));
-
-		TaggedAssetBrowserWidgetWeak = TaggedBrowser;
+			.OnPrimaryFilterSelectionChanged(STaggedAssetBrowser::FOnPrimaryFilterSelectionChanged::CreateSP(this, &SUAFBrowser::OnPrimaryFilterSelectionChanged))
+			.AdditionalBottomWidget(SuggestionStrip);
 
 		// @TODO: Implement below, not used at the moment.
 		// Initialize current section from the browser's initial state so the Add button is correct before any section switch
-		CurrentActiveSection = TaggedBrowser->GetActiveSection();
-		return TaggedBrowser;
+		CurrentActiveSection = TaggedAssetBrowser->GetActiveSection();
+
+		return TaggedAssetBrowser.ToSharedRef();
 	}
 
 	return SNullWidget::NullWidget;
@@ -445,6 +453,7 @@ void SUAFBrowser::OnOverrideAssetPickerConfig(FAssetPickerConfig& Config)
 	static const FString UAFBrowserConfigName = TEXT("UAFBrowser");
 
 	Config.OnExtendAssetPickerTopBar.BindSP(this, &SUAFBrowser::OnExtendAssetPickerTopBar);
+	Config.OnSearchTextChanged.BindSP(this, &SUAFBrowser::OnContentSearchTextChanged);
 	Config.bAddFilterUI = true;
 	Config.bAllowDragging = true;
 	Config.bShowPathInColumnView = true;
@@ -462,7 +471,7 @@ void SUAFBrowser::OnOverrideAssetPickerConfig(FAssetPickerConfig& Config)
 
 		Config.HiddenColumnNames = UAFSettings->UAFBrowserHiddenColumns;
 	}
-
+	
 	Config.OnGetAssetContextMenu.Unbind();
 	Config.OnAssetsActivated.Unbind();
 	Config.OnGetCustomAssetToolTip.Unbind();
@@ -481,6 +490,8 @@ void SUAFBrowser::OnOverrideAssetPickerConfig(FAssetPickerConfig& Config)
 	Config.OnColumnVisibilityChanged.BindSP(this, &SUAFBrowser::OnBrowserColumnVisibilityChanged);
 	Config.SaveAssetViewSettingsDelegates.Add(&SaveAssetViewSettingsDelegate);
 	Config.LoadAssetViewSettingsDelegates.Add(&LoadAssetViewSettingsDelegate);
+
+	PopulatePluginFilters(Config);
 }
 
 void SUAFBrowser::OnExtendAssetPickerTopBar(TSharedRef<SHorizontalBox> InAssetPickerTopBar)
@@ -711,6 +722,7 @@ TSharedRef<SWidget> SUAFBrowser::BuildAddImportButtons()
 	if (Mode != EUAFAddButtonMode::ImportOnly)
 	{
 		Box->AddSlot()
+			.Padding(FMargin(2, 0, 6, 0))
 			.AutoWidth()
 			[
 				SNew(SPositiveActionButton)
@@ -722,6 +734,7 @@ TSharedRef<SWidget> SUAFBrowser::BuildAddImportButtons()
 	if (Mode == EUAFAddButtonMode::ImportOnly || Mode == EUAFAddButtonMode::Both)
 	{
 		Box->AddSlot()
+			.Padding(FMargin(2, 0, 6, 0))
 			.AutoWidth()
 			[
 				SNew(SActionButton)
@@ -773,18 +786,18 @@ FString SUAFBrowser::GetNewAssetPath() const
 TArray<UClass*> SUAFBrowser::GetCurrentSectionFilterClasses() const
 {
 	TArray<UClass*> Result;
-	TSharedPtr<STaggedAssetBrowser> Browser = TaggedAssetBrowserWidgetWeak.Pin();
-	if (!Browser)
+
+	if (!TaggedAssetBrowser)
 	{
 		return Result;
 	}
 
 	// Prefer the currently selected primary filters; fall back to all filters in the section
 	// so that GetAddButtonMode() returns the correct mode even before the user picks a filter.
-	TArray<const UTaggedAssetBrowserFilterBase*> Filters = Browser->GetSelectedPrimaryFilters();
+	TArray<const UTaggedAssetBrowserFilterBase*> Filters = TaggedAssetBrowser->GetSelectedPrimaryFilters();
 	if (Filters.IsEmpty())
 	{
-		Filters = Browser->GetAllPrimaryFilters();
+		Filters = TaggedAssetBrowser->GetAllPrimaryFilters();
 	}
 
 	for (const UTaggedAssetBrowserFilterBase* Filter : Filters)
@@ -962,6 +975,220 @@ void SUAFBrowser::CreateAssetWithNameDialog(UFactory* InFactory, const FText& As
 	}
 }
 
+void SUAFBrowser::PopulatePluginFilters(FAssetPickerConfig& Config)
+{
+	const UUAFEditorProjectSettings* Settings = GetDefault<UUAFEditorProjectSettings>();
+
+	TSharedRef<FFrontendFilterCategory> PluginsCategory = MakeShared<FFrontendFilterCategory>(
+		LOCTEXT("PluginsFilterCategory", "Plugins"),
+		LOCTEXT("PluginsFilterCategoryTooltip", "Filter assets by plugin")
+	);
+
+	PluginFilters.Reset();
+
+	TSet<FName> ExcludedPlugins(Settings->ExcludedBrowserPluginFilters);
+	TSet<FName> AddedPluginNames;
+
+	// Auto-discover project plugins with content
+	for (const TSharedRef<IPlugin>& Plugin : IPluginManager::Get().GetEnabledPluginsWithContent())
+	{
+		if (Plugin->GetType() != EPluginType::Project)
+		{
+			continue;
+		}
+
+		const FName PluginFName(*Plugin->GetName());
+		if (ExcludedPlugins.Contains(PluginFName))
+		{
+			continue;
+		}
+
+		FString ContentPath = Plugin->GetMountedAssetPath();
+		TSharedRef<FFrontendFilter_UAFPlugin> Filter = MakeShared<FFrontendFilter_UAFPlugin>(
+			PluginsCategory,
+			Plugin->GetName(),
+			Plugin->GetFriendlyName(),
+			MoveTemp(ContentPath)
+		);
+
+		Config.ExtraFrontendFilters.Add(Filter);
+		PluginFilters.Add(Filter);
+		AddedPluginNames.Add(PluginFName);
+	}
+
+	// Add explicitly configured additional plugins not already discovered
+	for (const FName& PluginName : Settings->AdditionalBrowserPluginFilters)
+	{
+		if (AddedPluginNames.Contains(PluginName))
+		{
+			continue;
+		}
+
+		TSharedPtr<IPlugin> Plugin = IPluginManager::Get().FindPlugin(PluginName.ToString());
+		if (!Plugin.IsValid() || !Plugin->CanContainContent())
+		{
+			continue;
+		}
+
+		FString ContentPath = FString::Printf(TEXT("/%s/"), *Plugin->GetName());
+		TSharedRef<FFrontendFilter_UAFPlugin> Filter = MakeShared<FFrontendFilter_UAFPlugin>(
+			PluginsCategory,
+			Plugin->GetName(),
+			Plugin->GetFriendlyName(),
+			MoveTemp(ContentPath)
+		);
+
+		Config.ExtraFrontendFilters.Add(Filter);
+		PluginFilters.Add(Filter);
+	}
+}
+
+void SUAFBrowser::OnContentSearchTextChanged(const FText& SearchText)
+{
+	// Capture the currently focused widget so we can restore focus after a chip click
+	if (TSharedPtr<SWidget> FocusedWidget = FSlateApplication::Get().GetUserFocusedWidget(0))
+	{
+		if (FocusedWidget->GetType() != FName(TEXT("SUAFBrowserChipButton")))
+		{
+			WeakSearchBox = FocusedWidget;
+		}
+	}
+
+	if (!SuggestionStrip.IsValid())
+	{
+		return;
+	}
+
+	const FString Query = SearchText.ToString();
+	TArray<TSharedRef<FFilterSuggestion>> Matches;
+
+	if (!Query.IsEmpty())
+	{
+		// Plugin filters - matched by plugin name or display name
+		for (const TSharedRef<FFrontendFilter_UAFPlugin>& Filter : PluginFilters)
+		{
+			if (Filter->GetPluginName().Contains(Query, ESearchCase::IgnoreCase) ||
+				Filter->GetDisplayName().ToString().Contains(Query, ESearchCase::IgnoreCase))
+			{
+				TSharedRef<FFilterSuggestion> Suggestion = MakeShared<FFilterSuggestion>();
+				Suggestion->DisplayName = Filter->GetDisplayName();
+				Suggestion->IconName    = Filter->GetIconName();
+				Suggestion->Color       = Filter->GetColor();
+				Suggestion->Category    = SUAFBrowserFilterSuggestionStrip::CategoryPlugins;
+				Suggestion->OnActivate.BindLambda([Filter]() { Filter->SetActive(true); });
+				Matches.Add(Suggestion);
+			}
+		}
+
+		// Obtain the filter list from the cached widget chain
+		TSharedPtr<SFilterList> FilterList;
+		if (TaggedAssetBrowser.IsValid())
+		{
+			if (TSharedPtr<SAssetPicker> AssetPicker = TaggedAssetBrowser->GetAssetPicker())
+			{
+				FilterList = AssetPicker->GetFilterList();
+			}
+		}
+
+		// All other frontend filters from the filter bar (excluding plugin filters already handled above)
+		if (FilterList.IsValid())
+		{
+			for (const TSharedRef<FFrontendFilter>& Filter : FilterList->GetAllFrontendFilters())
+			{
+				const bool bIsPluginFilter = PluginFilters.ContainsByPredicate([&Filter](const TSharedRef<FFrontendFilter_UAFPlugin>& PluginFilter) 
+				{ 
+					return PluginFilter == Filter; 
+				});
+
+				if (bIsPluginFilter)
+				{
+					continue;
+				}
+
+				if (Filter->GetName().Contains(Query, ESearchCase::IgnoreCase) ||
+					Filter->GetDisplayName().ToString().Contains(Query, ESearchCase::IgnoreCase))
+				{
+					TSharedRef<FFilterSuggestion> Suggestion = MakeShared<FFilterSuggestion>();
+					Suggestion->DisplayName = Filter->GetDisplayName();
+					Suggestion->IconName = Filter->GetIconName();
+					Suggestion->Color = Filter->GetColor();
+					Suggestion->Category = SUAFBrowserFilterSuggestionStrip::CategoryMisc;
+					TWeakPtr<FFrontendFilter> WeakFilter = Filter;
+					Suggestion->OnActivate.BindLambda([WeakFilter]()
+					{
+						if (TSharedPtr<FFrontendFilter> PinnedFilter = WeakFilter.Pin())
+						{
+							PinnedFilter->SetActive(true);
+						}
+					});
+					Matches.Add(Suggestion);
+				}
+			}
+
+			// Asset class/type filters (e.g. Animation Sequence, Blend Space)
+			TWeakPtr<SFilterList> WeakFilterList = FilterList;
+			for (const TSharedRef<FCustomClassFilterData>& ClassFilter : FilterList->GetCustomClassFilters())
+			{
+				if (ClassFilter->GetName().ToString().Contains(Query, ESearchCase::IgnoreCase) ||
+					ClassFilter->GetFilterName().Contains(Query, ESearchCase::IgnoreCase))
+				{
+					TSharedRef<FFilterSuggestion> Suggestion = MakeShared<FFilterSuggestion>();
+					Suggestion->DisplayName = ClassFilter->GetName();
+					Suggestion->IconName = FName("ClassIcon.Object");
+					Suggestion->Color = ClassFilter->GetColor();
+					Suggestion->Category  = SUAFBrowserFilterSuggestionStrip::CategoryTypes;
+					FTopLevelAssetPath ClassPath = ClassFilter->GetClassPathName();
+					Suggestion->OnActivate.BindLambda([WeakFilterList, ClassPath]()
+					{
+						if (TSharedPtr<SFilterList> PinnedFilterList = WeakFilterList.Pin())
+						{
+							PinnedFilterList->ActivateClassFilterByPath(ClassPath);
+						}
+					});
+					Matches.Add(Suggestion);
+				}
+			}
+		}
+	}
+
+	SuggestionStrip->SetSuggestions(Matches);
+}
+
+void SUAFBrowser::OnFilterSuggestionSelected()
+{
+	// Clear the search text — this triggers OnContentSearchTextChanged("") which collapses the suggestion strip
+	if (TaggedAssetBrowser.IsValid())
+	{
+		if (TSharedPtr<SAssetPicker> AssetPicker = TaggedAssetBrowser->GetAssetPicker())
+		{
+			AssetPicker->ClearSearchText();
+		}
+	}
+
+	// Return keyboard focus to the search box so the user can continue typing
+	if (TSharedPtr<SWidget> SearchBox = WeakSearchBox.Pin())
+	{
+		FSlateApplication::Get().SetKeyboardFocus(SearchBox, EFocusCause::SetDirectly);
+	}
+}
+
+FReply SUAFBrowser::OnPreviewKeyDown(const FGeometry& MyGeometry, const FKeyEvent& InKeyEvent)
+{
+	// When Tab is pressed while the search box has focus and the suggestion strip has chips,
+	// redirect focus to the first chip instead of moving to the next focusable widget.
+	if (InKeyEvent.GetKey() == EKeys::Tab && !InKeyEvent.IsShiftDown())
+	{
+		TSharedPtr<SWidget> FocusedWidget = FSlateApplication::Get().GetUserFocusedWidget(0);
+		if (SuggestionStrip.IsValid() && SuggestionStrip->HasAnySuggestions()
+			&& FocusedWidget.IsValid() && FocusedWidget == WeakSearchBox.Pin())
+		{
+			SuggestionStrip->FocusFirstChip();
+			return FReply::Handled();
+		}
+	}
+	return FReply::Unhandled();
+}
+
 } // namespace UE::UAF::Editor
 
 #undef LOCTEXT_NAMESPACE
diff --git a/Engine/Plugins/Experimental/UAF/UAF/Source/UAFEditor/Private/Common/SUAFBrowser.h b/Engine/Plugins/Experimental/UAF/UAF/Source/UAFEditor/Private/Common/SUAFBrowser.h
--- a/Engine/Plugins/Experimental/UAF/UAF/Source/UAFEditor/Private/Common/SUAFBrowser.h
+++ b/Engine/Plugins/Experimental/UAF/UAF/Source/UAFEditor/Private/Common/SUAFBrowser.h
@@ -3,7 +3,9 @@
 #pragma once
 
 #include "CoreMinimal.h"
+#include "Common/UAFBrowserFrontendFilters.h"
 #include "ContentBrowserDelegates.h"
+#include "IContentBrowserSingleton.h"
 #include "Framework/Docking/TabManager.h"
 #include "Widgets/Docking/SDockTab.h"
 #include "Widgets/Layout/SSplitter.h"
@@ -19,10 +21,12 @@ namespace UE::UAF::Editor
 	class IUAFAssetPreview;
 };
 
-class STaggedAssetBrowser;
 class UTaggedAssetBrowserSection;
 struct FAssetPickerConfig;
+class FFrontendFilter;
 class SHorizontalBox;
+class STaggedAssetBrowser;
+class SUAFBrowserFilterSuggestionStrip;
 
 namespace UE::UAF::Editor
 {
@@ -82,6 +86,15 @@ private:
 	/** Callback to override asset picker configs for our tagged asset browser */
 	void OnOverrideAssetPickerConfig(FAssetPickerConfig& Config);
 
+	/** Populates Config.ExtraFrontendFilters with plugin-based filters discovered from project settings and IPluginManager. */
+	void PopulatePluginFilters(FAssetPickerConfig& Config);
+
+	/** Called when the content area search box text changes; updates the suggestion strip. */
+	void OnContentSearchTextChanged(const FText& SearchText);
+
+	/** Called when any suggestion chip is clicked; returns keyboard focus to the search box. */
+	void OnFilterSuggestionSelected();
+
 	/** Callback to customize top bar of our asset browser. To add layout control buttons. */
 	void OnExtendAssetPickerTopBar(TSharedRef<SHorizontalBox> InAssetPickerTopBar);
 
@@ -129,6 +142,9 @@ private:
 
 private:
 
+	/** Intercepts Tab from the search box to redirect focus to the first suggestion chip. */
+	virtual FReply OnPreviewKeyDown(const FGeometry& MyGeometry, const FKeyEvent& InKeyEvent) override;
+
 	/** Callback for when an asset is selected. */
 	void OnAssetSelected(const FAssetData& InAssetData);
 
@@ -146,9 +162,6 @@ private:
 	/** App Hosting this Browser, to query settings against */
 	TWeakPtr<UE::Workspace::IWorkspaceEditor> WeakHostingApp = nullptr;
 
-	/** The tagged asset browser widget, kept so we can query GetActiveSection() directly. */
-	TWeakPtr<STaggedAssetBrowser> TaggedAssetBrowserWidgetWeak;
-
 	/** The currently active section, updated via OnActiveSectionChanged. */
 	TWeakObjectPtr<const UTaggedAssetBrowserSection> CurrentActiveSection;
 
@@ -175,6 +188,18 @@ private:
 
 	/** Delegate fired to refresh asset browser whenever config .ini changes occur. */
 	FLoadAssetViewSettingsDelegate LoadAssetViewSettingsDelegate;
+
+	/** Plugin frontend filters added to the content area filter bar. Kept so we can match them against search text. */
+	TArray<TSharedRef<FFrontendFilter_UAFPlugin>> PluginFilters;
+
+	/** The tagged asset browser widget; held so we can access the asset picker and filter list directly. */
+	TSharedPtr<STaggedAssetBrowser> TaggedAssetBrowser;
+
+	/** Horizontally-scrolling strip of filter suggestion chips, shown when search text matches a known filter. */
+	TSharedPtr<SUAFBrowserFilterSuggestionStrip> SuggestionStrip;
+
+	/** Weak reference to the content area search box widget, captured when OnContentSearchTextChanged fires. Used to restore focus after a chip click. */
+	TWeakPtr<SWidget> WeakSearchBox;
 };
 
 } // namespace UE::UAF::Editor
\ No newline at end of file
diff --git a/Engine/Plugins/Experimental/UAF/UAF/Source/UAFEditor/Private/Common/SUAFBrowserFilterSuggestionStrip.cpp b/Engine/Plugins/Experimental/UAF/UAF/Source/UAFEditor/Private/Common/SUAFBrowserFilterSuggestionStrip.cpp
new file mode 100644
--- /dev/null
+++ b/Engine/Plugins/Experimental/UAF/UAF/Source/UAFEditor/Private/Common/SUAFBrowserFilterSuggestionStrip.cpp
@@ -0,0 +1,365 @@
+// Copyright Epic Games, Inc. All Rights Reserved.
+
+#include "Common/SUAFBrowserFilterSuggestionStrip.h"
+
+#include "Framework/Application/SlateApplication.h"
+#include "Styling/AppStyle.h"
+#include "Widgets/Images/SImage.h"
+#include "Widgets/Input/SButton.h"
+#include "Widgets/Layout/SBorder.h"
+#include "Widgets/Layout/SScrollBox.h"
+#include "Widgets/SBoxPanel.h"
+#include "Widgets/Text/STextBlock.h"
+
+#define LOCTEXT_NAMESPACE "UAFBrowserFilterSuggestions"
+
+// ---------------------------------------------------------------------------
+// SUAFBrowserChipButton
+// ---------------------------------------------------------------------------
+
+/**
+ * Minimal SButton subclass that passes a hover attribute during construction so that
+ * the button renders with hover visuals whenever it has keyboard focus.
+ * SetHover is a protected SWidget method, so it must be called from within Construct.
+ */
+class SUAFBrowserChipButton : public SButton
+{
+public:
+	SLATE_BEGIN_ARGS(SUAFBrowserChipButton)
+		: _ButtonStyle(&FAppStyle::Get().GetWidgetStyle<FButtonStyle>("SimpleButton"))
+		, _ContentPadding(FMargin(6.f, 2.f))
+	{}
+		SLATE_STYLE_ARGUMENT(FButtonStyle, ButtonStyle)
+		SLATE_ATTRIBUTE(FMargin, ContentPadding)
+		SLATE_ATTRIBUTE(FText, ToolTipText)
+		SLATE_EVENT(FOnClicked, OnClicked)
+		SLATE_DEFAULT_SLOT(FArguments, Content)
+	SLATE_END_ARGS()
+
+	virtual void OnMouseEnter(const FGeometry& MyGeometry, const FPointerEvent& MouseEvent) override
+	{
+		bHovered = true;
+		SButton::OnMouseEnter(MyGeometry, MouseEvent);
+	}
+
+	virtual void OnMouseLeave(const FPointerEvent& MouseEvent) override
+	{
+		bHovered = false;
+		SButton::OnMouseLeave(MouseEvent);
+	}
+
+	/** Retuns TRUE if chip should appear as hovered. */
+	bool ShouldAppearHovered() const
+	{
+		return bHovered || HasKeyboardFocus();
+	}
+
+	void Construct(const FArguments& InArgs)
+	{
+		SButton::Construct(
+			SButton::FArguments()
+			.ButtonStyle(InArgs._ButtonStyle)
+			.ContentPadding(InArgs._ContentPadding)
+			.ToolTipText(InArgs._ToolTipText)
+			.OnClicked(InArgs._OnClicked)
+			[
+				InArgs._Content.Widget
+			]
+		);
+
+		bHovered = false;
+
+		SetHover(TAttribute<bool>::CreateSP(this, &SUAFBrowserChipButton::ShouldAppearHovered));
+	}
+
+	/** Keep an internal IsHovered flag*/
+	bool bHovered = false;
+};
+
+// ---------------------------------------------------------------------------
+// SUAFBrowserFilterSuggestionStrip — static members
+// ---------------------------------------------------------------------------
+
+const FName SUAFBrowserFilterSuggestionStrip::CategoryPlugins(TEXT("Plugins"));
+const FName SUAFBrowserFilterSuggestionStrip::CategoryTypes(TEXT("Types"));
+const FName SUAFBrowserFilterSuggestionStrip::CategoryMisc(TEXT("Misc"));
+
+const TArray<FName> SUAFBrowserFilterSuggestionStrip::DefaultCategoryOrder = 
+{
+	SUAFBrowserFilterSuggestionStrip::CategoryPlugins,
+	SUAFBrowserFilterSuggestionStrip::CategoryTypes,
+	SUAFBrowserFilterSuggestionStrip::CategoryMisc,
+};
+
+// ---------------------------------------------------------------------------
+// SUAFBrowserFilterSuggestionStrip — construction
+// ---------------------------------------------------------------------------
+
+void SUAFBrowserFilterSuggestionStrip::Construct(const FArguments& InArgs)
+{
+	OnFilterSuggestionSelectedDelegate = InArgs._OnFilterSuggestionSelected;
+
+	ChildSlot
+	[
+		SNew(SVerticalBox)
+
+		// Hint text shown above the rows whenever the strip is visible
+		+ SVerticalBox::Slot()
+		.AutoHeight()
+		.Padding(FMargin(4.f, 2.f, 4.f, 0.f))
+		[
+			SAssignNew(HintTextWidget, STextBlock)
+			.Text(LOCTEXT("KeyboardHint", "Tab to select filter"))
+			.TextStyle(FAppStyle::Get(), "SmallText")
+			.ColorAndOpacity(FSlateColor::UseSubduedForeground())
+		]
+
+		// Row container — rows are added/removed dynamically in SetSuggestions()
+		+ SVerticalBox::Slot()
+		.AutoHeight()
+		[
+			SNew(SBorder)
+			.BorderImage(FAppStyle::Get().GetBrush("Brushes.Header"))
+			.Padding(FMargin(4.f, 2.f))
+			[
+				SAssignNew(RowsContainer, SVerticalBox)
+			]
+		]
+	];
+
+	// Start collapsed — shown only when suggestions exist
+	SetVisibility(EVisibility::Collapsed);
+}
+
+// ---------------------------------------------------------------------------
+// SUAFBrowserFilterSuggestionStrip — public interface
+// ---------------------------------------------------------------------------
+
+void SUAFBrowserFilterSuggestionStrip::SetSuggestions(const TArray<TSharedRef<FFilterSuggestion>>& Suggestions)
+{
+	RowsContainer->ClearChildren();
+	ActiveRows.Reset();
+
+	if (HintTextWidget.IsValid())
+	{
+		HintTextWidget->SetText(LOCTEXT("KeyboardHint", "Tab to select filter"));
+	}
+
+	if (Suggestions.IsEmpty())
+	{
+		SetVisibility(EVisibility::Collapsed);
+		return;
+	}
+
+	// Build the iteration order: known categories first (in DefaultCategoryOrder), then
+	// any unknown category keys encountered in Suggestions (in the order they first appear).
+	TArray<FName> OrderedKeys = DefaultCategoryOrder;
+	for (const TSharedRef<FFilterSuggestion>& Suggestion : Suggestions)
+	{
+		OrderedKeys.AddUnique(Suggestion->Category);
+	}
+
+	for (const FName& CategoryKey : OrderedKeys)
+	{
+		// Collect all suggestions for this category
+		TArray<TSharedRef<FFilterSuggestion>> CategorySuggestions;
+		for (const TSharedRef<FFilterSuggestion>& Suggestion : Suggestions)
+		{
+			if (Suggestion->Category == CategoryKey)
+			{
+				CategorySuggestions.Add(Suggestion);
+			}
+		}
+		if (CategorySuggestions.IsEmpty())
+		{
+			continue;
+		}
+
+		// Build chip widgets for this category
+		FFilterSuggestionRow& NewRow = ActiveRows.AddDefaulted_GetRef();
+		NewRow.CategoryKey = CategoryKey;
+		for (const TSharedRef<FFilterSuggestion>& Suggestion : CategorySuggestions)
+		{
+			NewRow.ChipWidgets.Add(BuildChipWidget(Suggestion));
+		}
+
+		// Build the row widget: label + horizontally-scrolling chip area
+		TSharedPtr<SHorizontalBox> ChipsBox = SNew(SHorizontalBox);
+
+		RowsContainer->AddSlot()
+		.AutoHeight()
+		.Padding(FMargin(0.f, 1.f))
+		[
+			SNew(SHorizontalBox)
+
+			// Row label (FName string doubles as display text)
+			+ SHorizontalBox::Slot()
+			.AutoWidth()
+			.VAlign(VAlign_Center)
+			.Padding(FMargin(0.f, 0.f, 6.f, 0.f))
+			[
+				SNew(STextBlock)
+				.Text(FText::FromName(CategoryKey))
+				.TextStyle(FAppStyle::Get(), "SmallText")
+				.ColorAndOpacity(FSlateColor::UseSubduedForeground())
+			]
+				
+			// Horizontally-scrolling chip area
+			+ SHorizontalBox::Slot()
+			.FillWidth(1.f)
+			[
+				SNew(SScrollBox)
+				.Orientation(Orient_Horizontal)
+				.ScrollBarVisibility(EVisibility::Collapsed)
+				+ SScrollBox::Slot()
+				[
+					SAssignNew(ChipsBox, SHorizontalBox)
+				]
+			]
+		];
+
+		// Populate the chip area
+		for (const TSharedRef<SWidget>& ChipWidget : NewRow.ChipWidgets)
+		{
+			ChipsBox->AddSlot()
+			.AutoWidth()
+			.Padding(FMargin(0.f, 0.f, 4.f, 0.f))
+			[
+				ChipWidget
+			];
+		}
+	}
+
+	SetVisibility(EVisibility::Visible);
+}
+
+bool SUAFBrowserFilterSuggestionStrip::HasAnySuggestions() const
+{
+	return !ActiveRows.IsEmpty();
+}
+
+void SUAFBrowserFilterSuggestionStrip::FocusFirstChip()
+{
+	if (!ActiveRows.IsEmpty() && !ActiveRows[0].ChipWidgets.IsEmpty())
+	{
+		if (HintTextWidget.IsValid())
+		{
+			HintTextWidget->SetText(LOCTEXT("KeyboardNavHint", "[ Enter ] Apply  \u00b7  [\u2190][\u2192] Navigate  \u00b7  [ \u2191 ][ \u2193 ] Row"));
+		}
+		FSlateApplication::Get().SetKeyboardFocus(ActiveRows[0].ChipWidgets[0]);
+	}
+}
+
+// ---------------------------------------------------------------------------
+// SUAFBrowserFilterSuggestionStrip — keyboard navigation
+// ---------------------------------------------------------------------------
+
+FReply SUAFBrowserFilterSuggestionStrip::OnPreviewKeyDown(const FGeometry& MyGeometry, const FKeyEvent& InKeyEvent)
+{
+	TSharedPtr<SWidget> FocusedWidget = FSlateApplication::Get().GetUserFocusedWidget(0);
+
+	// Find which row and chip index currently has keyboard focus
+	int32 FocusedRowIndex = INDEX_NONE;
+	int32 FocusedChipIndex = INDEX_NONE;
+
+	for (int32 RowIndex = 0; RowIndex < ActiveRows.Num() && FocusedRowIndex == INDEX_NONE; ++RowIndex)
+	{
+		for (int32 ChipIndex = 0; ChipIndex < ActiveRows[RowIndex].ChipWidgets.Num(); ++ChipIndex)
+		{
+			if (ActiveRows[RowIndex].ChipWidgets[ChipIndex] == FocusedWidget)
+			{
+				FocusedRowIndex = RowIndex;
+				FocusedChipIndex = ChipIndex;
+				break;
+			}
+		}
+	}
+
+	if (FocusedRowIndex == INDEX_NONE)
+	{
+		return FReply::Unhandled();
+	}
+
+	const FKey PressedKey = InKeyEvent.GetKey();
+	const TArray<TSharedRef<SWidget>>& CurrentRowChips = ActiveRows[FocusedRowIndex].ChipWidgets;
+
+	if (PressedKey == EKeys::Left)
+	{
+		if (FocusedChipIndex > 0)
+		{
+			FSlateApplication::Get().SetKeyboardFocus(CurrentRowChips[FocusedChipIndex - 1]);
+		}
+		return FReply::Handled();
+	}
+	else if (PressedKey == EKeys::Right)
+	{
+		if (FocusedChipIndex + 1 < CurrentRowChips.Num())
+		{
+			FSlateApplication::Get().SetKeyboardFocus(CurrentRowChips[FocusedChipIndex + 1]);
+		}
+		return FReply::Handled();
+	}
+	else if (PressedKey == EKeys::Up || PressedKey == EKeys::Down)
+	{
+		if (ActiveRows.Num() <= 1)
+		{
+			// No adjacent row to navigate to — consume to prevent focus escaping the strip
+			return FReply::Handled();
+		}
+
+		const int32 DirectionOffset = (PressedKey == EKeys::Down) ? 1 : -1;
+		const int32 TargetRowIndex = (FocusedRowIndex + DirectionOffset + ActiveRows.Num()) % ActiveRows.Num();
+		const int32 TargetChipIndex = FMath::Min(FocusedChipIndex, ActiveRows[TargetRowIndex].ChipWidgets.Num() - 1);
+		FSlateApplication::Get().SetKeyboardFocus(ActiveRows[TargetRowIndex].ChipWidgets[TargetChipIndex]);
+		return FReply::Handled();
+	}
+
+	return FReply::Unhandled();
+}
+
+// ---------------------------------------------------------------------------
+// SUAFBrowserFilterSuggestionStrip — private helpers
+// ---------------------------------------------------------------------------
+
+TSharedRef<SWidget> SUAFBrowserFilterSuggestionStrip::BuildChipWidget(TSharedRef<FFilterSuggestion> Suggestion)
+{
+	const FName IconName = Suggestion->IconName.IsNone() ? FName("Icons.Help") : Suggestion->IconName;
+
+	return SNew(SUAFBrowserChipButton)
+		.ButtonStyle(&FAppStyle::Get().GetWidgetStyle<FButtonStyle>("SimpleButton"))
+		.ContentPadding(FMargin(6.f, 2.f))
+		.ToolTipText(Suggestion->DisplayName)
+		.OnClicked(this, &SUAFBrowserFilterSuggestionStrip::OnChipClicked, Suggestion)
+		[
+			SNew(SHorizontalBox)
+
+			+ SHorizontalBox::Slot()
+			.AutoWidth()
+			.VAlign(VAlign_Center)
+			.Padding(FMargin(0.f, 0.f, 4.f, 0.f))
+			[
+				SNew(SImage)
+				.Image(FAppStyle::GetBrush(IconName))
+				.DesiredSizeOverride(FVector2D(12.f, 12.f))
+				.ColorAndOpacity(FSlateColor::UseForeground())
+			]
+
+			+ SHorizontalBox::Slot()
+			.AutoWidth()
+			.VAlign(VAlign_Center)
+			[
+				SNew(STextBlock)
+				.Text(Suggestion->DisplayName)
+				.TextStyle(FAppStyle::Get(), "SmallText")
+			]
+		];
+}
+
+FReply SUAFBrowserFilterSuggestionStrip::OnChipClicked(TSharedRef<FFilterSuggestion> Suggestion)
+{
+	Suggestion->OnActivate.ExecuteIfBound();
+	OnFilterSuggestionSelectedDelegate.ExecuteIfBound();
+	return FReply::Handled();
+}
+
+#undef LOCTEXT_NAMESPACE
diff --git a/Engine/Plugins/Experimental/UAF/UAF/Source/UAFEditor/Private/Common/SUAFBrowserFilterSuggestionStrip.h b/Engine/Plugins/Experimental/UAF/UAF/Source/UAFEditor/Private/Common/SUAFBrowserFilterSuggestionStrip.h
new file mode 100644
--- /dev/null
+++ b/Engine/Plugins/Experimental/UAF/UAF/Source/UAFEditor/Private/Common/SUAFBrowserFilterSuggestionStrip.h
@@ -0,0 +1,100 @@
+// Copyright Epic Games, Inc. All Rights Reserved.
+
+#pragma once
+
+#include "Delegates/Delegate.h"
+#include "Internationalization/Text.h"
+#include "Math/Color.h"
+#include "UObject/NameTypes.h"
+#include "Widgets/SCompoundWidget.h"
+
+class SVerticalBox;
+class STextBlock;
+
+/**
+ * Display-and-activation descriptor for a single suggestion chip in the UAF Browser filter suggestion strip.
+ * Decouples display (name, icon, color) from activation logic so the strip doesn't need to know the filter type.
+ */
+struct FFilterSuggestion
+{
+	FText DisplayName;
+	FName IconName;
+	FLinearColor Color = FLinearColor::White;
+
+	/**
+	 * Row grouping key. Use one of the built-in SUAFBrowserFilterSuggestionStrip::Category* constants
+	 * or any FName to create a custom row. The FName string is used as the row label.
+	 */
+	FName Category = FName(TEXT("Misc"));
+
+	/** Executed when the chip is clicked (or activated via keyboard) to apply the associated filter. */
+	TDelegate<void()> OnActivate;
+};
+
+/**
+ * A horizontally-scrolling, categorized strip of filter suggestion chips.
+ * Appears at the bottom of the UAF Browser when the user types text that matches available filters.
+ *
+ * Suggestions are grouped into named rows determined by FFilterSuggestion::Category.
+ * Each row is only visible when it has at least one matching suggestion.
+ *
+ * Keyboard navigation:
+ *   - Tab (from search box): moves focus to the first chip (intercepted by parent SUAFBrowser)
+ *   - Left / Right: move focus within the current row
+ *   - Up / Down: move focus to the corresponding chip in the adjacent row (wraps, no-op if single row)
+ *   - Enter / Space: activate the focused chip (handled natively by SButton)
+ */
+class SUAFBrowserFilterSuggestionStrip : public SCompoundWidget
+{
+public:
+	/** Built-in category key constants. The FName string doubles as the row label. */
+	static const FName CategoryPlugins;
+	static const FName CategoryTypes;
+	static const FName CategoryMisc;
+
+	SLATE_BEGIN_ARGS(SUAFBrowserFilterSuggestionStrip) {}
+		/** Called after any chip's OnActivate fires. Use to restore search box focus and clear search text. */
+		SLATE_EVENT(TDelegate<void()>, OnFilterSuggestionSelected)
+	SLATE_END_ARGS()
+
+	void Construct(const FArguments& InArgs);
+
+	/** Updates the displayed suggestion chips, grouped into rows by Category FName. Collapses the strip when Suggestions is empty. */
+	void SetSuggestions(const TArray<TSharedRef<FFilterSuggestion>>& Suggestions);
+
+	/** Returns true if any category row has at least one chip. */
+	bool HasAnySuggestions() const;
+
+	/** Moves keyboard focus to the first chip in the first non-empty row. */
+	void FocusFirstChip();
+
+	virtual FReply OnPreviewKeyDown(const FGeometry& MyGeometry, const FKeyEvent& InKeyEvent) override;
+
+private:
+	/** Internal per-row data, rebuilt on every SetSuggestions() call. */
+	struct FFilterSuggestionRow
+	{
+		FName CategoryKey;
+		TArray<TSharedRef<SWidget>> ChipWidgets;
+	};
+
+	TSharedRef<SWidget> BuildChipWidget(TSharedRef<FFilterSuggestion> Suggestion);
+	FReply OnChipClicked(TSharedRef<FFilterSuggestion> Suggestion);
+
+	/**
+	 * Defines the preferred display order for known category keys.
+	 * Any unknown keys encountered in SetSuggestions() are appended in the order they first appear.
+	 */
+	static const TArray<FName> DefaultCategoryOrder;
+
+	/** Live rows, rebuilt on every SetSuggestions() call. Only contains categories with at least one suggestion. */
+	TArray<FFilterSuggestionRow> ActiveRows;
+
+	/** Vertical container into which row widgets are added dynamically in SetSuggestions(). */
+	TSharedPtr<SVerticalBox> RowsContainer;
+
+	/** Hint text above the rows; toggled between invite and navigation hints depending on focus state. */
+	TSharedPtr<STextBlock> HintTextWidget;
+
+	TDelegate<void()> OnFilterSuggestionSelectedDelegate;
+};
diff --git a/Engine/Plugins/Experimental/UAF/UAF/Source/UAFEditor/Private/Common/UAFBrowserFrontendFilters.h b/Engine/Plugins/Experimental/UAF/UAF/Source/UAFEditor/Private/Common/UAFBrowserFrontendFilters.h
new file mode 100644
--- /dev/null
+++ b/Engine/Plugins/Experimental/UAF/UAF/Source/UAFEditor/Private/Common/UAFBrowserFrontendFilters.h
@@ -0,0 +1,51 @@
+// Copyright Epic Games, Inc. All Rights Reserved.
+
+#pragma once
+
+#include "FrontendFilterBase.h"
+
+/**
+ * A content browser frontend filter that restricts the asset view to content from a specific Unreal Engine plugin.
+ * These filters appear under the "Plugins" category in the UAF Browser's "Add Filter" dropdown.
+ */
+class FFrontendFilter_UAFPlugin : public FFrontendFilter
+{
+public:
+	FFrontendFilter_UAFPlugin(TSharedPtr<FFrontendFilterCategory> InCategory,
+	                          const FString& InPluginName,
+	                          const FString& InFriendlyName,
+	                          const FString& InPluginContentPath)
+		: FFrontendFilter(MoveTemp(InCategory))
+		, PluginName(InPluginName)
+		, FriendlyName(InFriendlyName)
+		, PluginContentPath(InPluginContentPath)
+	{}
+
+	// FFilterBase interface
+	virtual FString GetName() const override { return PluginName; }
+	virtual FText GetDisplayName() const override { return FText::FromString(FriendlyName); }
+	virtual FText GetToolTipText() const override
+	{
+		return FText::Format(
+			NSLOCTEXT("UAFBrowser", "PluginFilterTooltip", "Show only assets from the '{0}' plugin (path: {1})"),
+			FText::FromString(FriendlyName),
+			FText::FromString(PluginContentPath));
+	}
+	virtual FLinearColor GetColor() const override { return FLinearColor(0.4f, 0.7f, 1.0f, 1.0f); }
+	virtual FName GetIconName() const override { return FName("Icons.Package"); }
+
+	// FFrontendFilter interface
+	virtual bool PassesFilter(FAssetFilterType const& InItem) const override
+	{
+		return InItem.GetInternalPath().ToString().StartsWith(PluginContentPath);
+	}
+
+	/** Returns the plugin name (without the friendly display name). Used for matching against search text. */
+	const FString& GetPluginName() const { return PluginName; }
+
+private:
+	FString PluginName;
+	FString FriendlyName;
+	/** Package path prefix for this plugin's content, e.g. "/Cosmo/". */
+	FString PluginContentPath;
+};
diff --git a/Engine/Plugins/Experimental/UAF/UAF/Source/UAFEditor/Private/Settings/UAFEditorProjectSettings.h b/Engine/Plugins/Experimental/UAF/UAF/Source/UAFEditor/Private/Settings/UAFEditorProjectSettings.h
--- a/Engine/Plugins/Experimental/UAF/UAF/Source/UAFEditor/Private/Settings/UAFEditorProjectSettings.h
+++ b/Engine/Plugins/Experimental/UAF/UAF/Source/UAFEditor/Private/Settings/UAFEditorProjectSettings.h
@@ -19,4 +19,12 @@ public:
 	/** The configuration used the UAF browser for setting up categories */
 	UPROPERTY(config, EditAnywhere, DisplayName = "UAF Browser Configuration", Category = UAF, meta = (AllowedClasses = "/Script/UserAssetTagsEditor.TaggedAssetBrowserConfiguration"))
 	FSoftObjectPath UAFBrowserCategoryConfiguration;
+
+	/** Additional plugin names to include in the UAF Browser filter bar beyond the auto-discovered project plugins. */
+	UPROPERTY(config, EditAnywhere, DisplayName = "Additional Browser Plugin Filters", Category = "UAF Browser Plugin Filters")
+	TArray<FName> AdditionalBrowserPluginFilters;
+
+	/** Plugin names to exclude from the auto-discovered project plugin list in the UAF Browser filter bar. */
+	UPROPERTY(config, EditAnywhere, DisplayName = "Excluded Browser Plugin Filters", Category = "UAF Browser Plugin Filters")
+	TArray<FName> ExcludedBrowserPluginFilters;
 };
\ No newline at end of file
diff --git a/Engine/Source/Editor/ContentBrowser/Internal/SAssetPicker.h b/Engine/Source/Editor/ContentBrowser/Internal/SAssetPicker.h
new file mode 100644
--- /dev/null
+++ b/Engine/Source/Editor/ContentBrowser/Internal/SAssetPicker.h
@@ -0,0 +1,238 @@
+// Copyright Epic Games, Inc. All Rights Reserved.
+
+#pragma once
+
+#include "AssetRegistry/ARFilter.h"
+#include "AssetTypeCategories.h"
+#include "Containers/Array.h"
+#include "Containers/ArrayView.h"
+#include "Containers/UnrealString.h"
+#include "ContentBrowserDelegates.h"
+#include "AssetViewContentSources.h"
+#include "HAL/Platform.h"
+#include "IAssetTypeActions.h"
+#include "IContentBrowserSingleton.h"
+#include "Input/Reply.h"
+#include "Internationalization/Text.h"
+#include "Templates/SharedPointer.h"
+#include "Types/SlateEnums.h"
+#include "Widgets/DeclarativeSyntaxSupport.h"
+#include "Widgets/SCompoundWidget.h"
+
+class FFilter_HideOtherDevelopers;
+class FUICommandList;
+class SAssetSearchBox;
+class SAssetView;
+class SComboButton;
+class SFilterList;
+class SWidget;
+struct FAssetData;
+struct FContentBrowserItem;
+struct FGeometry;
+struct FKeyEvent;
+
+enum class ECheckBoxState : uint8;
+
+/**
+ * Small content browser designed to allow for asset picking
+ */
+class SAssetPicker : public SCompoundWidget
+{
+public:
+	SLATE_BEGIN_ARGS( SAssetPicker ){}
+
+		/** A struct containing details about how the asset picker should behave */
+		SLATE_ARGUMENT(FAssetPickerConfig, AssetPickerConfig)
+
+	SLATE_END_ARGS()
+
+	virtual ~SAssetPicker();
+
+	/** Constructs this widget with InArgs */
+	void Construct( const FArguments& InArgs );
+
+	// SWidget implementation
+	virtual FReply OnKeyDown( const FGeometry& MyGeometry, const FKeyEvent& InKeyEvent ) override;
+	// End of SWidget implementation
+
+	/** Return the associated AssetView */
+	const TSharedPtr<SAssetView>& GetAssetView() const { return AssetViewPtr; }
+
+	/** Function that's called when someone wants to execute a Rename Command*/
+	void ExecuteRenameCommand();
+
+	TSharedPtr<SWidget> GetSearchBox() const;
+
+	/** Clears the search box text programmatically. Triggers OnSearchTextChanged -> OnSearchTextChangedDelegate. */
+	CONTENTBROWSER_API void ClearSearchText();
+
+	/** Returns the filter list widget, or nullptr if none was created. */
+	CONTENTBROWSER_API TSharedPtr<SFilterList> GetFilterList() const;
+
+	/** Syncs the AssetPicker to select the provided items. */
+	void SyncToItems(TArrayView<const FContentBrowserItem> ItemsToSync, const bool bAllowImplicitSync = false, const bool bDisableFiltersThatHideAssets = true);
+
+private:
+	/** Focuses the search box post-construct */
+	EActiveTimerReturnType SetFocusPostConstruct( double InCurrentTime, float InDeltaTime );
+
+	/** Special case handling for SAssetSearchBox key commands */
+	FReply HandleKeyDownFromSearchBox(const FGeometry& MyGeometry, const FKeyEvent& InKeyEvent);
+
+	void FolderEntered(const FString& FolderPath);
+
+	/** Called when the editable text needs to be set or cleared */
+	void SetSearchBoxText(const FText& InSearchText);
+
+	/** Called by the editable text control when the search text is changed by the user */
+	void OnSearchBoxChanged(const FText& InSearchText);
+
+	/** Called by the editable text control when the user commits a text change */
+	void OnSearchBoxCommitted(const FText& InSearchText, ETextCommit::Type CommitInfo);
+
+	/** Called from external code to set the filter after the widget was created */
+	void SetNewBackendFilter(const FARFilter& NewFilter);
+
+	/** Called when the user changes filters */
+	void OnFilterChanged();
+
+	/** Handler for when the "None" button is clicked */
+	FReply OnNoneButtonClicked();
+
+	/** Handler for when a new item is requested in the asset view */
+	void HandleNewItemRequested(const FContentBrowserItem& NewItem);
+
+	/** Handle forwarding picking events. We wrap OnAssetSelected here to prevent 'Direct' selections being identified as user actions */
+	void HandleItemSelectionChanged(const FContentBrowserItem& InSelectedItem, ESelectInfo::Type InSelectInfo);
+
+	/** Handler for when the user double clicks, presses enter, or presses space on an asset */
+	void HandleItemsActivated(TArrayView<const FContentBrowserItem> ActivatedItems, EAssetTypeActivationMethod::Type ActivationMethod);
+
+	/** Handler for when the user drags an asset */
+	void HandleItemsDragged(TArrayView<const FContentBrowserItem> DraggedItems);
+
+	/** Selects the paths containing the specified assets. */
+	void SyncToAssets(const TArray<FAssetData>& AssetDataList);
+
+	/** @return The currently selected asset */
+	TArray< FAssetData > GetCurrentSelection();
+
+	/** Forces a refresh */
+	void RefreshAssetView(bool bRefreshSources);
+
+	/** @return The text to highlight on the assets  */
+	FText GetHighlightedText() const;
+
+	/** @return The tooltip for the other developers filter button depending on checked state.  */
+	FText GetShowOtherDevelopersToolTip() const;
+
+	/** Toggles the filter for showing other developers assets */
+	void HandleShowOtherDevelopersCheckStateChanged( const ECheckBoxState InCheckboxState );
+
+	/** Gets if showing other developers assets */
+	ECheckBoxState GetShowOtherDevelopersCheckState() const;
+
+	/** Returns whether the ShowOtherDevelopersFilter is active i.e. if we should show content belonging to other developers */
+	bool IsShowingOtherDevelopersContent() const;
+
+	/** Called upon the Rename UICommand being executed, sends rename request to the asset view */
+	void OnRenameRequested() const;
+
+	/** Returns true if the user is able to execute a rename request */
+	bool CanExecuteRenameRequested();
+
+	/** Returns true if the user is able to execute filter by selection request */
+	bool CanFilterBySelection();
+
+	/** Executes command and applies filters based on active selection */
+	void OnFilterBySelection();
+
+	/** Bind our UI commands */
+	void BindCommands();
+
+	/** Loads settings for this asset picker if SaveSettingsName was set */
+	void LoadSettings();
+
+	/** Saves settings for this asset picker if SaveSettingsName was set */
+	void SaveSettings() const;
+
+	void HandleSearchSettingsChanged();
+
+	/** Handler for the context menu for asset or folder items */
+	TSharedPtr<SWidget> GetItemContextMenu(TArrayView<const FContentBrowserItem> SelectedItems);
+
+	/** Converts an asset type category into an asset category path, if possible. This is done because the asset picker config still uses asset type categories. */
+	TOptional<FAssetCategoryPath> ConvertAssetTypeCategoryToAssetCategoryPath(EAssetTypeCategories::Type DefaultFilterMenuExpansion);
+
+private:
+
+	/** The list of FrontendFilters currently applied to the asset view */
+	TSharedPtr<FAssetFilterCollectionType> FrontendFilters;
+
+	/** The asset view widget */
+	TSharedPtr<SAssetView> AssetViewPtr;
+
+	/** The search box */
+	TSharedPtr<SAssetSearchBox> SearchBoxPtr;
+
+	/** The filter list */
+	TSharedPtr<SFilterList> FilterListPtr;
+
+	/** The filter combo button */
+	TSharedPtr<SComboButton> FilterComboButtonPtr;
+
+	/** Called to when an asset is selected or the none button is pressed */
+	FOnAssetSelected OnAssetSelected;
+
+	/** Called to when an asset is double clicked */
+	FOnAssetDoubleClicked OnAssetDoubleClicked;
+
+	/** Called when enter is pressed while an asset is selected */
+	FOnAssetEnterPressed OnAssetEnterPressed;
+
+	/** Called when any number of assets are activated */
+	FOnAssetsActivated OnAssetsActivated;
+
+	/** Called any number of assets are dragged */
+	FOnAssetsDragged OnAssetsDragged;
+
+	/** Called when a folder is entered in the asset view */
+	FOnPathSelected OnFolderEnteredDelegate;
+
+	/** Fired when the search box text changes, forwarded from FAssetPickerConfig::OnSearchTextChanged. */
+	FOnSearchTextChanged OnSearchTextChangedDelegate;
+
+	/** Delegate to invoke when a context menu for an asset is opening. */
+	FOnGetAssetContextMenu OnGetAssetContextMenu;
+
+	/** Delegate to invoke when a context menu for a folder is opening. */
+	FOnGetFolderContextMenu OnGetFolderContextMenu;
+
+	/** True if the search box will take keyboard focus next frame */
+	bool bPendingFocusNextFrame;
+
+	/** Filters needed for filtering the assets */
+	TSharedPtr< FAssetFilterCollectionType > FilterCollection;
+	TSharedPtr<class FAssetTextFilter> TextFilter;
+	TSharedPtr<FFilter_HideOtherDevelopers> HideOtherDevelopersFilter;
+
+	EAssetTypeCategories::Type DefaultFilterMenuExpansion;
+
+	/** The sources data currently used by the picker */
+	FAssetViewContentSources CurrentContentSources;
+
+	/** Current filter we are using, needed reset asset view after we have custom filtered */
+	FARFilter CurrentBackendFilter;
+
+	/** UICommand list, holds list of actions for processing */
+	TSharedPtr< FUICommandList > Commands;
+
+	/** If set, view settings will be saved and loaded for the asset view using this name in ini files */
+	FString SaveSettingsName;
+
+	/** If set, uses this name for the filter bar identifier instead of the SaveSettingsName */
+	TOptional<FString> FilterBarSaveSettingsNameOverride;
+
+	/** Wether to allow renaming of items */
+	bool bAllowRename;
+};
diff --git a/Engine/Source/Editor/ContentBrowser/Internal/SFilterList.h b/Engine/Source/Editor/ContentBrowser/Internal/SFilterList.h
new file mode 100644
--- /dev/null
+++ b/Engine/Source/Editor/ContentBrowser/Internal/SFilterList.h
@@ -0,0 +1,313 @@
+// Copyright Epic Games, Inc. All Rights Reserved.
+
+#pragma once
+
+#include "AssetTypeCategories.h"
+#include "Containers/Array.h"
+#include "Containers/ArrayView.h"
+#include "Containers/Map.h"
+#include "Containers/UnrealString.h"
+#include "Delegates/Delegate.h"
+#include "Filters/CustomTextFilters.h"
+#include "Filters/FilterBase.h"
+#include "Filters/SAssetFilterBar.h"
+#include "Filters/SBasicFilterBar.h"
+#include "FrontendFilters.h"
+#include "HAL/Platform.h"
+#include "IContentBrowserSingleton.h"
+#include "Input/Reply.h"
+#include "InputCoreTypes.h"
+#include "Internationalization/Text.h"
+#include "Math/Color.h"
+#include "Templates/SharedPointer.h"
+#include "UObject/NameTypes.h"
+#include "Widgets/DeclarativeSyntaxSupport.h"
+
+class FFrontendFilter;
+class SWidget;
+class UClass;
+class UToolMenu;
+struct FContentBrowserItem;
+struct FCustomTextFilterState;
+struct FFilterBarSettings;
+struct FGeometry;
+struct FPointerEvent;
+
+enum class ECheckBoxState : uint8;
+
+/**
+ * A list of filters currently applied to an asset view.
+ */
+class SFilterList : public SAssetFilterBar<FAssetFilterType>
+{
+	using Super = SAssetFilterBar<FAssetFilterType>; 
+public:
+	/**
+	 * An event delegate that is executed when a custom text filter has been created/modified/deleted in any FilterList
+	 * that is using shared settings.
+	 */
+	DECLARE_EVENT_OneParam(SFilterList, FCustomTextFilterEvent, TSharedPtr<SWidget> /* BroadcastingFilterList */);
+	
+	using FOnFilterChanged = typename SAssetFilterBar<FAssetFilterType>::FOnFilterChanged;
+
+	SLATE_BEGIN_ARGS( SFilterList )
+	: _UseSharedSettings(false)
+	, _FilterBarLayout(EFilterBarLayout::Horizontal)
+	, _CanChangeOrientation(false)
+	, _DefaultMenuExpansionCategory(EAssetCategoryPaths::Basic)
+	, _bUseSectionsForCustomCategories(false)
+	{
+		
+	}
+		/** Delegate for when filters have changed */
+		SLATE_EVENT( FOnFilterChanged, OnFilterChanged )
+
+		/** Delegate that lets the user modify the menu after the fact */
+		SLATE_EVENT(FOnExtendAddFilterMenu, OnExtendAddFilterMenu)
+	
+		/** The filter collection used to further filter down assets returned from the backend */
+		SLATE_ARGUMENT( TSharedPtr<FAssetFilterCollectionType>, FrontendFilters)
+
+		/** An array of classes to filter the menu by */
+		SLATE_ARGUMENT( TArray<UClass*>, InitialClassFilters)
+
+		/** Custom front end filters to be displayed */
+		SLATE_ARGUMENT( TArray< TSharedRef<FFrontendFilter> >, ExtraFrontendFilters )
+	
+		/** A delegate to create a Text Filter for FilterType items. If provided, will allow creation of custom text filters
+		 *  from the filter dropdown menu.
+		 */
+		SLATE_ARGUMENT(FCreateTextFilter, CreateTextFilter)
+	
+		/** A unique identifier for this filter bar needed to enable saving settings in a config file */
+		SLATE_ARGUMENT(FName, FilterBarIdentifier)
+
+		/** If true, CustomTextFilters are saved/loaded from a config shared between multiple SFilterLists
+		 *	Currently used to sync all content browser custom text filters.
+		 */
+		SLATE_ARGUMENT(bool, UseSharedSettings)
+
+		/** Called when the filter bar layout is changing, before the filters are added to the layout */
+		SLATE_EVENT(FOnFilterBarLayoutChanging, OnFilterBarLayoutChanging)
+
+		/** The layout that determines how the filters are laid out */
+		SLATE_ARGUMENT(EFilterBarLayout, FilterBarLayout)
+			
+		/** If true, allow dynamically changing the orientation and saving in the config */
+		SLATE_ARGUMENT(bool, CanChangeOrientation)
+
+		/** Expands the specified asset category, if specified. If not, it will expand Basic/Common instead. */
+		SLATE_ARGUMENT(TOptional<FAssetCategoryPath>, DefaultMenuExpansionCategory)
+
+		/** If true, adds custom categories as sections (expanded) vs. as sub-menus */
+		SLATE_ARGUMENT(bool, bUseSectionsForCustomCategories)
+
+		/** The command list provided to the filter context menu */
+		SLATE_ARGUMENT(TSharedPtr<FUICommandList>, CommandList)
+
+	SLATE_END_ARGS()
+
+	/** Constructs this widget with InArgs */
+	void Construct( const FArguments& InArgs );
+
+	/** Set the check box state of the specified frontend filter (in the filter drop down) and pin/unpin a filter widget on/from the filter bar. When a filter is pinned (was not already pinned), it is activated and deactivated when unpinned. */
+	void SetFrontendFilterCheckState(const TSharedPtr<FFrontendFilter>& InFrontendFilter, ECheckBoxState CheckState);
+
+	/** Returns the check box state of the specified frontend filter (in the filter drop down). This tells whether the filter is pinned or not on the filter bar, but not if filter is active or not. @see IsFrontendFilterActive(). */
+	ECheckBoxState GetFrontendFilterCheckState(const TSharedPtr<FFrontendFilter>& InFrontendFilter) const;
+
+	/** Returns true if the specified frontend filter is both checked (pinned on the filter bar) and active (contributing to filter the result). */
+	bool IsFrontendFilterActive(const TSharedPtr<FFrontendFilter>& InFrontendFilter) const;
+
+	/** Retrieve a specific frontend filter */
+	TSharedPtr<FFrontendFilter> GetFrontendFilter(const FString& InName) const;
+	
+	/** 
+	 * Replaces super version of GetCombinedBackendFilter to allow returning FPathPermissionLists as well.
+	 * Permission lists will be combined into a content browser data filter after a delay. Filters may modify them 
+	 * via the shared ref but should not do so off the game thread.
+	 */
+	FARFilter GetCombinedBackendFilter(TArray<TSharedRef<const FPathPermissionList>>& OutPermissionLists) const;
+	FARFilter GetCombinedBackendFilter() const = delete;
+	
+	/** Handler for when the floating add filter button was clicked */
+	TSharedRef<SWidget> ExternalMakeAddFilterMenu();
+
+	/** Disables any active filters that would hide the supplied items */
+	void DisableFiltersThatHideItems(TArrayView<const FContentBrowserItem> ItemList);
+
+	/** Returns the class filters specified at construction using argument 'InitialClassFilters'. */
+	const TArray<UClass*>& GetInitialClassFilters();
+
+	/** Open the dialog to create a custom filter from the given text */
+	void CreateCustomFilterDialog(const FText& InText);
+	
+	/** Updates bIncludeClassName, bIncludeAssetPath and bIncludeCollectionNames for all custom text filters */
+	void UpdateCustomTextFilterIncludes(const bool InIncludeClassName, const bool InIncludeAssetPath, const bool InIncludeCollectionNames);
+
+	/** Add a custom widget to the filter bar alongside the filters */
+	void AddWidgetToCurrentLayout(TSharedRef<SWidget> InWidget);
+
+	using Super::SaveSettings;
+	virtual void LoadSettings() override;
+
+	/* Copy the settings from a specific instance name */
+	void LoadSettings(const FName& InInstanceName);
+
+	/** Helper functions for backwards compatibility with filters that save state until they are ported to EditorConfig */
+	void SaveSettings(const FString& IniFilename, const FString& IniSection, const FString& SettingsString);
+	void LoadSettings(const FName& InInstanceName, const FString& IniFilename, const FString& IniSection, const FString& SettingsString);
+
+	/** Handler for when a custom text filter is created */
+	virtual void OnCreateCustomTextFilter(const FCustomTextFilterData& InFilterData, bool bApplyFilter) override;
+
+	/** Deletes the custom text filter with the matching label */
+	void DeleteCustomTextFilterByLabel(const FText& FilterLabel);
+
+	/** Modifies the custom text filter with the matching label */
+	void ModifyCustomTextFilterByLabel(const FCustomTextFilterData& NewFilterData, const FText& FilterLabel);
+
+	/** Whether the filters are being applied to paths recursively */
+	bool IsFilteringPathsRecursively() const;
+
+	/** Sets whether the filters are being applied to paths recursively */
+	void SetFilterPathsRecursively(const bool bInFilterRecursively);
+
+	/** Toggle whether the filters are being applied to paths recursively */
+	void ToggleFilterPathsRecursively();
+
+	/** Toggle whether the filters are sorted and displayed alphabetically */
+	void ToggleSortFiltersAlphabetically();
+
+	/** Returns a pointer to the active filter container widget (Horizontal or Vertical), nullptr if neither is active */
+	TSharedPtr<SWidget> GetActiveFilterContainer() const;
+
+	/** Returns all frontend filters registered with this filter bar (built-in, extra, and extension-provided). */
+	const TArray<TSharedRef<FFrontendFilter>>& GetAllFrontendFilters() const { return AllFrontendFilters_Internal; }
+
+	/** Returns all custom class/type filters registered with this filter bar. */
+	const TArray<TSharedRef<FCustomClassFilterData>>& GetCustomClassFilters() const { return CustomClassFilters; }
+
+	/**
+	 * Activates a class filter by class path:
+	 * - If not yet pinned: pins and enables it.
+	 * - If already pinned (e.g. disabled): enables it in-place without re-pinning.
+	 */
+	CONTENTBROWSER_API void ActivateClassFilterByPath(const FTopLevelAssetPath& ClassPath);
+
+protected:
+	
+	/** Handler for when the add filter button was clicked */
+	//virtual TSharedRef<SWidget> MakeAddFilterMenu() override;
+
+	virtual void SaveSettingsInternal(FFilterBarSettings* InSettingsToSave) override;
+	
+	/** Handler for when a custom text filter is modified */
+	virtual void OnModifyCustomTextFilter(const FCustomTextFilterData& InFilterData, TSharedPtr<ICustomTextFilter<FAssetFilterType>> InFilter) override;
+
+	/** Handler for when a custom text filter is deleted */
+	virtual void OnDeleteCustomTextFilter(const TSharedPtr<ICustomTextFilter<FAssetFilterType>> InFilter) override;
+
+private:
+
+	// /** Exists for backwards compatibility with ExternalMakeAddFilterMenu */
+	// TSharedRef<SWidget> MakeAddFilterMenu(FAssetCategoryPath MenuExpansion = EAssetCategoryPaths::Basic);
+
+	virtual UAssetFilterBarContext* CreateAssetFilterBarContext() override;
+
+	/** Handler to Populate the Add Filter Menu. Use OnFilterAssetType in Subclasses to add classes to the exclusion list */
+	virtual void PopulateAddFilterMenu(UToolMenu* Menu, TSharedPtr<FFilterCategory> MenuExpansion, FOnFilterAssetType OnFilterAssetType) override;
+	
+	/** Handler for when another SFilterList using shared settings creates a custom text filter */
+	void OnExternalCustomTextFilterCreated(TSharedPtr<SWidget> BroadcastingFilterList);
+	
+	/** Empty our list of custom text filters, and load from the given config */
+	void LoadCustomTextFilters(const FFilterBarSettings* FilterBarConfig);
+
+	/** Find the custom text filter corresponding to the specified state, and restore it's state to what is specified
+	 *  @return True if the filter was restored successfully, false if not
+	 */
+	bool RestoreCustomTextFilterState(const FCustomTextFilterState& InFilterState);
+	
+private:
+	/** List of classes that our filters must match */
+	TArray<UClass*> InitialClassFilters;
+
+	/** Delegate for when filters have changed */
+	FOnFilterChanged OnFilterChanged;
+
+	/** A reference to AllFrontEndFilters so we can access the filters as FFrontEndFilter instead of FFilterBase<FAssetFilterType> */
+	TArray< TSharedRef<FFrontendFilter> > AllFrontendFilters_Internal;
+
+	/** An identifier shared by all SFilterLists, used to save and load settings common to every instance */
+	static const FName SharedIdentifier;
+
+	/** If bIncludeClassName is true, custom text filters will include an asset's class name in the search */
+	bool bIncludeClassName;
+
+	/** If bIncludeAssetPath is true, custom text filters will match against full Asset path */
+	bool bIncludeAssetPath;
+
+	/** If bIncludeCollectionNames is true, custom text filters will match against collection names as well */
+	bool bIncludeCollectionNames;
+
+	/** Whether this FilterList wants to load/save from the settings common to all instances */
+	bool bUseSharedSettings;
+
+	/** Whether this FilterList applies to paths recursively */
+	bool bFilterPathsRecursively = true;
+
+	/** The event that executes when a custom text filter is created/modified/deleted in any filter list that is using Shared Settings */
+	static FCustomTextFilterEvent CustomTextFilterEvent;
+};
+
+/* A custom implementation of ICustomTextFilter that uses FFrontendFilter_Text to handle comparing items in the
+ * Asset View to Custom Text Filters. This ensures that the advanced search syntax etc behaves properly when used
+ * by a custom text filter created by saving a search.
+ * This is implemented as a fake frontend filter which modifies background asynchronous text filtering behavior in the asset view.
+ */
+class FFrontendFilter_CustomText :
+	public ICustomTextFilter<FAssetFilterType>,
+	public FFrontendFilter,
+	public TSharedFromThis<FFrontendFilter_CustomText>
+{
+public:
+	FFrontendFilter_CustomText();
+
+	// FFrontendFilter implementation
+	virtual FString GetName() const override;
+	virtual FText GetDisplayName() const override;
+	virtual FText GetToolTipText() const override;
+	virtual FLinearColor GetColor() const override;
+	virtual TOptional<FText> GetAsCustomTextFilter();
+	virtual bool PassesFilter(FAssetFilterType InItem) const override { return true; }
+	
+	/** Updates bIncludeClassName, bIncludeAssetPath and bIncludeCollectionNames for this filter */
+	void UpdateCustomTextFilterIncludes(const bool InIncludeClassName, const bool InIncludeAssetPath, const bool InIncludeCollectionNames);
+
+	//ICustomTextFilter interface
+	
+	/** Set the internals of this filter from an FCustomTextFilterData */
+	virtual void SetFromCustomTextFilterData(const FCustomTextFilterData& InFilterData) override;
+
+	/** Create an FCustomTextFilterData from the internals of this filter */
+	virtual FCustomTextFilterData CreateCustomTextFilterData() const override;
+
+	/** Get the actual filter */
+	virtual TSharedPtr<FFilterBase<FAssetFilterType>> GetFilter() override;
+
+protected:
+
+	/* The Display Name of this custom filter that the user sees */
+	FText DisplayName;
+
+	/** The raw text of the query, either simple string or advanced search syntax. */
+	FText RawFilterText;
+
+	/* The Color of this filter pill */
+	FLinearColor Color;
+
+	bool bIncludeClassName = false;
+	bool bIncludeAssetPath = false;
+	bool bIncludeCollectionNames = false;
+};
\ No newline at end of file
diff --git a/Engine/Source/Editor/ContentBrowser/Private/SAssetPicker.cpp b/Engine/Source/Editor/ContentBrowser/Private/SAssetPicker.cpp
--- a/Engine/Source/Editor/ContentBrowser/Private/SAssetPicker.cpp
+++ b/Engine/Source/Editor/ContentBrowser/Private/SAssetPicker.cpp
@@ -89,7 +89,8 @@ void SAssetPicker::Construct( const FArguments& InArgs )
 	OnFolderEnteredDelegate = InArgs._AssetPickerConfig.OnFolderEntered;
 	OnGetAssetContextMenu = InArgs._AssetPickerConfig.OnGetAssetContextMenu;
 	OnGetFolderContextMenu = InArgs._AssetPickerConfig.OnGetFolderContextMenu;
-	
+	OnSearchTextChangedDelegate = InArgs._AssetPickerConfig.OnSearchTextChanged;
+
 	// Break up the incoming filter into a sources data and backend filter.
 	TArray<FCollectionRef> CollectionsFilter = InArgs._AssetPickerConfig.CollectionsFilter;
 
@@ -440,6 +441,19 @@ void SAssetPicker::Construct( const FArguments& InArgs )
 	AssetViewPtr->RequestSlowFullListRefresh();
 }
 
+void SAssetPicker::ClearSearchText()
+{
+	if (SearchBoxPtr.IsValid())
+	{
+		SearchBoxPtr->SetText(FText::GetEmpty());
+	}
+}
+
+TSharedPtr<SFilterList> SAssetPicker::GetFilterList() const
+{
+	return FilterListPtr;
+}
+
 TSharedPtr<SWidget> SAssetPicker::GetSearchBox() const
 {
 	return SearchBoxPtr;
@@ -541,6 +555,7 @@ void SAssetPicker::SetSearchBoxText(const FText& InSearchText)
 void SAssetPicker::OnSearchBoxChanged(const FText& InSearchText)
 {
 	SetSearchBoxText( InSearchText );
+	OnSearchTextChangedDelegate.ExecuteIfBound(InSearchText);
 }
 
 void SAssetPicker::OnSearchBoxCommitted(const FText& InSearchText, ETextCommit::Type CommitInfo)
diff --git a/Engine/Source/Editor/ContentBrowser/Private/SAssetPicker.h b/Engine/Source/Editor/ContentBrowser/Private/SAssetPicker.h
deleted file mode 100644
diff --git a/Engine/Source/Editor/ContentBrowser/Private/SFilterList.cpp b/Engine/Source/Editor/ContentBrowser/Private/SFilterList.cpp
--- a/Engine/Source/Editor/ContentBrowser/Private/SFilterList.cpp
+++ b/Engine/Source/Editor/ContentBrowser/Private/SFilterList.cpp
@@ -201,6 +201,20 @@ TSharedPtr<FFrontendFilter> SFilterList::GetFrontendFilter(const FString& InName
 	return TSharedPtr<FFrontendFilter>();
 }
 
+void SFilterList::ActivateClassFilterByPath(const FTopLevelAssetPath& ClassPath)
+{
+	if (GetAssetTypeFilterCheckState(ClassPath) == ECheckBoxState::Checked)
+	{
+		// Checked means there's already a filter for it (checked in filter list), enable
+		ToggleAssetTypeFilterEnabled(ClassPath, true);
+	}
+	else
+	{
+		// Create and enable the filter if needed
+		SetAssetTypeFilterCheckState(ClassPath, ECheckBoxState::Checked, true);
+	}
+}
+
 TSharedRef<SWidget> SFilterList::ExternalMakeAddFilterMenu()
 {
 	return SAssetFilterBar<FAssetFilterType>::MakeAddFilterMenu();
diff --git a/Engine/Source/Editor/ContentBrowser/Private/SFilterList.h b/Engine/Source/Editor/ContentBrowser/Private/SFilterList.h
deleted file mode 100644
diff --git a/Engine/Source/Editor/ContentBrowser/Public/ContentBrowserDelegates.h b/Engine/Source/Editor/ContentBrowser/Public/ContentBrowserDelegates.h
--- a/Engine/Source/Editor/ContentBrowser/Public/ContentBrowserDelegates.h
+++ b/Engine/Source/Editor/ContentBrowser/Public/ContentBrowserDelegates.h
@@ -48,6 +48,9 @@ DECLARE_DELEGATE_RetVal_OneParam(bool, FOnShouldFilterItem, const FContentBrowse
 /** Called to check if an asset tag should be display in details view. Return false to exclude the asset from the view. */
 DECLARE_DELEGATE_RetVal_TwoParams(bool, FOnShouldDisplayAssetTag, FName /*AssetType*/, FName /*TagName*/);
 
+/** Fired when the search box text changes. Allows embedding widgets to observe the search text without replacing the built-in search behavior. */
+DECLARE_DELEGATE_OneParam(FOnSearchTextChanged, const FText& /*InSearchText*/)
+
 /** Called to clear the selection of the specified assetdata or all selection if an invalid assetdata is passed */
 DECLARE_DELEGATE( FClearSelectionDelegate );
 
diff --git a/Engine/Source/Editor/ContentBrowser/Public/IContentBrowserSingleton.h b/Engine/Source/Editor/ContentBrowser/Public/IContentBrowserSingleton.h
--- a/Engine/Source/Editor/ContentBrowser/Public/IContentBrowserSingleton.h
+++ b/Engine/Source/Editor/ContentBrowser/Public/IContentBrowserSingleton.h
@@ -23,6 +23,7 @@
 #include "Filters/SBasicFilterBar.h"
 
 class FViewport;
+class FFrontendFilter;
 class IPlugin;
 class SWidget;
 class UFactory;
@@ -283,6 +284,9 @@ struct FAssetPickerConfig
 	    determine if it should be displayed or not.  If it returns true or isn't bound, the tag will be displayed normally. */
 	FOnShouldDisplayAssetTag OnAssetTagWantsToBeDisplayed;
 
+	/** Fired when the search box text changes. Allows embedding widgets to observe the search text without replacing the built-in search behavior. */
+	FOnSearchTextChanged OnSearchTextChanged;
+
 	/** The default view mode */
 	EAssetViewType::Type InitialAssetViewType;
 
diff --git a/Engine/Source/Editor/EditorWidgets/Public/Filters/SAssetFilterBar.h b/Engine/Source/Editor/EditorWidgets/Public/Filters/SAssetFilterBar.h
--- a/Engine/Source/Editor/EditorWidgets/Public/Filters/SAssetFilterBar.h
+++ b/Engine/Source/Editor/EditorWidgets/Public/Filters/SAssetFilterBar.h
@@ -788,6 +788,7 @@ public:
 			if(AssetFilter->GetCustomClassFilterData()->GetClassPathName() == InClassPath)
 			{
 				AssetFilter->SetEnabled(bEnable);
+				return;
 			}
 		}
 	}
diff --git a/Engine/Source/Editor/UserAssetTagsEditor/Private/Widgets/STaggedAssetBrowser.cpp b/Engine/Source/Editor/UserAssetTagsEditor/Private/Widgets/STaggedAssetBrowser.cpp
--- a/Engine/Source/Editor/UserAssetTagsEditor/Private/Widgets/STaggedAssetBrowser.cpp
+++ b/Engine/Source/Editor/UserAssetTagsEditor/Private/Widgets/STaggedAssetBrowser.cpp
@@ -2,6 +2,7 @@
 
 #include "Widgets/STaggedAssetBrowser.h"
 
+#include "SAssetPicker.h"
 #include "UserAssetTagEditorUtilities.h"
 #include "UserAssetTagsEditorModule.h"
 #include "ContentBrowserMenuContexts.h"
@@ -133,6 +134,11 @@ TAttribute<const FMainMRUFavoritesList*> STaggedAssetBrowser::GetFavoritesListAt
 	return RecentAndFavoritesList;
 }
 
+TSharedPtr<SAssetPicker> STaggedAssetBrowser::GetAssetPicker() const
+{
+	return AssetBrowserContent.IsValid() ? AssetBrowserContent->GetAssetPicker() : nullptr;
+}
+
 void STaggedAssetBrowser::InitializeValues()
 {
 	PrimaryFiltersSlotValues.DefaultMinSize = 50.f;
diff --git a/Engine/Source/Editor/UserAssetTagsEditor/Private/Widgets/STaggedAssetBrowserContent.cpp b/Engine/Source/Editor/UserAssetTagsEditor/Private/Widgets/STaggedAssetBrowserContent.cpp
--- a/Engine/Source/Editor/UserAssetTagsEditor/Private/Widgets/STaggedAssetBrowserContent.cpp
+++ b/Engine/Source/Editor/UserAssetTagsEditor/Private/Widgets/STaggedAssetBrowserContent.cpp
@@ -5,6 +5,7 @@
 #include "ContentBrowserModule.h"
 #include "IContentBrowserDataModule.h"
 #include "IContentBrowserSingleton.h"
+#include "SAssetPicker.h"
 #include "SlateOptMacros.h"
 
 
@@ -21,12 +22,28 @@ void STaggedAssetBrowserContent::Construct(const FArguments& InArgs)
 	Config.bCanShowRealTimeThumbnails = true;
 	
 	FContentBrowserModule& ContentBrowserModule = FModuleManager::Get().LoadModuleChecked<FContentBrowserModule>("ContentBrowser");
+	TSharedRef<SWidget> AssetPickerWidget = ContentBrowserModule.Get().CreateAssetPicker(Config);
+
+	if (AssetPickerWidget->GetType() == FName(TEXT("SAssetPicker")))
+	{
+		AssetPickerPtr = StaticCastSharedRef<SAssetPicker>(AssetPickerWidget);
+	}
+	else
+	{
+		AssetPickerPtr = nullptr;
+	}
+
 	ChildSlot
 	[
-		ContentBrowserModule.Get().CreateAssetPicker(Config)
+		AssetPickerWidget
 	];
 }
 
+TSharedPtr<SAssetPicker> STaggedAssetBrowserContent::GetAssetPicker() const
+{
+	return AssetPickerPtr;
+}
+
 void STaggedAssetBrowserContent::SetARFilter(FARFilter InFilter)
 {
 	SetNewFilterDelegate.Execute(InFilter);
diff --git a/Engine/Source/Editor/UserAssetTagsEditor/Public/Widgets/STaggedAssetBrowser.h b/Engine/Source/Editor/UserAssetTagsEditor/Public/Widgets/STaggedAssetBrowser.h
--- a/Engine/Source/Editor/UserAssetTagsEditor/Public/Widgets/STaggedAssetBrowser.h
+++ b/Engine/Source/Editor/UserAssetTagsEditor/Public/Widgets/STaggedAssetBrowser.h
@@ -17,6 +17,7 @@
 
 #define UE_API USERASSETTAGSEDITOR_API
 
+class SAssetPicker;
 class UTaggedAssetBrowserMenuContext;
 struct FAssetPickerConfig;
 
@@ -143,6 +144,9 @@ public:
 	/** Returns the currently active section, or nullptr if no section is active. */
 	UE_API const UTaggedAssetBrowserSection* GetActiveSection() const;
 
+	/** Returns the asset picker widget inside this browser, if it exists */
+	UE_API TSharedPtr<SAssetPicker> GetAssetPicker() const;
+
 	/** Returns the currently selected primary filter items. */
 	UE_API TArray<const UTaggedAssetBrowserFilterBase*> GetSelectedPrimaryFilters() const;
 
diff --git a/Engine/Source/Editor/UserAssetTagsEditor/Public/Widgets/STaggedAssetBrowserContent.h b/Engine/Source/Editor/UserAssetTagsEditor/Public/Widgets/STaggedAssetBrowserContent.h
--- a/Engine/Source/Editor/UserAssetTagsEditor/Public/Widgets/STaggedAssetBrowserContent.h
+++ b/Engine/Source/Editor/UserAssetTagsEditor/Public/Widgets/STaggedAssetBrowserContent.h
@@ -9,12 +9,14 @@
 
 #define UE_API USERASSETTAGSEDITOR_API
 
+class SAssetPicker;
+
 class STaggedAssetBrowserContent : public SCompoundWidget
 {
 public:
 	DECLARE_DELEGATE_RetVal_OneParam(bool, FOnShouldFilterAsset, const FAssetData&)
 	DECLARE_DELEGATE_RetVal(TArray<TSharedRef<FFrontendFilter>>, FOnGetExtraFrontendFilters)
-	
+
 	SLATE_BEGIN_ARGS(STaggedAssetBrowserContent)
 		{
 		}
@@ -27,13 +29,19 @@ public:
 	UE_API void SetARFilter(FARFilter InFilter);
 
 	UE_API TArray<FAssetData> GetCurrentSelection() const;
-private:	
+
+	/** Returns the asset picker widget created during construction, if it exists. */
+	UE_API TSharedPtr<SAssetPicker> GetAssetPicker() const;
+
+private:
 	FRefreshAssetViewDelegate RefreshAssetViewDelegate;
 	FSyncToAssetsDelegate SyncToAssetsDelegate;
 	FSetARFilterDelegate SetNewFilterDelegate;
 	FGetCurrentSelectionDelegate GetCurrentSelectionDelegate;
 
 	FAssetPickerConfig InitialAssetPickerConfig;
+
+	TSharedPtr<SAssetPicker> AssetPickerPtr;
 };
 
 #undef UE_API
